[CoreUObject]
* Moved dll storage to methods/variables instead of on types. This is needed to get merged modules working

diff --git a/Engine/Source/Runtime/CoreUObject/Public/AssetRegistry/AssetData.h b/Engine/Source/Runtime/CoreUObject/Public/AssetRegistry/AssetData.h
--- a/Engine/Source/Runtime/CoreUObject/Public/AssetRegistry/AssetData.h
+++ b/Engine/Source/Runtime/CoreUObject/Public/AssetRegistry/AssetData.h
@@ -1034,7 +1034,7 @@ public:
 	COREUOBJECT_API void NetworkWrite(FCbWriter& Writer) const;
 	COREUOBJECT_API bool TryNetworkRead(FCbFieldView Field);
 
-	COREUOBJECT_API void SetPackageLocation(FPackageName::EPackageLocationFilter PackageLocation)
+	void SetPackageLocation(FPackageName::EPackageLocationFilter PackageLocation)
 	{	
 		// Clear existing location flags
 		Flags = Flags & ~FLAG_PACKAGE_LOCATION_MASK;
@@ -1044,7 +1044,7 @@ public:
 		Flags |= uint32(PackageLocation) << FLAG_PACKAGE_LOCATION_TO_ENUM_SHIFT_SIZE;
 	}
 
-	COREUOBJECT_API FPackageName::EPackageLocationFilter GetPackageLocation() const
+	FPackageName::EPackageLocationFilter GetPackageLocation() const
 	{
 		return FPackageName::EPackageLocationFilter((Flags & FLAG_PACKAGE_LOCATION_MASK) >> FLAG_PACKAGE_LOCATION_TO_ENUM_SHIFT_SIZE);
 	}
diff --git a/Engine/Source/Runtime/CoreUObject/Public/AssetRegistry/AssetDataTagMap.h b/Engine/Source/Runtime/CoreUObject/Public/AssetRegistry/AssetDataTagMap.h
--- a/Engine/Source/Runtime/CoreUObject/Public/AssetRegistry/AssetDataTagMap.h
+++ b/Engine/Source/Runtime/CoreUObject/Public/AssetRegistry/AssetDataTagMap.h
@@ -299,7 +299,7 @@ public:
 };
 
 /** Reference-counted handle to a loose FAssetDataTagMap or a fixed / immutable cooked tag map */
-class COREUOBJECT_API FAssetDataTagMapSharedView
+class FAssetDataTagMapSharedView
 {
 	union
 	{
@@ -337,15 +337,15 @@ class COREUOBJECT_API FAssetDataTagMapSharedView
 
 public:
 	FAssetDataTagMapSharedView() = default;
-	FAssetDataTagMapSharedView(const FAssetDataTagMapSharedView& O);
-	FAssetDataTagMapSharedView(FAssetDataTagMapSharedView&& O);
-	explicit FAssetDataTagMapSharedView(FixedTagPrivate::FMapHandle InFixed);
-	explicit FAssetDataTagMapSharedView(FAssetDataTagMap&& InLoose);
+	COREUOBJECT_API FAssetDataTagMapSharedView(const FAssetDataTagMapSharedView& O);
+	COREUOBJECT_API FAssetDataTagMapSharedView(FAssetDataTagMapSharedView&& O);
+	COREUOBJECT_API explicit FAssetDataTagMapSharedView(FixedTagPrivate::FMapHandle InFixed);
+	COREUOBJECT_API explicit FAssetDataTagMapSharedView(FAssetDataTagMap&& InLoose);
 
-	FAssetDataTagMapSharedView& operator=(const FAssetDataTagMapSharedView&);
-	FAssetDataTagMapSharedView& operator=(FAssetDataTagMapSharedView&&);
+	COREUOBJECT_API FAssetDataTagMapSharedView& operator=(const FAssetDataTagMapSharedView&);
+	COREUOBJECT_API FAssetDataTagMapSharedView& operator=(FAssetDataTagMapSharedView&&);
 
-	~FAssetDataTagMapSharedView();
+	COREUOBJECT_API ~FAssetDataTagMapSharedView();
 
 	using FFindTagResult = FAssetTagValueRef;
 
@@ -384,7 +384,7 @@ public:
 	}
 
 	/** Copy map contents to a loose FAssetDataTagMap */
-	FAssetDataTagMap CopyMap() const;
+	COREUOBJECT_API FAssetDataTagMap CopyMap() const;
 
 	template<typename Func>
 	void ForEach(Func Fn) const
@@ -407,7 +407,7 @@ public:
 	COREUOBJECT_API friend bool operator==(const FAssetDataTagMapSharedView& A, const FAssetDataTagMapSharedView& B);
 
 	///** Shrinks the contained map */
-	void Shrink();
+	COREUOBJECT_API void Shrink();
 
 	class TConstIterator
 	{
@@ -471,15 +471,15 @@ public:
 	}
 
 	/** Helps count deduplicated memory usage */
-	class COREUOBJECT_API FMemoryCounter
+	class FMemoryCounter
 	{
 		TSet<uint32> FixedStoreIndices;
 		SIZE_T LooseBytes = 0;
 	public:
-		void Include(const FAssetDataTagMapSharedView& Tags);
+		COREUOBJECT_API void Include(const FAssetDataTagMapSharedView& Tags);
 		SIZE_T GetLooseSize() const { return LooseBytes; }
-		SIZE_T GetFixedSize() const;
-		void ReportFixedStoreBreakdown() const;
+		COREUOBJECT_API SIZE_T GetFixedSize() const;
+		COREUOBJECT_API void ReportFixedStoreBreakdown() const;
 	};
 
 	friend inline bool operator==(const FAssetDataTagMap& A, const FAssetDataTagMapSharedView& B)			{ return B == A; }
diff --git a/Engine/Source/Runtime/CoreUObject/Public/AssetRegistry/AssetDataToken.h b/Engine/Source/Runtime/CoreUObject/Public/AssetRegistry/AssetDataToken.h
--- a/Engine/Source/Runtime/CoreUObject/Public/AssetRegistry/AssetDataToken.h
+++ b/Engine/Source/Runtime/CoreUObject/Public/AssetRegistry/AssetDataToken.h
@@ -28,7 +28,7 @@ public:
 	virtual const FOnMessageTokenActivated& GetOnMessageTokenActivated() const override;
 	/** End IMessageToken interface */
 
-    COREUOBJECT_API const FAssetData& GetAssetData() const
+    const FAssetData& GetAssetData() const
     {
         return AssetData;
     }
diff --git a/Engine/Source/Runtime/CoreUObject/Public/Misc/EditorPathObjectInterface.h b/Engine/Source/Runtime/CoreUObject/Public/Misc/EditorPathObjectInterface.h
--- a/Engine/Source/Runtime/CoreUObject/Public/Misc/EditorPathObjectInterface.h
+++ b/Engine/Source/Runtime/CoreUObject/Public/Misc/EditorPathObjectInterface.h
@@ -21,7 +21,7 @@ class UEditorPathObjectInterface : public UInterface
  * Example implementation of this is to allow Objects outside of a Level Instance to reference objects inside a Level Instance.
  */
 
-class COREUOBJECT_API IEditorPathObjectInterface
+class IEditorPathObjectInterface
 {
 	GENERATED_IINTERFACE_BODY()
 
diff --git a/Engine/Source/Runtime/CoreUObject/Public/StructUtils/InstancedStruct.h b/Engine/Source/Runtime/CoreUObject/Public/StructUtils/InstancedStruct.h
--- a/Engine/Source/Runtime/CoreUObject/Public/StructUtils/InstancedStruct.h
+++ b/Engine/Source/Runtime/CoreUObject/Public/StructUtils/InstancedStruct.h
@@ -6,6 +6,8 @@
 
 #include "InstancedStruct.generated.h"
 
+#define UE_API COREUOBJECT_API 
+
 struct FConstStructView;
 template<typename BaseStructT> struct TConstStructView;
 struct FConstSharedStruct;
@@ -24,21 +26,21 @@ class UUserDefinedStruct;
  *	TArray<FInstancedStruct> TestArray;
  */
 USTRUCT(BlueprintType, meta = (DisableSplitPin, HasNativeMake = "/Script/Engine.BlueprintInstancedStructLibrary.MakeInstancedStruct"))
-struct [[nodiscard]] COREUOBJECT_API FInstancedStruct
+struct [[nodiscard]] FInstancedStruct
 {
 	GENERATED_BODY()
 
 public:
 
-	FInstancedStruct();
+	UE_API FInstancedStruct();
 
-	explicit FInstancedStruct(const UScriptStruct* InScriptStruct);
+	UE_API explicit FInstancedStruct(const UScriptStruct* InScriptStruct);
 
 	/**
 	 * This constructor is explicit to avoid accidentally converting struct views to instanced structs (which would result in costly copy of the struct to be made).
 	 * Implicit conversion could happen e.g. when comparing FInstancedStruct to FConstStructView.
 	 */
-	explicit FInstancedStruct(const FConstStructView InOther);
+	UE_API explicit FInstancedStruct(const FConstStructView InOther);
 
 	FInstancedStruct(const FInstancedStruct& InOther)
 	{
@@ -56,7 +58,7 @@ public:
 		Reset();
 	}
 
-	FInstancedStruct& operator=(const FConstStructView InOther);
+	UE_API FInstancedStruct& operator=(const FConstStructView InOther);
 
 	FInstancedStruct& operator=(const FInstancedStruct& InOther)
 	{
@@ -80,7 +82,7 @@ public:
 	}
 
 	/** Initializes from struct type and optional data. */
-	void InitializeAs(const UScriptStruct* InScriptStruct, const uint8* InStructMemory = nullptr);
+	UE_API void InitializeAs(const UScriptStruct* InScriptStruct, const uint8* InStructMemory = nullptr);
 
 	/** Initializes from struct type and emplace construct. */
 	template<typename T, typename... TArgs>
@@ -149,17 +151,17 @@ public:
 	}
 
 	/** For StructOpsTypeTraits */
-	bool Serialize(FArchive& Ar);
-	bool Identical(const FInstancedStruct* Other, uint32 PortFlags) const;
-	void AddStructReferencedObjects(class FReferenceCollector& Collector);
-	bool ExportTextItem(FString& ValueStr, FInstancedStruct const& DefaultValue, class UObject* Parent, int32 PortFlags, class UObject* ExportRootScope) const;
-	bool ImportTextItem(const TCHAR*& Buffer, int32 PortFlags, UObject* Parent, FOutputDevice* ErrorText, FArchive* InSerializingArchive = nullptr);
-	bool SerializeFromMismatchedTag(const FPropertyTag& Tag, FStructuredArchive::FSlot Slot);
-	void GetPreloadDependencies(TArray<UObject*>& OutDeps);
-	bool NetSerialize(FArchive& Ar, class UPackageMap* Map, bool& bOutSuccess);
-	bool FindInnerPropertyInstance(FName PropertyName, const FProperty*& OutProp, const void*& OutData) const;
-	EPropertyVisitorControlFlow Visit(FPropertyVisitorContext& Context, const TFunctionRef<EPropertyVisitorControlFlow(const FPropertyVisitorContext& /*Context*/)> InFunc) const;
-	void* ResolveVisitedPathInfo(const FPropertyVisitorInfo& Info) const;
+	UE_API bool Serialize(FArchive& Ar);
+	UE_API bool Identical(const FInstancedStruct* Other, uint32 PortFlags) const;
+	UE_API void AddStructReferencedObjects(class FReferenceCollector& Collector);
+	UE_API bool ExportTextItem(FString& ValueStr, FInstancedStruct const& DefaultValue, class UObject* Parent, int32 PortFlags, class UObject* ExportRootScope) const;
+	UE_API bool ImportTextItem(const TCHAR*& Buffer, int32 PortFlags, UObject* Parent, FOutputDevice* ErrorText, FArchive* InSerializingArchive = nullptr);
+	UE_API bool SerializeFromMismatchedTag(const FPropertyTag& Tag, FStructuredArchive::FSlot Slot);
+	UE_API void GetPreloadDependencies(TArray<UObject*>& OutDeps);
+	UE_API bool NetSerialize(FArchive& Ar, class UPackageMap* Map, bool& bOutSuccess);
+	UE_API bool FindInnerPropertyInstance(FName PropertyName, const FProperty*& OutProp, const void*& OutData) const;
+	UE_API EPropertyVisitorControlFlow Visit(FPropertyVisitorContext& Context, const TFunctionRef<EPropertyVisitorControlFlow(const FPropertyVisitorContext& /*Context*/)> InFunc) const;
+	UE_API void* ResolveVisitedPathInfo(const FPropertyVisitorInfo& Info) const;
 
 	/** Returns struct type. */
 	const UScriptStruct* GetScriptStruct() const
@@ -174,7 +176,7 @@ public:
 	}
 
 	/** Reset to empty. */
-	void Reset();
+	UE_API void Reset();
 
 	/** Returns const reference to the struct, this getter assumes that all data is valid. */
 	template<typename T>
@@ -244,7 +246,7 @@ public:
 #endif // WITH_EDITOR
 
 	DECLARE_DELEGATE_RetVal_ThreeParams(bool, FNetSerializeInstancedStruct, FInstancedStruct&, FArchive&, UPackageMap*);
-	static FNetSerializeInstancedStruct NetSerializeScriptStructDelegate;
+	static UE_API FNetSerializeInstancedStruct NetSerializeScriptStructDelegate;
 
 protected:
 
@@ -505,3 +507,5 @@ private:
 
 	template <typename U> friend struct TInstancedStruct;
 };
+
+#undef UE_API
\ No newline at end of file
diff --git a/Engine/Source/Runtime/CoreUObject/Public/StructUtils/InstancedStructContainer.h b/Engine/Source/Runtime/CoreUObject/Public/StructUtils/InstancedStructContainer.h
--- a/Engine/Source/Runtime/CoreUObject/Public/StructUtils/InstancedStructContainer.h
+++ b/Engine/Source/Runtime/CoreUObject/Public/StructUtils/InstancedStructContainer.h
@@ -5,6 +5,8 @@
 #include "StructUtils/StructView.h"
 #include "InstancedStructContainer.generated.h"
 
+#define UE_API COREUOBJECT_API
+
 struct FInstancedStruct;
 class FReferenceCollector;
 
@@ -27,7 +29,7 @@ class FReferenceCollector;
  */
 
 USTRUCT()
-struct COREUOBJECT_API FInstancedStructContainer
+struct FInstancedStructContainer
 {
 	GENERATED_BODY()
 
@@ -63,46 +65,46 @@ public:
 	/** How much memory is used to store info about each item. */
 	static constexpr int32 OverheadPerItem = sizeof(FItem);
 	
-	FInstancedStructContainer();
-	FInstancedStructContainer(const FInstancedStructContainer& InOther);
-	FInstancedStructContainer(FInstancedStructContainer&& InOther);
+	UE_API FInstancedStructContainer();
+	UE_API FInstancedStructContainer(const FInstancedStructContainer& InOther);
+	UE_API FInstancedStructContainer(FInstancedStructContainer&& InOther);
 	
 	~FInstancedStructContainer() { Empty(); }
 
-	FInstancedStructContainer& operator=(const FInstancedStructContainer& InOther);
-	FInstancedStructContainer& operator=(FInstancedStructContainer&& InOther);
-	FInstancedStructContainer& operator=(TConstArrayView<FInstancedStruct> InItems);
-	FInstancedStructContainer& operator=(TConstArrayView<FStructView> InItems);
-	FInstancedStructContainer& operator=(TConstArrayView<FConstStructView> InItems);
+	UE_API FInstancedStructContainer& operator=(const FInstancedStructContainer& InOther);
+	UE_API FInstancedStructContainer& operator=(FInstancedStructContainer&& InOther);
+	UE_API FInstancedStructContainer& operator=(TConstArrayView<FInstancedStruct> InItems);
+	UE_API FInstancedStructContainer& operator=(TConstArrayView<FStructView> InItems);
+	UE_API FInstancedStructContainer& operator=(TConstArrayView<FConstStructView> InItems);
 
 	/** Appends items to the array. */
-	void Append(const FInstancedStructContainer& Other);
-	void Append(TConstArrayView<FInstancedStruct> NewItemValues);
-	void Append(TConstArrayView<FConstStructView> NewItemValues);
+	UE_API void Append(const FInstancedStructContainer& Other);
+	UE_API void Append(TConstArrayView<FInstancedStruct> NewItemValues);
+	UE_API void Append(TConstArrayView<FConstStructView> NewItemValues);
 
 	/** Insert new items at specified location. */
-	void InsertAt(const int32 InsertAtIndex, const FInstancedStructContainer& Other);
-	void InsertAt(const int32 InsertAtIndex, TConstArrayView<FInstancedStruct> ValuesToInsert);
-	void InsertAt(const int32 InsertAtIndex, TConstArrayView<FStructView> ValuesToInsert);
-	void InsertAt(const int32 InsertAtIndex, TConstArrayView<FConstStructView> ValuesToInsert);
+	UE_API void InsertAt(const int32 InsertAtIndex, const FInstancedStructContainer& Other);
+	UE_API void InsertAt(const int32 InsertAtIndex, TConstArrayView<FInstancedStruct> ValuesToInsert);
+	UE_API void InsertAt(const int32 InsertAtIndex, TConstArrayView<FStructView> ValuesToInsert);
+	UE_API void InsertAt(const int32 InsertAtIndex, TConstArrayView<FConstStructView> ValuesToInsert);
 
 	/** Remove items at specific location. Does not change memory allocation. */
-	void RemoveAt(const int32 RemoveAtIndex, const int32 Count);
+	UE_API void RemoveAt(const int32 RemoveAtIndex, const int32 Count);
 
 	/** Reserves at least 'NumBytes' for internal storage. */
-	void ReserveBytes(const int32 NumBytes, const int32 MinAlignment = DefaultMinAlignment);
+	UE_API void ReserveBytes(const int32 NumBytes, const int32 MinAlignment = DefaultMinAlignment);
 	
 	/** Sets the number of items in the array. Note: this can only shrink the array. Does not change memory allocation. */
-	void SetNum(const int32 NewNum);
+	UE_API void SetNum(const int32 NewNum);
 
 	/** Returns number of bytes allocated for the array  */
 	int32 GetAllocatedMemory() const { return AllocatedSize; }
 	
 	/** Empties the array, destroys entities. Does not change memory allocation. */
-	void Reset();
+	UE_API void Reset();
 
 	/** Empties the array, destroys entities. Frees memory. */
-	void Empty();
+	UE_API void Empty();
 
 	/** @return true if the instance is correctly initialized. */
 	bool IsValid() const { return NumItems > 0; }
@@ -223,12 +225,12 @@ public:
 	FConstIterator end() const { return FConstIterator(*this, NumItems); }
 
 	/** Type traits */
-	void AddStructReferencedObjects(FReferenceCollector& Collector);
-	bool Identical(const FInstancedStructContainer* Other, uint32 PortFlags) const;
-	bool Serialize(FArchive& Ar);
-	void GetPreloadDependencies(TArray<UObject*>& OutDeps) const;
-	bool ExportTextItem(FString& ValueStr, FInstancedStructContainer const& DefaultValue, UObject* Parent, int32 PortFlags, UObject* ExportRootScope) const;
-	bool ImportTextItem(const TCHAR*& Buffer, int32 PortFlags, UObject* Parent, FOutputDevice* ErrorText, FArchive* InSerializingArchive = nullptr);
+	UE_API void AddStructReferencedObjects(FReferenceCollector& Collector);
+	UE_API bool Identical(const FInstancedStructContainer* Other, uint32 PortFlags) const;
+	UE_API bool Serialize(FArchive& Ar);
+	UE_API void GetPreloadDependencies(TArray<UObject*>& OutDeps) const;
+	UE_API bool ExportTextItem(FString& ValueStr, FInstancedStructContainer const& DefaultValue, UObject* Parent, int32 PortFlags, UObject* ExportRootScope) const;
+	UE_API bool ImportTextItem(const TCHAR*& Buffer, int32 PortFlags, UObject* Parent, FOutputDevice* ErrorText, FArchive* InSerializingArchive = nullptr);
 
 private:
 
@@ -279,3 +281,5 @@ struct TStructOpsTypeTraits<FInstancedStructContainer> : public TStructOpsTypeTr
 		WithImportTextItem = true,
 	};
 };
+
+#undef UE_API
\ No newline at end of file
diff --git a/Engine/Source/Runtime/CoreUObject/Public/StructUtils/PropertyBag.h b/Engine/Source/Runtime/CoreUObject/Public/StructUtils/PropertyBag.h
--- a/Engine/Source/Runtime/CoreUObject/Public/StructUtils/PropertyBag.h
+++ b/Engine/Source/Runtime/CoreUObject/Public/StructUtils/PropertyBag.h
@@ -9,6 +9,8 @@
 #include "UObject/ObjectMacros.h"
 #include "PropertyBag.generated.h"
 
+#define UE_API COREUOBJECT_API
+
 /** Property bag property type, loosely based on BluePrint pin types. */
 UENUM(BlueprintType)
 enum class EPropertyBagPropertyType : uint8
@@ -48,7 +50,7 @@ enum class EPropertyBagContainerType : uint8
 
 /** Helper to manage container types, with nested container support. */
 USTRUCT()
-struct COREUOBJECT_API FPropertyBagContainerTypes
+struct FPropertyBagContainerTypes
 {
 	GENERATED_BODY()
 
@@ -123,9 +125,9 @@ struct COREUOBJECT_API FPropertyBagContainerTypes
 		return ensure(Index < NumContainers) ? Types[Index] : EPropertyBagContainerType::None;
 	}
 
-	EPropertyBagContainerType PopHead();
+	UE_API EPropertyBagContainerType PopHead();
 
-	void Serialize(FArchive& Ar);
+	UE_API void Serialize(FArchive& Ar);
 
 	friend FArchive& operator<<(FArchive& Ar, FPropertyBagContainerTypes& ContainerTypesData)
 	{
@@ -133,7 +135,7 @@ struct COREUOBJECT_API FPropertyBagContainerTypes
 		return Ar;
 	}
 
-	bool operator == (const FPropertyBagContainerTypes& Other) const;
+	UE_API bool operator == (const FPropertyBagContainerTypes& Other) const;
 
 	FORCEINLINE bool operator !=(const FPropertyBagContainerTypes& Other) const
 	{
@@ -183,7 +185,7 @@ enum class EPropertyBagAlterationResult : uint8
 };
 
 USTRUCT()
-struct COREUOBJECT_API FPropertyBagPropertyDescMetaData
+struct FPropertyBagPropertyDescMetaData
 {
 	GENERATED_BODY()
 	
@@ -200,7 +202,7 @@ struct COREUOBJECT_API FPropertyBagPropertyDescMetaData
 	UPROPERTY()
 	FString Value;
 
-	void Serialize(FArchive& Ar);
+	UE_API void Serialize(FArchive& Ar);
 
 	bool operator==(const FPropertyBagPropertyDescMetaData& Other) const
 	{
@@ -236,14 +238,14 @@ struct COREUOBJECT_API FPropertyBagPropertyDescMetaData
 
 /** Describes a property in the property bag. */
 USTRUCT()
-struct COREUOBJECT_API FPropertyBagPropertyDesc
+struct FPropertyBagPropertyDesc
 {
 	GENERATED_BODY()
 
 	static_assert(std::is_same_v<std::underlying_type_t<EPropertyFlags>, uint64>, "FPropertyBagPropertyDesc::PropertyFlag does not match EPropertyFlags type");
 
 	FPropertyBagPropertyDesc() = default;
-	FPropertyBagPropertyDesc(const FName InName, const FProperty* InSourceProperty);
+	UE_API FPropertyBagPropertyDesc(const FName InName, const FProperty* InSourceProperty);
 	FPropertyBagPropertyDesc(const FName InName, const EPropertyBagPropertyType InValueType, const UObject* InValueTypeObject = nullptr)
 		: ValueTypeObject(InValueTypeObject)
 		, Name(InName)
@@ -269,28 +271,28 @@ struct COREUOBJECT_API FPropertyBagPropertyDesc
 	}
 
 	/** @return true if the two descriptors have the same type. Object types are compatible if Other can be cast to this type. */
-	bool CompatibleType(const FPropertyBagPropertyDesc& Other) const;
+	UE_API bool CompatibleType(const FPropertyBagPropertyDesc& Other) const;
 
 	/** @return true if the property type is numeric (bool, (u)int32, (u)int64, float, double, enum) */
-	bool IsNumericType() const;
+	UE_API bool IsNumericType() const;
 
 	/** @return true if the property type is unsigned (uint32, uint64) */
-	bool IsUnsignedNumericType() const;
+	UE_API bool IsUnsignedNumericType() const;
 	
 	/** @return true if the property type is floating point numeric (float, double) */
-	bool IsNumericFloatType() const;
+	UE_API bool IsNumericFloatType() const;
 
 	/** @return true if the property type is object or soft object */
-	bool IsObjectType() const;
+	UE_API bool IsObjectType() const;
 
 	/** @return true if the property type is class or soft class */
-	bool IsClassType() const;
+	UE_API bool IsClassType() const;
 
 	/** @return index of the property after constructed in the bag */
 	int32 GetCachedIndex() const { return CachedIndex; }
 
 	/** @return true if the property bag's unique IDs match */
-	bool operator==(const FPropertyBagPropertyDesc& OtherDesc) const;
+	UE_API bool operator==(const FPropertyBagPropertyDesc& OtherDesc) const;
 
 	/** Pointer to object that defines the Enum, Struct, or Class. */
 	UPROPERTY(EditAnywhere, Category="Default")
@@ -383,7 +385,7 @@ class FPropertyBagArrayRef;
 class FPropertyBagSetRef;
 
 USTRUCT()
-struct COREUOBJECT_API FInstancedPropertyBag
+struct FInstancedPropertyBag
 {
 	GENERATED_BODY()
 
@@ -407,17 +409,17 @@ struct COREUOBJECT_API FInstancedPropertyBag
 	}
 
 	/** Initializes the instance from a bag struct. */
-	void InitializeFromBagStruct(const UPropertyBag* NewBagStruct);
+	UE_API void InitializeFromBagStruct(const UPropertyBag* NewBagStruct);
 	
 	/**
 	 * Copies matching property values from another bag of potentially mismatching layout.
 	 * The properties are matched between the bags based on the property ID.
 	 * @param NewDescs Reference to the bag to copy the values from
 	 */
-	void CopyMatchingValuesByID(const FInstancedPropertyBag& NewDescs);
+	UE_API void CopyMatchingValuesByID(const FInstancedPropertyBag& NewDescs);
 
 	/** Returns number of the Properties in this Property Bag */
-	int32 GetNumPropertiesInBag() const;
+	UE_API int32 GetNumPropertiesInBag() const;
 
 	/**
 	 * Adds properties to the bag. If property of same name already exists, it will be replaced with the new type.
@@ -426,7 +428,7 @@ struct COREUOBJECT_API FInstancedPropertyBag
 	 * @param bOverwrite Overwrite the property if it already exists.
 	 * @return The result of the alteration.
 	 */
-	EPropertyBagAlterationResult AddProperties(const TConstArrayView<FPropertyBagPropertyDesc> Descs, bool bOverwrite = true);
+	UE_API EPropertyBagAlterationResult AddProperties(const TConstArrayView<FPropertyBagPropertyDesc> Descs, bool bOverwrite = true);
 	
 	/**
 	 * Adds a new property to the bag. If property of same name already exists, it will be replaced with the new type.
@@ -437,7 +439,7 @@ struct COREUOBJECT_API FInstancedPropertyBag
 	 * @param bOverwrite Overwrite the property if it already exists.
 	 * @return The result of the alteration.
 	 */
-	EPropertyBagAlterationResult AddProperty(const FName InName, const EPropertyBagPropertyType InValueType, const UObject* InValueTypeObject = nullptr, bool bOverwrite = true);
+	UE_API EPropertyBagAlterationResult AddProperty(const FName InName, const EPropertyBagPropertyType InValueType, const UObject* InValueTypeObject = nullptr, bool bOverwrite = true);
 
 	/**
 	 * Adds a new container property to the bag. If property of same name already exists, it will be replaced with the new type.
@@ -448,7 +450,7 @@ struct COREUOBJECT_API FInstancedPropertyBag
 	 * @param bOverwrite Overwrite the property if it already exists.
 	 * @return The result of the alteration.
 	 */
-	EPropertyBagAlterationResult AddContainerProperty(const FName InName, const EPropertyBagContainerType InContainerType, const EPropertyBagPropertyType InValueType, const UObject* InValueTypeObject = nullptr, bool bOverwrite = true);
+	UE_API EPropertyBagAlterationResult AddContainerProperty(const FName InName, const EPropertyBagContainerType InContainerType, const EPropertyBagPropertyType InValueType, const UObject* InValueTypeObject = nullptr, bool bOverwrite = true);
 
 	/**
 	 * Adds a new container property to the bag. If property of same name already exists, it will be replaced with the new type.
@@ -459,7 +461,7 @@ struct COREUOBJECT_API FInstancedPropertyBag
 	 * @param bOverwrite Overwrite the property if it already exists.
 	 * @return The result of the alteration.
 	 */
-	EPropertyBagAlterationResult AddContainerProperty(const FName InName, const FPropertyBagContainerTypes InContainerTypes, const EPropertyBagPropertyType InValueType, UObject* InValueTypeObject, bool bOverwrite = true);
+	UE_API EPropertyBagAlterationResult AddContainerProperty(const FName InName, const FPropertyBagContainerTypes InContainerTypes, const EPropertyBagPropertyType InValueType, UObject* InValueTypeObject, bool bOverwrite = true);
 
 	/**
 	 * Adds a new property to the bag. Property type duplicated from source property to. If property of same name already exists, it will be replaced with the new type.
@@ -469,7 +471,7 @@ struct COREUOBJECT_API FInstancedPropertyBag
 	 * @return The result of the alteration.
 
 	 */
-	EPropertyBagAlterationResult AddProperty(const FName InName, const FProperty* InSourceProperty, bool bOverwrite = true);
+	UE_API EPropertyBagAlterationResult AddProperty(const FName InName, const FProperty* InSourceProperty, bool bOverwrite = true);
 
 	/**
 	 * Clears all properties, then adds the supplied properties to the bag and sets their respective values.
@@ -477,21 +479,21 @@ struct COREUOBJECT_API FInstancedPropertyBag
 	 * @param InDescs Descriptors of new properties to add. Must be the same size as InValues.
 	 * @param InValues Values of new properties to add. Must be the same size as InDescs.
 	 */
-	EPropertyBagResult ReplaceAllPropertiesAndValues(const TConstArrayView<FPropertyBagPropertyDesc> InDescs, const TConstArrayView<TConstArrayView<uint8>> InValues);
+	UE_API EPropertyBagResult ReplaceAllPropertiesAndValues(const TConstArrayView<FPropertyBagPropertyDesc> InDescs, const TConstArrayView<TConstArrayView<uint8>> InValues);
 
 	/**
 	 * Removes properties from the bag by name if they exist.
 	 * @param PropertiesToRemove The names of the properties to remove.
 	 * @return The result of the alteration.
 	 */
-	EPropertyBagAlterationResult RemovePropertiesByName(const TConstArrayView<FName> PropertiesToRemove);
+	UE_API EPropertyBagAlterationResult RemovePropertiesByName(const TConstArrayView<FName> PropertiesToRemove);
 	
 	/**
 	 * Removes a property from the bag by name if it exists.
 	 * @param PropertyToRemove The name of the property to remove.
 	 * @return The result of the alteration.
 	 */
-	EPropertyBagAlterationResult RemovePropertyByName(const FName PropertyToRemove);
+	UE_API EPropertyBagAlterationResult RemovePropertyByName(const FName PropertyToRemove);
 
 	/**
 	 * Renames a property in the bag if it exists.
@@ -499,7 +501,7 @@ struct COREUOBJECT_API FInstancedPropertyBag
 	 * @param NewName The new name of the property
 	 * @return The result of the alteration
 	 */
-	EPropertyBagAlterationResult RenameProperty(FName PropertyToRename, FName NewName);
+	UE_API EPropertyBagAlterationResult RenameProperty(FName PropertyToRename, FName NewName);
 
 	/**
 	 * Reorders a property either before or after another target property.
@@ -508,7 +510,7 @@ struct COREUOBJECT_API FInstancedPropertyBag
 	 * @param bInsertBefore Insert the source property before the target property. If false, it will be inserted after.
 	 * @return The result of the alteration.
 	 */
-	EPropertyBagAlterationResult ReorderProperty(FName SourcePropertyName, FName TargetPropertyName, bool bInsertBefore = true);
+	UE_API EPropertyBagAlterationResult ReorderProperty(FName SourcePropertyName, FName TargetPropertyName, bool bInsertBefore = true);
 
 	/**
 	 * Reorders a property either before or after another target property.
@@ -517,14 +519,14 @@ struct COREUOBJECT_API FInstancedPropertyBag
 	 * @param bInsertBefore Insert the source property before the target property. If false, it will be inserted after.
 	 * @return The result of the alteration.
 	 */
-	EPropertyBagAlterationResult ReorderProperty(int32 SourcePropertyIndex, int32 TargetPropertyIndex, bool bInsertBefore = true);
+	UE_API EPropertyBagAlterationResult ReorderProperty(int32 SourcePropertyIndex, int32 TargetPropertyIndex, bool bInsertBefore = true);
 
 	/**
 	 * Changes the type of this bag and migrates existing values.
 	 * The properties are matched between the bags based on the property ID.
 	 * @param NewBagStruct Pointer to the new type.
 	 */
-	void MigrateToNewBagStruct(const UPropertyBag* NewBagStruct);
+	UE_API void MigrateToNewBagStruct(const UPropertyBag* NewBagStruct);
 
 	/**
 	 * Changes the type of this bag to the InNewBagInstance, and migrates existing values over.
@@ -532,7 +534,7 @@ struct COREUOBJECT_API FInstancedPropertyBag
 	 * The properties are matched between the bags based on the property ID.
 	 * @param InNewBagInstance New bag composition and values used for new properties.
 	 */
-	void MigrateToNewBagInstance(const FInstancedPropertyBag& InNewBagInstance);
+	UE_API void MigrateToNewBagInstance(const FInstancedPropertyBag& InNewBagInstance);
 
 	/**
 	 * Changes the type of this bag to the InNewBagInstance, and migrates existing values over if marked as overridden in the OverriddenPropertyIDs.
@@ -541,22 +543,22 @@ struct COREUOBJECT_API FInstancedPropertyBag
 	 * @param InNewBagInstance New bag composition and values used for new properties.
 	 * @param OverriddenPropertyIDs Array if property IDs which should be copied over to the new instance. 
 	 */
-	void MigrateToNewBagInstanceWithOverrides(const FInstancedPropertyBag& InNewBagInstance, TConstArrayView<FGuid> OverriddenPropertyIDs);
+	UE_API void MigrateToNewBagInstanceWithOverrides(const FInstancedPropertyBag& InNewBagInstance, TConstArrayView<FGuid> OverriddenPropertyIDs);
 
 	/** @return pointer to the property bag struct. */ 
-	const UPropertyBag* GetPropertyBagStruct() const;
+	UE_API const UPropertyBag* GetPropertyBagStruct() const;
 	
 	/** Returns property descriptor by specified name. */
-	const FPropertyBagPropertyDesc* FindPropertyDescByID(const FGuid ID) const;
+	UE_API const FPropertyBagPropertyDesc* FindPropertyDescByID(const FGuid ID) const;
 	
 	/** Returns property descriptor by specified ID. */
-	const FPropertyBagPropertyDesc* FindPropertyDescByName(const FName Name) const;
+	UE_API const FPropertyBagPropertyDesc* FindPropertyDescByName(const FName Name) const;
 
 	/** Returns true if we own the supplied property description, false otherwise. */
-	bool OwnsPropertyDesc(const FPropertyBagPropertyDesc& Desc) const;
+	UE_API bool OwnsPropertyDesc(const FPropertyBagPropertyDesc& Desc) const;
 
 	/** Returns true if our propertybag has the same layout & per property types as another propertybag. */
-	bool HasSameLayout(const FInstancedPropertyBag& Other) const;
+	UE_API bool HasSameLayout(const FInstancedPropertyBag& Other) const;
 
 	/** @return const view to the struct that holds the values. NOTE: The returned value/view cannot be serialized, use this to access the struct only temporarily. */
 	FConstStructView GetValue() const { return Value; };
@@ -569,25 +571,25 @@ struct COREUOBJECT_API FInstancedPropertyBag
 	 * Numeric types (bool, (u)int32, (u)int64, float, double) support type conversion.
 	 */
 
-	TValueOrError<bool, EPropertyBagResult> GetValueBool(const FName Name) const;
-	TValueOrError<uint8, EPropertyBagResult> GetValueByte(const FName Name) const;
-	TValueOrError<int32, EPropertyBagResult> GetValueInt32(const FName Name) const;
-	TValueOrError<uint32, EPropertyBagResult> GetValueUInt32(const FName Name) const;
-	TValueOrError<int64, EPropertyBagResult> GetValueInt64(const FName Name) const;
-	TValueOrError<uint64, EPropertyBagResult> GetValueUInt64(const FName Name) const;
-	TValueOrError<float, EPropertyBagResult> GetValueFloat(const FName Name) const;
-	TValueOrError<double, EPropertyBagResult> GetValueDouble(const FName Name) const;
-	TValueOrError<FName, EPropertyBagResult> GetValueName(const FName Name) const;
-	TValueOrError<FString, EPropertyBagResult> GetValueString(const FName Name) const;
-	TValueOrError<FText, EPropertyBagResult> GetValueText(const FName Name) const;
-	TValueOrError<uint8, EPropertyBagResult> GetValueEnum(const FName Name, const UEnum* RequestedEnum) const;
-	TValueOrError<FStructView, EPropertyBagResult> GetValueStruct(const FName Name, const UScriptStruct* RequestedStruct = nullptr) const;
-	TValueOrError<UObject*, EPropertyBagResult> GetValueObject(const FName Name, const UClass* RequestedClass = nullptr) const;
-	TValueOrError<UClass*, EPropertyBagResult> GetValueClass(const FName Name) const;
-	TValueOrError<FSoftObjectPath, EPropertyBagResult> GetValueSoftPath(const FName Name) const;
+	UE_API TValueOrError<bool, EPropertyBagResult> GetValueBool(const FName Name) const;
+	UE_API TValueOrError<uint8, EPropertyBagResult> GetValueByte(const FName Name) const;
+	UE_API TValueOrError<int32, EPropertyBagResult> GetValueInt32(const FName Name) const;
+	UE_API TValueOrError<uint32, EPropertyBagResult> GetValueUInt32(const FName Name) const;
+	UE_API TValueOrError<int64, EPropertyBagResult> GetValueInt64(const FName Name) const;
+	UE_API TValueOrError<uint64, EPropertyBagResult> GetValueUInt64(const FName Name) const;
+	UE_API TValueOrError<float, EPropertyBagResult> GetValueFloat(const FName Name) const;
+	UE_API TValueOrError<double, EPropertyBagResult> GetValueDouble(const FName Name) const;
+	UE_API TValueOrError<FName, EPropertyBagResult> GetValueName(const FName Name) const;
+	UE_API TValueOrError<FString, EPropertyBagResult> GetValueString(const FName Name) const;
+	UE_API TValueOrError<FText, EPropertyBagResult> GetValueText(const FName Name) const;
+	UE_API TValueOrError<uint8, EPropertyBagResult> GetValueEnum(const FName Name, const UEnum* RequestedEnum) const;
+	UE_API TValueOrError<FStructView, EPropertyBagResult> GetValueStruct(const FName Name, const UScriptStruct* RequestedStruct = nullptr) const;
+	UE_API TValueOrError<UObject*, EPropertyBagResult> GetValueObject(const FName Name, const UClass* RequestedClass = nullptr) const;
+	UE_API TValueOrError<UClass*, EPropertyBagResult> GetValueClass(const FName Name) const;
+	UE_API TValueOrError<FSoftObjectPath, EPropertyBagResult> GetValueSoftPath(const FName Name) const;
 
 	/** @return string-based serialized representation of the value. */
-	TValueOrError<FString, EPropertyBagResult> GetValueSerializedString(const FName Name) const;
+	UE_API TValueOrError<FString, EPropertyBagResult> GetValueSerializedString(const FName Name) const;
 
 	/** @return enum value of specified type. */
 	template <typename T>
@@ -641,26 +643,26 @@ struct COREUOBJECT_API FInstancedPropertyBag
 		return MakeError(EPropertyBagResult::TypeMismatch);
 	}
 
-	TValueOrError<bool, EPropertyBagResult> GetValueBool(const FPropertyBagPropertyDesc& Desc) const;
-	TValueOrError<uint8, EPropertyBagResult> GetValueByte(const FPropertyBagPropertyDesc& Desc) const;
-	TValueOrError<int32, EPropertyBagResult> GetValueInt32(const FPropertyBagPropertyDesc& Desc) const;
-	TValueOrError<uint32, EPropertyBagResult> GetValueUInt32(const FPropertyBagPropertyDesc& Desc) const;
-	TValueOrError<int64, EPropertyBagResult> GetValueInt64(const FPropertyBagPropertyDesc& Desc) const;
-	TValueOrError<uint64, EPropertyBagResult> GetValueUInt64(const FPropertyBagPropertyDesc& Desc) const;
-	TValueOrError<float, EPropertyBagResult> GetValueFloat(const FPropertyBagPropertyDesc& Desc) const;
-	TValueOrError<double, EPropertyBagResult> GetValueDouble(const FPropertyBagPropertyDesc& Desc) const;
-	TValueOrError<FName, EPropertyBagResult> GetValueName(const FPropertyBagPropertyDesc& Desc) const;
-	TValueOrError<FString, EPropertyBagResult> GetValueString(const FPropertyBagPropertyDesc& Desc) const;
-	TValueOrError<FText, EPropertyBagResult> GetValueText(const FPropertyBagPropertyDesc& Desc) const;
-	TValueOrError<uint8, EPropertyBagResult> GetValueEnum(const FPropertyBagPropertyDesc& Desc, const UEnum* RequestedEnum) const;
-	TValueOrError<FStructView, EPropertyBagResult> GetValueStruct(const FPropertyBagPropertyDesc& Desc, const UScriptStruct* RequestedStruct = nullptr) const;
-	TValueOrError<UObject*, EPropertyBagResult> GetValueObject(const FPropertyBagPropertyDesc& Desc, const UClass* RequestedClass = nullptr) const;
-	TValueOrError<UClass*, EPropertyBagResult> GetValueClass(const FPropertyBagPropertyDesc& Desc) const;
-	TValueOrError<FSoftObjectPath, EPropertyBagResult> GetValueSoftPath(const FPropertyBagPropertyDesc& Desc) const;
-	TValueOrError<FPropertyBagArrayRef, EPropertyBagResult> GetMutableArrayRef(const FPropertyBagPropertyDesc& Desc);
-	TValueOrError<const FPropertyBagArrayRef, EPropertyBagResult> GetArrayRef(const FPropertyBagPropertyDesc& Desc) const;
-	TValueOrError<FPropertyBagSetRef, EPropertyBagResult> GetMutableSetRef(const FPropertyBagPropertyDesc& Desc);
-	TValueOrError<const FPropertyBagSetRef, EPropertyBagResult> GetSetRef(const FPropertyBagPropertyDesc& Desc) const;
+	UE_API TValueOrError<bool, EPropertyBagResult> GetValueBool(const FPropertyBagPropertyDesc& Desc) const;
+	UE_API TValueOrError<uint8, EPropertyBagResult> GetValueByte(const FPropertyBagPropertyDesc& Desc) const;
+	UE_API TValueOrError<int32, EPropertyBagResult> GetValueInt32(const FPropertyBagPropertyDesc& Desc) const;
+	UE_API TValueOrError<uint32, EPropertyBagResult> GetValueUInt32(const FPropertyBagPropertyDesc& Desc) const;
+	UE_API TValueOrError<int64, EPropertyBagResult> GetValueInt64(const FPropertyBagPropertyDesc& Desc) const;
+	UE_API TValueOrError<uint64, EPropertyBagResult> GetValueUInt64(const FPropertyBagPropertyDesc& Desc) const;
+	UE_API TValueOrError<float, EPropertyBagResult> GetValueFloat(const FPropertyBagPropertyDesc& Desc) const;
+	UE_API TValueOrError<double, EPropertyBagResult> GetValueDouble(const FPropertyBagPropertyDesc& Desc) const;
+	UE_API TValueOrError<FName, EPropertyBagResult> GetValueName(const FPropertyBagPropertyDesc& Desc) const;
+	UE_API TValueOrError<FString, EPropertyBagResult> GetValueString(const FPropertyBagPropertyDesc& Desc) const;
+	UE_API TValueOrError<FText, EPropertyBagResult> GetValueText(const FPropertyBagPropertyDesc& Desc) const;
+	UE_API TValueOrError<uint8, EPropertyBagResult> GetValueEnum(const FPropertyBagPropertyDesc& Desc, const UEnum* RequestedEnum) const;
+	UE_API TValueOrError<FStructView, EPropertyBagResult> GetValueStruct(const FPropertyBagPropertyDesc& Desc, const UScriptStruct* RequestedStruct = nullptr) const;
+	UE_API TValueOrError<UObject*, EPropertyBagResult> GetValueObject(const FPropertyBagPropertyDesc& Desc, const UClass* RequestedClass = nullptr) const;
+	UE_API TValueOrError<UClass*, EPropertyBagResult> GetValueClass(const FPropertyBagPropertyDesc& Desc) const;
+	UE_API TValueOrError<FSoftObjectPath, EPropertyBagResult> GetValueSoftPath(const FPropertyBagPropertyDesc& Desc) const;
+	UE_API TValueOrError<FPropertyBagArrayRef, EPropertyBagResult> GetMutableArrayRef(const FPropertyBagPropertyDesc& Desc);
+	UE_API TValueOrError<const FPropertyBagArrayRef, EPropertyBagResult> GetArrayRef(const FPropertyBagPropertyDesc& Desc) const;
+	UE_API TValueOrError<FPropertyBagSetRef, EPropertyBagResult> GetMutableSetRef(const FPropertyBagPropertyDesc& Desc);
+	UE_API TValueOrError<const FPropertyBagSetRef, EPropertyBagResult> GetSetRef(const FPropertyBagPropertyDesc& Desc) const;
 
 	/** @return enum value of specified type. */
 	template <typename T>
@@ -718,29 +720,29 @@ struct COREUOBJECT_API FInstancedPropertyBag
 	 * Value Setters. A property must exist in that bag before it can be set.  
 	 * Numeric types (bool, (u)int32, (u)int64, float, double) support type conversion.
 	 */
-	EPropertyBagResult SetValueBool(const FName Name, const bool bInValue);
-	EPropertyBagResult SetValueByte(const FName Name, const uint8 InValue);
-	EPropertyBagResult SetValueInt32(const FName Name, const int32 InValue);
-	EPropertyBagResult SetValueUInt32(const FName Name, const uint32 InValue);
-	EPropertyBagResult SetValueInt64(const FName Name, const int64 InValue);
-	EPropertyBagResult SetValueUInt64(const FName Name, const uint64 InValue);
-	EPropertyBagResult SetValueFloat(const FName Name, const float InValue);
-	EPropertyBagResult SetValueDouble(const FName Name, const double InValue);
-	EPropertyBagResult SetValueName(const FName Name, const FName InValue);
-	EPropertyBagResult SetValueString(const FName Name, const FString& InValue);
-	EPropertyBagResult SetValueText(const FName Name, const FText& InValue);
-	EPropertyBagResult SetValueEnum(const FName Name, const uint8 InValue, const UEnum* Enum);
-	EPropertyBagResult SetValueStruct(const FName Name, FConstStructView InValue);
-	EPropertyBagResult SetValueObject(const FName Name, UObject* InValue);
-	EPropertyBagResult SetValueClass(const FName Name, UClass* InValue);
-	EPropertyBagResult SetValueSoftPath(const FName Name, const FSoftObjectPath& InValue);
-	EPropertyBagResult SetValueSoftPath(const FName Name, const UObject* InValue);
+	UE_API EPropertyBagResult SetValueBool(const FName Name, const bool bInValue);
+	UE_API EPropertyBagResult SetValueByte(const FName Name, const uint8 InValue);
+	UE_API EPropertyBagResult SetValueInt32(const FName Name, const int32 InValue);
+	UE_API EPropertyBagResult SetValueUInt32(const FName Name, const uint32 InValue);
+	UE_API EPropertyBagResult SetValueInt64(const FName Name, const int64 InValue);
+	UE_API EPropertyBagResult SetValueUInt64(const FName Name, const uint64 InValue);
+	UE_API EPropertyBagResult SetValueFloat(const FName Name, const float InValue);
+	UE_API EPropertyBagResult SetValueDouble(const FName Name, const double InValue);
+	UE_API EPropertyBagResult SetValueName(const FName Name, const FName InValue);
+	UE_API EPropertyBagResult SetValueString(const FName Name, const FString& InValue);
+	UE_API EPropertyBagResult SetValueText(const FName Name, const FText& InValue);
+	UE_API EPropertyBagResult SetValueEnum(const FName Name, const uint8 InValue, const UEnum* Enum);
+	UE_API EPropertyBagResult SetValueStruct(const FName Name, FConstStructView InValue);
+	UE_API EPropertyBagResult SetValueObject(const FName Name, UObject* InValue);
+	UE_API EPropertyBagResult SetValueClass(const FName Name, UClass* InValue);
+	UE_API EPropertyBagResult SetValueSoftPath(const FName Name, const FSoftObjectPath& InValue);
+	UE_API EPropertyBagResult SetValueSoftPath(const FName Name, const UObject* InValue);
 
 	/**
 	 * Sets property value from a serialized representation of the value. If the string value provided
 	 * cannot be parsed by the property, the operation will fail.
 	 */
-	EPropertyBagResult SetValueSerializedString(const FName Name, const FString& InValue);
+	UE_API EPropertyBagResult SetValueSerializedString(const FName Name, const FString& InValue);
 
 	/** Sets enum value specified type. */
 	template <typename T>
@@ -769,25 +771,25 @@ struct COREUOBJECT_API FInstancedPropertyBag
 	 * Sets property value from given source property and source container address
 	 * A property must exist in that bag before it can be set. 
 	 */
-	EPropertyBagResult SetValue(const FName Name, const FProperty* InSourceProperty, const void* InSourceContainerAddress);
-
-	EPropertyBagResult SetValueBool(const FPropertyBagPropertyDesc& Desc, const bool bInValue);
-	EPropertyBagResult SetValueByte(const FPropertyBagPropertyDesc& Desc, const uint8 InValue);
-	EPropertyBagResult SetValueInt32(const FPropertyBagPropertyDesc& Desc, const int32 InValue);
-	EPropertyBagResult SetValueUInt32(const FPropertyBagPropertyDesc& Desc, const uint32 InValue);
-	EPropertyBagResult SetValueInt64(const FPropertyBagPropertyDesc& Desc, const int64 InValue);
-	EPropertyBagResult SetValueUInt64(const FPropertyBagPropertyDesc& Desc, const uint64 InValue);
-	EPropertyBagResult SetValueFloat(const FPropertyBagPropertyDesc& Desc, const float InValue);
-	EPropertyBagResult SetValueDouble(const FPropertyBagPropertyDesc& Desc, const double InValue);
-	EPropertyBagResult SetValueName(const FPropertyBagPropertyDesc& Desc, const FName InValue);
-	EPropertyBagResult SetValueString(const FPropertyBagPropertyDesc& Desc, const FString& InValue);
-	EPropertyBagResult SetValueText(const FPropertyBagPropertyDesc& Desc, const FText& InValue);
-	EPropertyBagResult SetValueEnum(const FPropertyBagPropertyDesc& Desc, const uint8 InValue, const UEnum* Enum);
-	EPropertyBagResult SetValueStruct(const FPropertyBagPropertyDesc& Desc, FConstStructView InValue);
-	EPropertyBagResult SetValueObject(const FPropertyBagPropertyDesc& Desc, UObject* InValue);
-	EPropertyBagResult SetValueClass(const FPropertyBagPropertyDesc& Desc, UClass* InValue);
-	EPropertyBagResult SetValueSoftPath(const FPropertyBagPropertyDesc& Desc, const FSoftObjectPath& InValue);
-	EPropertyBagResult SetValueSoftPath(const FPropertyBagPropertyDesc& Desc, const UObject* InValue);
+	UE_API EPropertyBagResult SetValue(const FName Name, const FProperty* InSourceProperty, const void* InSourceContainerAddress);
+
+	UE_API EPropertyBagResult SetValueBool(const FPropertyBagPropertyDesc& Desc, const bool bInValue);
+	UE_API EPropertyBagResult SetValueByte(const FPropertyBagPropertyDesc& Desc, const uint8 InValue);
+	UE_API EPropertyBagResult SetValueInt32(const FPropertyBagPropertyDesc& Desc, const int32 InValue);
+	UE_API EPropertyBagResult SetValueUInt32(const FPropertyBagPropertyDesc& Desc, const uint32 InValue);
+	UE_API EPropertyBagResult SetValueInt64(const FPropertyBagPropertyDesc& Desc, const int64 InValue);
+	UE_API EPropertyBagResult SetValueUInt64(const FPropertyBagPropertyDesc& Desc, const uint64 InValue);
+	UE_API EPropertyBagResult SetValueFloat(const FPropertyBagPropertyDesc& Desc, const float InValue);
+	UE_API EPropertyBagResult SetValueDouble(const FPropertyBagPropertyDesc& Desc, const double InValue);
+	UE_API EPropertyBagResult SetValueName(const FPropertyBagPropertyDesc& Desc, const FName InValue);
+	UE_API EPropertyBagResult SetValueString(const FPropertyBagPropertyDesc& Desc, const FString& InValue);
+	UE_API EPropertyBagResult SetValueText(const FPropertyBagPropertyDesc& Desc, const FText& InValue);
+	UE_API EPropertyBagResult SetValueEnum(const FPropertyBagPropertyDesc& Desc, const uint8 InValue, const UEnum* Enum);
+	UE_API EPropertyBagResult SetValueStruct(const FPropertyBagPropertyDesc& Desc, FConstStructView InValue);
+	UE_API EPropertyBagResult SetValueObject(const FPropertyBagPropertyDesc& Desc, UObject* InValue);
+	UE_API EPropertyBagResult SetValueClass(const FPropertyBagPropertyDesc& Desc, UClass* InValue);
+	UE_API EPropertyBagResult SetValueSoftPath(const FPropertyBagPropertyDesc& Desc, const FSoftObjectPath& InValue);
+	UE_API EPropertyBagResult SetValueSoftPath(const FPropertyBagPropertyDesc& Desc, const UObject* InValue);
 
 	/** Sets enum value specified type. */
 	template <typename T>
@@ -817,33 +819,33 @@ struct COREUOBJECT_API FInstancedPropertyBag
 	 * Note: The array reference is not valid after the layout of the referenced property bag has changed!
 	 * @returns helper class to modify and access arrays
 	*/
-	TValueOrError<FPropertyBagArrayRef, EPropertyBagResult> GetMutableArrayRef(const FName Name);
+	UE_API TValueOrError<FPropertyBagArrayRef, EPropertyBagResult> GetMutableArrayRef(const FName Name);
 
 	/**
 	 * Returns helper class to access an array property.
 	 * Note: The array reference is not valid after the layout of the referenced property bag has changed!
 	 * @returns helper class to modify and access arrays
 	*/
-	TValueOrError<const FPropertyBagArrayRef, EPropertyBagResult> GetArrayRef(const FName Name) const;
+	UE_API TValueOrError<const FPropertyBagArrayRef, EPropertyBagResult> GetArrayRef(const FName Name) const;
 
 	/**
 	 * Returns helper class to modify and access a set property.
 	 * Note: The set reference is not valid after the layout of the referenced property bag has changed!
 	 * @returns helper class to modify and access sets
 	*/
-	TValueOrError<FPropertyBagSetRef, EPropertyBagResult> GetMutableSetRef(const FName Name);
+	UE_API TValueOrError<FPropertyBagSetRef, EPropertyBagResult> GetMutableSetRef(const FName Name);
 	
 	/**
 	 * Returns helper class to access a set property.
 	 * Note: The set reference is not valid after the layout of the referenced property bag has changed!
 	 * @returns helper class to modify and access sets
 	*/
-	TValueOrError<const FPropertyBagSetRef, EPropertyBagResult> GetSetRef(const FName Name) const;
+	UE_API TValueOrError<const FPropertyBagSetRef, EPropertyBagResult> GetSetRef(const FName Name) const;
 
-	bool Identical(const FInstancedPropertyBag* Other, uint32 PortFlags) const;
-	bool Serialize(FArchive& Ar);
-	void AddStructReferencedObjects(FReferenceCollector& Collector);
-	void GetPreloadDependencies(TArray<UObject*>& OutDeps);
+	UE_API bool Identical(const FInstancedPropertyBag* Other, uint32 PortFlags) const;
+	UE_API bool Serialize(FArchive& Ar);
+	UE_API void AddStructReferencedObjects(FReferenceCollector& Collector);
+	UE_API void GetPreloadDependencies(TArray<UObject*>& OutDeps);
 
 	/**
 	 * Checks whether a provided name is a valid property bag property name.
@@ -880,8 +882,8 @@ struct COREUOBJECT_API FInstancedPropertyBag
 	static FName SanitizePropertyName(FName Name);
 
 protected:
-	const void* GetValueAddress(const FPropertyBagPropertyDesc* Desc) const;
-	void* GetMutableValueAddress(const FPropertyBagPropertyDesc* Desc);
+	UE_API const void* GetValueAddress(const FPropertyBagPropertyDesc* Desc) const;
+	UE_API void* GetMutableValueAddress(const FPropertyBagPropertyDesc* Desc);
 	
 	UPROPERTY(EditAnywhere, Category="", meta=(SequencerUseParentPropertyName=true))
 	FInstancedStruct Value;
@@ -917,7 +919,7 @@ template<> struct TStructOpsTypeTraits<FInstancedPropertyBag> : public TStructOp
  * 
  * Note: The array reference is not valid after the layout of the referenced property bag has changed! 
  */
-class COREUOBJECT_API FPropertyBagArrayRef : public FScriptArrayHelper
+class FPropertyBagArrayRef : public FScriptArrayHelper
 {
 	FPropertyBagPropertyDesc ValueDesc;
 
@@ -963,22 +965,22 @@ public:
 	 * Numeric types (bool, (u)int32, (u)int64, float, double) support type conversion.
 	 */
 	
-	TValueOrError<bool, EPropertyBagResult> GetValueBool(const int32 Index) const;
-	TValueOrError<uint8, EPropertyBagResult> GetValueByte(const int32 Index) const;
-	TValueOrError<int32, EPropertyBagResult> GetValueInt32(const int32 Index) const;
-	TValueOrError<uint32, EPropertyBagResult> GetValueUInt32(const int32 Index) const;
-	TValueOrError<int64, EPropertyBagResult> GetValueInt64(const int32 Index) const;
-	TValueOrError<uint64, EPropertyBagResult> GetValueUInt64(const int32 Index) const;
-	TValueOrError<float, EPropertyBagResult> GetValueFloat(const int32 Index) const;
-	TValueOrError<double, EPropertyBagResult> GetValueDouble(const int32 Index) const;
-	TValueOrError<FName, EPropertyBagResult> GetValueName(const int32 Index) const;
-	TValueOrError<FString, EPropertyBagResult> GetValueString(const int32 Index) const;
-	TValueOrError<FText, EPropertyBagResult> GetValueText(const int32 Index) const;
-	TValueOrError<uint8, EPropertyBagResult> GetValueEnum(const int32 Index, const UEnum* RequestedEnum) const;
-	TValueOrError<FStructView, EPropertyBagResult> GetValueStruct(const int32 Index, const UScriptStruct* RequestedStruct = nullptr) const;
-	TValueOrError<UObject*, EPropertyBagResult> GetValueObject(const int32 Index, const UClass* RequestedClass = nullptr) const;
-	TValueOrError<UClass*, EPropertyBagResult> GetValueClass(const int32 Index) const;
-	TValueOrError<FSoftObjectPath, EPropertyBagResult> GetValueSoftPath(const int32 Index) const;
+	UE_API TValueOrError<bool, EPropertyBagResult> GetValueBool(const int32 Index) const;
+	UE_API TValueOrError<uint8, EPropertyBagResult> GetValueByte(const int32 Index) const;
+	UE_API TValueOrError<int32, EPropertyBagResult> GetValueInt32(const int32 Index) const;
+	UE_API TValueOrError<uint32, EPropertyBagResult> GetValueUInt32(const int32 Index) const;
+	UE_API TValueOrError<int64, EPropertyBagResult> GetValueInt64(const int32 Index) const;
+	UE_API TValueOrError<uint64, EPropertyBagResult> GetValueUInt64(const int32 Index) const;
+	UE_API TValueOrError<float, EPropertyBagResult> GetValueFloat(const int32 Index) const;
+	UE_API TValueOrError<double, EPropertyBagResult> GetValueDouble(const int32 Index) const;
+	UE_API TValueOrError<FName, EPropertyBagResult> GetValueName(const int32 Index) const;
+	UE_API TValueOrError<FString, EPropertyBagResult> GetValueString(const int32 Index) const;
+	UE_API TValueOrError<FText, EPropertyBagResult> GetValueText(const int32 Index) const;
+	UE_API TValueOrError<uint8, EPropertyBagResult> GetValueEnum(const int32 Index, const UEnum* RequestedEnum) const;
+	UE_API TValueOrError<FStructView, EPropertyBagResult> GetValueStruct(const int32 Index, const UScriptStruct* RequestedStruct = nullptr) const;
+	UE_API TValueOrError<UObject*, EPropertyBagResult> GetValueObject(const int32 Index, const UClass* RequestedClass = nullptr) const;
+	UE_API TValueOrError<UClass*, EPropertyBagResult> GetValueClass(const int32 Index) const;
+	UE_API TValueOrError<FSoftObjectPath, EPropertyBagResult> GetValueSoftPath(const int32 Index) const;
 
 	/** @return enum value of specified type. */
 	template <typename T>
@@ -1037,36 +1039,36 @@ public:
      * Note: Note: The array reference is not valid after the layout of the referenced property bag has changed!
      * @returns helper class to modify and access arrays
     */
-	TValueOrError<FPropertyBagArrayRef, EPropertyBagResult> GetMutableNestedArrayRef(const int32 Index = 0) const;
+	UE_API TValueOrError<FPropertyBagArrayRef, EPropertyBagResult> GetMutableNestedArrayRef(const int32 Index = 0) const;
 
 	/**
      * Returns helper class to access a nested array (const version).
      * Note: Note: The array reference is not valid after the layout of the referenced property bag has changed!
      * @returns helper class to access arrays
     */
-	TValueOrError<const FPropertyBagArrayRef, EPropertyBagResult> GetNestedArrayRef(const int32 Index = 0) const;
+	UE_API TValueOrError<const FPropertyBagArrayRef, EPropertyBagResult> GetNestedArrayRef(const int32 Index = 0) const;
 
 	/**
 	 * Value Setters. A property must exist in that bag before it can be set.  
 	 * Numeric types (bool, (u)int32, (u)int64, float, double) support type conversion.
 	 */
-	EPropertyBagResult SetValueBool(const int32 Index, const bool bInValue);
-	EPropertyBagResult SetValueByte(const int32 Index, const uint8 InValue);
-	EPropertyBagResult SetValueInt32(const int32 Index, const int32 InValue);
-	EPropertyBagResult SetValueUInt32(const int32 Index, const uint32 InValue);
-	EPropertyBagResult SetValueInt64(const int32 Index, const int64 InValue);
-	EPropertyBagResult SetValueUInt64(const int32 Index, const uint64 InValue);
-	EPropertyBagResult SetValueFloat(const int32 Index, const float InValue);
-	EPropertyBagResult SetValueDouble(const int32 Index, const double InValue);
-	EPropertyBagResult SetValueName(const int32 Index, const FName InValue);
-	EPropertyBagResult SetValueString(const int32 Index, const FString& InValue);
-	EPropertyBagResult SetValueText(const int32 Index, const FText& InValue);
-	EPropertyBagResult SetValueEnum(const int32 Index, const uint8 InValue, const UEnum* Enum);
-	EPropertyBagResult SetValueStruct(const int32 Index, FConstStructView InValue);
-	EPropertyBagResult SetValueObject(const int32 Index, UObject* InValue);
-	EPropertyBagResult SetValueClass(const int32 Index, UClass* InValue);
-	EPropertyBagResult SetValueSoftPath(const int32 Index, const FSoftObjectPath& InValue);
-	EPropertyBagResult SetValueSoftPath(const int32 Index, const UObject* InValue);
+	UE_API EPropertyBagResult SetValueBool(const int32 Index, const bool bInValue);
+	UE_API EPropertyBagResult SetValueByte(const int32 Index, const uint8 InValue);
+	UE_API EPropertyBagResult SetValueInt32(const int32 Index, const int32 InValue);
+	UE_API EPropertyBagResult SetValueUInt32(const int32 Index, const uint32 InValue);
+	UE_API EPropertyBagResult SetValueInt64(const int32 Index, const int64 InValue);
+	UE_API EPropertyBagResult SetValueUInt64(const int32 Index, const uint64 InValue);
+	UE_API EPropertyBagResult SetValueFloat(const int32 Index, const float InValue);
+	UE_API EPropertyBagResult SetValueDouble(const int32 Index, const double InValue);
+	UE_API EPropertyBagResult SetValueName(const int32 Index, const FName InValue);
+	UE_API EPropertyBagResult SetValueString(const int32 Index, const FString& InValue);
+	UE_API EPropertyBagResult SetValueText(const int32 Index, const FText& InValue);
+	UE_API EPropertyBagResult SetValueEnum(const int32 Index, const uint8 InValue, const UEnum* Enum);
+	UE_API EPropertyBagResult SetValueStruct(const int32 Index, FConstStructView InValue);
+	UE_API EPropertyBagResult SetValueObject(const int32 Index, UObject* InValue);
+	UE_API EPropertyBagResult SetValueClass(const int32 Index, UClass* InValue);
+	UE_API EPropertyBagResult SetValueSoftPath(const int32 Index, const FSoftObjectPath& InValue);
+	UE_API EPropertyBagResult SetValueSoftPath(const int32 Index, const UObject* InValue);
 
 	/** Sets enum value specified type. */
 	template <typename T>
@@ -1109,7 +1111,7 @@ public:
  * 
  * Note: The set reference is not valid after the layout of the referenced property bag has changed! 
  */
-class COREUOBJECT_API FPropertyBagSetRef : private FScriptSetHelper
+class FPropertyBagSetRef : private FScriptSetHelper
 {
 public:	
 	FORCEINLINE FPropertyBagSetRef(const FPropertyBagPropertyDesc& InDesc, const void* InSet)
@@ -1127,22 +1129,22 @@ public:
 	}
 
 	/** Add value to set. If value is already present, it will not be added */
-	EPropertyBagResult AddValueBool(const bool bInValue);
-	EPropertyBagResult AddValueByte(const uint8 InValue);
-	EPropertyBagResult AddValueInt32(const int32 InValue);
-	EPropertyBagResult AddValueUInt32(const uint32 InValue);
-	EPropertyBagResult AddValueInt64(const int64 InValue);
-	EPropertyBagResult AddValueUInt64(const uint64 InValue);
-	EPropertyBagResult AddValueFloat(const float InValue);
-	EPropertyBagResult AddValueDouble(const double InValue);
-	EPropertyBagResult AddValueName(const FName InValue);
-	EPropertyBagResult AddValueString(const FString& InValue);
-	EPropertyBagResult AddValueText(const FText& InValue);
-	EPropertyBagResult AddValueEnum(const int64 InValue, const UEnum* Enum);
-	EPropertyBagResult AddValueStruct(FConstStructView InValue);
-	EPropertyBagResult AddValueObject(UObject* InValue);
-	EPropertyBagResult AddValueClass(UClass* InValue);
-	EPropertyBagResult AddValueSoftPath(const FSoftObjectPath& InValue);
+	UE_API EPropertyBagResult AddValueBool(const bool bInValue);
+	UE_API EPropertyBagResult AddValueByte(const uint8 InValue);
+	UE_API EPropertyBagResult AddValueInt32(const int32 InValue);
+	UE_API EPropertyBagResult AddValueUInt32(const uint32 InValue);
+	UE_API EPropertyBagResult AddValueInt64(const int64 InValue);
+	UE_API EPropertyBagResult AddValueUInt64(const uint64 InValue);
+	UE_API EPropertyBagResult AddValueFloat(const float InValue);
+	UE_API EPropertyBagResult AddValueDouble(const double InValue);
+	UE_API EPropertyBagResult AddValueName(const FName InValue);
+	UE_API EPropertyBagResult AddValueString(const FString& InValue);
+	UE_API EPropertyBagResult AddValueText(const FText& InValue);
+	UE_API EPropertyBagResult AddValueEnum(const int64 InValue, const UEnum* Enum);
+	UE_API EPropertyBagResult AddValueStruct(FConstStructView InValue);
+	UE_API EPropertyBagResult AddValueObject(UObject* InValue);
+	UE_API EPropertyBagResult AddValueClass(UClass* InValue);
+	UE_API EPropertyBagResult AddValueSoftPath(const FSoftObjectPath& InValue);
 
 	/** Adds enum value specified type. */
 	template <typename T>
@@ -1235,13 +1237,13 @@ enum class EPropertyBagMissingEnum : uint8
 };
 
 USTRUCT()
-struct COREUOBJECT_API FPropertyBagMissingStruct
+struct FPropertyBagMissingStruct
 {
 	GENERATED_BODY()
 };
 
-UCLASS()
-class COREUOBJECT_API UPropertyBagMissingObject : public UObject
+UCLASS(MinimalAPI)
+class UPropertyBagMissingObject : public UObject
 {
 	GENERATED_BODY()
 };
@@ -1254,8 +1256,8 @@ class COREUOBJECT_API UPropertyBagMissingObject : public UObject
  *
  * Note: Should not be used directly.
  */
-UCLASS(Transient)
-class COREUOBJECT_API UPropertyBag : public UScriptStruct
+UCLASS(Transient, MinimalAPI)
+class UPropertyBag : public UScriptStruct
 {
 public:
 	GENERATED_BODY()
@@ -1268,42 +1270,42 @@ public:
 	 * If there are multiple properties that have the same name, only the first property is added.
 	 * The caller is expected to ensure unique names for the property descriptions.
 	 */
-	static const UPropertyBag* GetOrCreateFromDescs(const TConstArrayView<FPropertyBagPropertyDesc> InPropertyDescs, const TCHAR* PrefixName = nullptr);
+	static UE_API const UPropertyBag* GetOrCreateFromDescs(const TConstArrayView<FPropertyBagPropertyDesc> InPropertyDescs, const TCHAR* PrefixName = nullptr);
 
 	/** Returns property descriptions that specify this struct. */
 	TConstArrayView<FPropertyBagPropertyDesc> GetPropertyDescs() const { return PropertyDescs; }
 
 	/** @return property description based on ID. */
-	const FPropertyBagPropertyDesc* FindPropertyDescByID(const FGuid ID) const;
+	UE_API const FPropertyBagPropertyDesc* FindPropertyDescByID(const FGuid ID) const;
 
 	/** @return property description based on name. */
-	const FPropertyBagPropertyDesc* FindPropertyDescByName(const FName Name) const;
+	UE_API const FPropertyBagPropertyDesc* FindPropertyDescByName(const FName Name) const;
 
 	/** @return property description based on the created property name. The name can be different from the descriptor name due to name sanitization. */
-	const FPropertyBagPropertyDesc* FindPropertyDescByPropertyName(const FName PropertyName) const;
+	UE_API const FPropertyBagPropertyDesc* FindPropertyDescByPropertyName(const FName PropertyName) const;
 
 	/** @return property description based on pointer to property. */
-	const FPropertyBagPropertyDesc* FindPropertyDescByProperty(const FProperty* Property) const;
+	UE_API const FPropertyBagPropertyDesc* FindPropertyDescByProperty(const FProperty* Property) const;
 
 	/** @return property description based on index. */
-	const FPropertyBagPropertyDesc* FindPropertyDescByIndex(int32 Index) const;
+	UE_API const FPropertyBagPropertyDesc* FindPropertyDescByIndex(int32 Index) const;
 
 	/** @return true if we own the supplied property description, false otherwise. */
-	bool OwnsPropertyDesc(const FPropertyBagPropertyDesc& Desc) const;
+	UE_API bool OwnsPropertyDesc(const FPropertyBagPropertyDesc& Desc) const;
 
 #if WITH_EDITOR
 	/** @return true if any of the properties on the bag has type of the specified user defined struct. */
-	bool ContainsUserDefinedStruct(const UUserDefinedStruct* UserDefinedStruct) const;
+	UE_API bool ContainsUserDefinedStruct(const UUserDefinedStruct* UserDefinedStruct) const;
 #endif // WITH_EDITOR
 	
 protected:
 
-	void DecrementRefCount() const;
-	void IncrementRefCount() const;
+	UE_API void DecrementRefCount() const;
+	UE_API void IncrementRefCount() const;
 	
-	virtual void InitializeStruct(void* Dest, int32 ArrayDim = 1) const override;
-	virtual void DestroyStruct(void* Dest, int32 ArrayDim = 1) const override;
-	virtual void FinishDestroy() override;
+	UE_API virtual void InitializeStruct(void* Dest, int32 ArrayDim = 1) const override;
+	UE_API virtual void DestroyStruct(void* Dest, int32 ArrayDim = 1) const override;
+	UE_API virtual void FinishDestroy() override;
 	
 	UPROPERTY()
 	TArray<FPropertyBagPropertyDesc> PropertyDescs;
@@ -1312,3 +1314,5 @@ protected:
 	
 	friend struct FInstancedPropertyBag;
 };
+
+#undef UE_API
\ No newline at end of file
diff --git a/Engine/Source/Runtime/CoreUObject/Public/StructUtils/SharedStruct.h b/Engine/Source/Runtime/CoreUObject/Public/StructUtils/SharedStruct.h
--- a/Engine/Source/Runtime/CoreUObject/Public/StructUtils/SharedStruct.h
+++ b/Engine/Source/Runtime/CoreUObject/Public/StructUtils/SharedStruct.h
@@ -8,6 +8,8 @@
 
 #include "SharedStruct.generated.h"
 
+#define UE_API COREUOBJECT_API
+
 struct FInstancedStruct;
 template<typename BaseStructT> struct TInstancedStruct;
 template<typename BaseStructT> struct TStructView;
@@ -30,7 +32,7 @@ template<typename BaseStructT> struct TConstStructView;
  * 
  * This type is designed to be used in composition and should not be used outside the types that compose it.
  */
-struct COREUOBJECT_API FStructSharedMemory
+struct FStructSharedMemory
 {
 	~FStructSharedMemory()
 	{
@@ -152,7 +154,7 @@ private:
  * In either case, the shared struct memory /data is mutable.
  */
 USTRUCT()
-struct [[nodiscard]] COREUOBJECT_API FSharedStruct
+struct [[nodiscard]] FSharedStruct
 {
 	GENERATED_BODY();
 
@@ -173,8 +175,8 @@ struct [[nodiscard]] COREUOBJECT_API FSharedStruct
 	FSharedStruct& operator=(FSharedStruct&& InOther) = default;
 
 	/** For StructOpsTypeTraits */
-	bool Identical(const FSharedStruct* Other, uint32 PortFlags) const;
-	void AddStructReferencedObjects(class FReferenceCollector& Collector) const;
+	UE_API bool Identical(const FSharedStruct* Other, uint32 PortFlags) const;
+	UE_API void AddStructReferencedObjects(class FReferenceCollector& Collector) const;
 
 	/** Returns struct type. */
 	const UScriptStruct* GetScriptStruct() const
@@ -517,7 +519,7 @@ struct TStructOpsTypeTraits<FSharedStruct> : public TStructOpsTypeTraitsBase2<FS
  * See FSharedStruct for more information.
  */
 USTRUCT()
-struct [[nodiscard]] COREUOBJECT_API FConstSharedStruct
+struct [[nodiscard]] FConstSharedStruct
 {
 	GENERATED_BODY();
 
@@ -548,8 +550,8 @@ struct [[nodiscard]] COREUOBJECT_API FConstSharedStruct
 	}
 
 	/** For StructOpsTypeTraits */
-	bool Identical(const FConstSharedStruct* Other, uint32 PortFlags) const;
-	void AddStructReferencedObjects(class FReferenceCollector& Collector);
+	UE_API bool Identical(const FConstSharedStruct* Other, uint32 PortFlags) const;
+	UE_API void AddStructReferencedObjects(class FReferenceCollector& Collector);
 
 	/** Returns struct type. */
 	const UScriptStruct* GetScriptStruct() const
@@ -906,3 +908,5 @@ struct TStructOpsTypeTraits<FConstSharedStruct> : public TStructOpsTypeTraitsBas
 		WithAddStructReferencedObjects = true,
 	};
 };
+
+#undef UE_API
\ No newline at end of file
diff --git a/Engine/Source/Runtime/CoreUObject/Public/StructUtils/StructTypeBitSet.h b/Engine/Source/Runtime/CoreUObject/Public/StructUtils/StructTypeBitSet.h
--- a/Engine/Source/Runtime/CoreUObject/Public/StructUtils/StructTypeBitSet.h
+++ b/Engine/Source/Runtime/CoreUObject/Public/StructUtils/StructTypeBitSet.h
@@ -7,6 +7,8 @@
 #include "StructUtils/InstancedStruct.h"
 #include "Templates/FunctionFwd.h"
 
+#define UE_API COREUOBJECT_API
+
 class FArchive;
 
 namespace FStructTypeBitSet
@@ -268,7 +270,7 @@ static_assert(sizeof(FBitSetContainer) == sizeof(FConstBitSetContainer), "FBitSe
  *
  * The FStructTracker assigns an index to a given type the first time it encounters it
  */
-struct COREUOBJECT_API FStructTracker
+struct FStructTracker
 {
 	using FBaseStructGetter = TFunction<UStruct*()>;
 	using FTypeValidation = TFunction<bool(const UStruct*)>;
@@ -279,22 +281,22 @@ struct COREUOBJECT_API FStructTracker
 	 * module loading (via DEFINE_TYPEBITSET macro), and the ::StaticStruct call fails at that point for types defined
 	 * in the same module.
 	 */
-	explicit FStructTracker(const FBaseStructGetter& InBaseStructGetter);
-	explicit FStructTracker(const UStruct* InBaseType, const FTypeValidation& InTypeVerification = FTypeValidation());
-	~FStructTracker();
+	UE_API explicit FStructTracker(const FBaseStructGetter& InBaseStructGetter);
+	UE_API explicit FStructTracker(const UStruct* InBaseType, const FTypeValidation& InTypeVerification = FTypeValidation());
+	UE_API ~FStructTracker();
 
 	UE_NONCOPYABLE(FStructTracker);
 
 	/** 
 	 * @return the index for the given type, or INDEX_NONE if it hasn't been registered with this struct tracker
 	 */
-	int32 FindStructTypeIndex(const UStruct& InStructType) const;
+	UE_API int32 FindStructTypeIndex(const UStruct& InStructType) const;
 
 	/**
 	 * Fetches the internal index representing the given UStruct instance. If it hasn't been registered yet, it will be automatically added.
 	 * @return the index for the given type.
 	 */
-	int32 FindOrAddStructTypeIndex(const UStruct& InStructType);
+	UE_API int32 FindOrAddStructTypeIndex(const UStruct& InStructType);
 
 	/**
 	 * Registers the given UStruct with the struct tracker.
@@ -319,7 +321,7 @@ struct COREUOBJECT_API FStructTracker
 	 * Retrieves the base UStruct type used for validation.
 	 * @return Pointer to the base UStruct.
 	 */
-	const UStruct* GetBaseType() const;
+	UE_API const UStruct* GetBaseType() const;
 
 	/**
 	 * Gets the number of registered struct types.
@@ -332,7 +334,7 @@ struct COREUOBJECT_API FStructTracker
 	 * @param Ar - The archive to serialize with.
 	 * @param StructTypesBitArray - The bit array representing struct types.
 	 */
-	void Serialize(FArchive& Ar, FStructTypeBitSet::FBitSetContainer& StructTypesBitArray);
+	UE_API void Serialize(FArchive& Ar, FStructTypeBitSet::FBitSetContainer& StructTypesBitArray);
 
 #if WITH_STRUCTUTILS_DEBUG
 	/**
@@ -359,14 +361,14 @@ struct COREUOBJECT_API FStructTracker
 		DebugStructTypeNamesList.Reset();
 	}
 
-	[[nodiscard]] const UStruct* DebugFindTypeByPartialName(const FString& PartialName) const;
+	[[nodiscard]] UE_API const UStruct* DebugFindTypeByPartialName(const FString& PartialName) const;
 
 protected:
 	mutable TArray<FName, TInlineAllocator<64>> DebugStructTypeNamesList;
 #endif // WITH_STRUCTUTILS_DEBUG
 
 protected:
-	int32 RegisterImplementation(const UStruct& InStructType, const bool bCheckPrevious);
+	UE_API int32 RegisterImplementation(const UStruct& InStructType, const bool bCheckPrevious);
 
 	/** Set mapping struct type hashes to their indices */
 	TSet<uint32> StructTypeToIndexSet;
@@ -1209,3 +1211,5 @@ public:
 
 #define DEFINE_TYPEBITSET(ContainerTypeName) \
 	FStructTracker ContainerTypeName##StructTrackerWrapper::StructTracker([](){ return UE::StructUtils::GetAsUStruct<ContainerTypeName##StructTrackerWrapper::FBaseStructType>();});
+
+#undef UE_API 
\ No newline at end of file
diff --git a/Engine/Source/Runtime/CoreUObject/Public/UObject/Class.h b/Engine/Source/Runtime/CoreUObject/Public/UObject/Class.h
--- a/Engine/Source/Runtime/CoreUObject/Public/UObject/Class.h
+++ b/Engine/Source/Runtime/CoreUObject/Public/UObject/Class.h
@@ -2082,7 +2082,7 @@ public:
 	COREUOBJECT_API virtual void Link(FArchive& Ar, bool bRelinkExistingProperties) override;
 
 	/** Returns parent function if there is one, or null */
-	COREUOBJECT_API UFunction* GetSuperFunction() const;
+	UFunction* GetSuperFunction() const;
 
 	/** Returns the return value property if there is one, or null */
 	COREUOBJECT_API FProperty* GetReturnProperty() const;
diff --git a/Engine/Source/Runtime/CoreUObject/Public/UObject/Field.h b/Engine/Source/Runtime/CoreUObject/Public/UObject/Field.h
--- a/Engine/Source/Runtime/CoreUObject/Public/UObject/Field.h
+++ b/Engine/Source/Runtime/CoreUObject/Public/UObject/Field.h
@@ -203,7 +203,7 @@ public: \
 	{ \
 	} \
 	static TRequiredAPI FFieldClass* StaticClass(); \
-	static FField* Construct(const FFieldVariant& InOwner, const FName& InName, EObjectFlags InObjectFlags); \
+	static TRequiredAPI FField* Construct(const FFieldVariant& InOwner, const FName& InName, EObjectFlags InObjectFlags); \
 	inline static constexpr uint64 StaticClassCastFlagsPrivate() \
 	{ \
 		return uint64(TStaticFlags); \
diff --git a/Engine/Source/Runtime/CoreUObject/Public/UObject/ObjectMacros.h b/Engine/Source/Runtime/CoreUObject/Public/UObject/ObjectMacros.h
--- a/Engine/Source/Runtime/CoreUObject/Public/UObject/ObjectMacros.h
+++ b/Engine/Source/Runtime/CoreUObject/Public/UObject/ObjectMacros.h
@@ -1896,7 +1896,10 @@ public: \
 	}
 
 #define DECLARE_CLASS_INTRINSIC(TClass,TSuperClass,TStaticFlags,TPackage) \
-	DECLARE_CLASS(TClass,TSuperClass,TStaticFlags|CLASS_Intrinsic,CASTCLASS_None,TPackage,NO_API ) \
+	DECLARE_CLASS_INTRINSIC_API(TClass,TSuperClass,TStaticFlags,TPackage, NO_API)
+
+#define DECLARE_CLASS_INTRINSIC_API(TClass,TSuperClass,TStaticFlags,TPackage, TRequiredAPI) \
+	DECLARE_CLASS(TClass,TSuperClass,TStaticFlags|CLASS_Intrinsic,CASTCLASS_None,TPackage,TRequiredAPI) \
 	RELAY_CONSTRUCTOR(TClass, TSuperClass) \
 	/** DO NOT USE. This constructor is for internal usage only for hot-reload purposes. */ \
 	TClass(FVTableHelper& Helper) : Super(Helper) {}; \
@@ -1924,7 +1927,7 @@ public: \
 	DECLARE_CLASS(TClass,TSuperClass,TStaticFlags|CLASS_Intrinsic,TStaticCastFlags,TPackage,TRequiredAPI ) \
 	RELAY_CONSTRUCTOR(TClass, TSuperClass) \
 	/** DO NOT USE. This constructor is for internal usage only for hot-reload purposes. */ \
-	TRequiredAPI TClass(FVTableHelper& Helper) : Super(Helper) {}; \
+	TClass(FVTableHelper& Helper) : Super(Helper) {}; \
 	static void StaticRegisterNatives##TClass() {} \
 	DECLARE_SERIALIZER(TClass) \
 	DEFINE_DEFAULT_OBJECT_INITIALIZER_CONSTRUCTOR_CALL(TClass) \
@@ -1948,7 +1951,7 @@ public: \
 #define DECLARE_CASTED_CLASS_INTRINSIC_NO_CTOR( TClass, TSuperClass, TStaticFlags, TPackage, TStaticCastFlags, TRequiredAPI ) \
 	DECLARE_CASTED_CLASS_INTRINSIC_NO_CTOR_NO_VTABLE_CTOR( TClass, TSuperClass, TStaticFlags, TPackage, TStaticCastFlags, TRequiredAPI ) \
 	/** DO NOT USE. This constructor is for internal usage only for hot-reload purposes. */ \
-	TRequiredAPI TClass(FVTableHelper& Helper) : Super(Helper) {}; \
+	TClass(FVTableHelper& Helper) : Super(Helper) {}; \
 
 
 #define DECLARE_CASTED_CLASS_INTRINSIC( TClass, TSuperClass, TStaticFlags, TPackage, TStaticCastFlags ) \
diff --git a/Engine/Source/Runtime/CoreUObject/Public/UObject/StrProperty.h.inl b/Engine/Source/Runtime/CoreUObject/Public/UObject/StrProperty.h.inl
--- a/Engine/Source/Runtime/CoreUObject/Public/UObject/StrProperty.h.inl
+++ b/Engine/Source/Runtime/CoreUObject/Public/UObject/StrProperty.h.inl
@@ -24,32 +24,32 @@
 	#error "StrProperty.h.inl should only be included after defining UE_STRPROPERTY_PROPERTYPARAMSSTRUCT"
 #endif
 
-class COREUOBJECT_API UE_STRPROPERTY_CLASS : public TProperty_WithEqualityAndSerializer<UE_STRPROPERTY_STRINGTYPE, FProperty>
+class UE_STRPROPERTY_CLASS : public TProperty_WithEqualityAndSerializer<UE_STRPROPERTY_STRINGTYPE, FProperty>
 {
-	DECLARE_FIELD(UE_STRPROPERTY_CLASS, (TProperty_WithEqualityAndSerializer<UE_STRPROPERTY_STRINGTYPE, FProperty>), UE_STRPROPERTY_CASTCLASSFLAG)
+	DECLARE_FIELD_API(UE_STRPROPERTY_CLASS, (TProperty_WithEqualityAndSerializer<UE_STRPROPERTY_STRINGTYPE, FProperty>), UE_STRPROPERTY_CASTCLASSFLAG, COREUOBJECT_API)
 public:
 	using TTypeFundamentals = Super::TTypeFundamentals;
 	using TCppType = TTypeFundamentals::TCppType;
 
-	UE_STRPROPERTY_CLASS(FFieldVariant InOwner, const FName& InName, EObjectFlags InObjectFlags);
+	COREUOBJECT_API UE_STRPROPERTY_CLASS(FFieldVariant InOwner, const FName& InName, EObjectFlags InObjectFlags);
 
 	/**
 	 * Constructor used for constructing compiled in properties
 	 * @param InOwner Owner of the property
 	 * @param PropBase Pointer to the compiled in structure describing the property
 	 **/
-	UE_STRPROPERTY_CLASS(FFieldVariant InOwner, const UECodeGen_Private::UE_STRPROPERTY_PROPERTYPARAMSSTRUCT& Prop);
+	COREUOBJECT_API UE_STRPROPERTY_CLASS(FFieldVariant InOwner, const UECodeGen_Private::UE_STRPROPERTY_PROPERTYPARAMSSTRUCT& Prop);
 
 #if WITH_EDITORONLY_DATA
-	explicit UE_STRPROPERTY_CLASS(UField* InField);
+	COREUOBJECT_API explicit UE_STRPROPERTY_CLASS(UField* InField);
 #endif // WITH_EDITORONLY_DATA
 
 	// FProperty interface
 protected:
-	virtual void ExportText_Internal( FString& ValueStr, const void* PropertyValueOrContainer, EPropertyPointerType PropertyPointerType, const void* DefaultValue, UObject* Parent, int32 PortFlags, UObject* ExportRootScope ) const override;
-	virtual const TCHAR* ImportText_Internal(const TCHAR* Buffer, void* ContainerOrPropertyPtr, EPropertyPointerType PropertyPointerType, UObject* OwnerObject, int32 PortFlags, FOutputDevice* ErrorText) const override;
+	COREUOBJECT_API virtual void ExportText_Internal( FString& ValueStr, const void* PropertyValueOrContainer, EPropertyPointerType PropertyPointerType, const void* DefaultValue, UObject* Parent, int32 PortFlags, UObject* ExportRootScope ) const override;
+	COREUOBJECT_API virtual const TCHAR* ImportText_Internal(const TCHAR* Buffer, void* ContainerOrPropertyPtr, EPropertyPointerType PropertyPointerType, UObject* OwnerObject, int32 PortFlags, FOutputDevice* ErrorText) const override;
 public:
-	virtual EConvertFromTypeResult ConvertFromType(const FPropertyTag& Tag, FStructuredArchive::FSlot Slot, uint8* Data, UStruct* DefaultsStruct, const uint8* Defaults) override;
-	uint32 GetValueTypeHashInternal(const void* Src) const override;
+	COREUOBJECT_API virtual EConvertFromTypeResult ConvertFromType(const FPropertyTag& Tag, FStructuredArchive::FSlot Slot, uint8* Data, UStruct* DefaultsStruct, const uint8* Defaults) override;
+	COREUOBJECT_API uint32 GetValueTypeHashInternal(const void* Src) const override;
 	// End of FProperty interface
 };
diff --git a/Engine/Source/Runtime/CoreUObject/Public/UObject/UnrealType.h b/Engine/Source/Runtime/CoreUObject/Public/UObject/UnrealType.h
--- a/Engine/Source/Runtime/CoreUObject/Public/UObject/UnrealType.h
+++ b/Engine/Source/Runtime/CoreUObject/Public/UObject/UnrealType.h
@@ -69,6 +69,8 @@
 #include "UObject/WeakObjectPtr.h"
 #include <type_traits>
 
+#define UE_API COREUOBJECT_API
+
 class FBlake3;
 class FOutputDevice;
 class UPackageMap;
@@ -177,7 +179,7 @@ namespace UE::CoreUObject::Private
 //
 class FProperty : public FField
 {
-	DECLARE_FIELD_API(FProperty, FField, CASTCLASS_FProperty, COREUOBJECT_API)
+	DECLARE_FIELD_API(FProperty, FField, CASTCLASS_FProperty, UE_API)
 
 	// Persistent variables.
 	int32			ArrayDim;
@@ -1668,35 +1670,35 @@ public:
 
 };
 
-class COREUOBJECT_API FNumericProperty : public FProperty
+class FNumericProperty : public FProperty
 {
-	DECLARE_FIELD(FNumericProperty, FProperty, CASTCLASS_FNumericProperty)
+	DECLARE_FIELD_API(FNumericProperty, FProperty, CASTCLASS_FNumericProperty, UE_API)
 
-	FNumericProperty(FFieldVariant InOwner, const FName& InName, EObjectFlags InObjectFlags);
+	UE_API FNumericProperty(FFieldVariant InOwner, const FName& InName, EObjectFlags InObjectFlags);
 
 	/**
 	 * Constructor used for constructing compiled in properties
 	 * @param InOwner Owner of the property
 	 * @param PropBase Pointer to the compiled in structure describing the property
 	 **/
-	FNumericProperty(FFieldVariant InOwner, const UECodeGen_Private::FPropertyParamsBaseWithOffset& Prop, EPropertyFlags AdditionalPropertyFlags = CPF_None);
+	UE_API FNumericProperty(FFieldVariant InOwner, const UECodeGen_Private::FPropertyParamsBaseWithOffset& Prop, EPropertyFlags AdditionalPropertyFlags = CPF_None);
 
 #if WITH_EDITORONLY_DATA
-	explicit FNumericProperty(UField* InField);
+	UE_API explicit FNumericProperty(UField* InField);
 #endif // WITH_EDITORONLY_DATA
 
 	// FProperty interface.
-	virtual const TCHAR* ImportText_Internal(const TCHAR* Buffer, void* ContainerOrPropertyPtr, EPropertyPointerType PropertyPointerType, UObject* Parent, int32 PortFlags, FOutputDevice* ErrorText) const override;
-	virtual void ExportText_Internal( FString& ValueStr, const void* PropertyValueOrContainer, EPropertyPointerType PropertyPointerType, const void* DefaultValue, UObject* Parent, int32 PortFlags, UObject* ExportRootScope ) const override;
+	UE_API virtual const TCHAR* ImportText_Internal(const TCHAR* Buffer, void* ContainerOrPropertyPtr, EPropertyPointerType PropertyPointerType, UObject* Parent, int32 PortFlags, FOutputDevice* ErrorText) const override;
+	UE_API virtual void ExportText_Internal( FString& ValueStr, const void* PropertyValueOrContainer, EPropertyPointerType PropertyPointerType, const void* DefaultValue, UObject* Parent, int32 PortFlags, UObject* ExportRootScope ) const override;
 	// End of FProperty interface
 
 	// FNumericProperty interface.
 
 	/** Return true if this property is for a floating point number **/
-	virtual bool IsFloatingPoint() const;
+	UE_API virtual bool IsFloatingPoint() const;
 
 	/** Return true if this property is for a integral or enum type **/
-	virtual bool IsInteger() const;
+	UE_API virtual bool IsInteger() const;
 
 	template <typename T>
 	bool CanHoldValue(T Value) const
@@ -1723,28 +1725,28 @@ class COREUOBJECT_API FNumericProperty : public FProperty
 	}
 
 	/** Return the UEnum if this property is a FByteProperty with a non-null Enum **/
-	virtual UEnum* GetIntPropertyEnum() const;
+	UE_API virtual UEnum* GetIntPropertyEnum() const;
 
 	/** 
 	 * Set the value of an unsigned integral property type
 	 * @param Data - pointer to property data to set
 	 * @param Value - Value to set data to
 	**/
-	virtual void SetIntPropertyValue(void* Data, uint64 Value) const;
+	UE_API virtual void SetIntPropertyValue(void* Data, uint64 Value) const;
 
 	/** 
 	 * Set the value of a signed integral property type
 	 * @param Data - pointer to property data to set
 	 * @param Value - Value to set data to
 	**/
-	virtual void SetIntPropertyValue(void* Data, int64 Value) const;
+	UE_API virtual void SetIntPropertyValue(void* Data, int64 Value) const;
 
 	/** 
 	 * Set the value of a floating point property type
 	 * @param Data - pointer to property data to set
 	 * @param Value - Value to set data to
 	**/
-	virtual void SetFloatingPointPropertyValue(void* Data, double Value) const;
+	UE_API virtual void SetFloatingPointPropertyValue(void* Data, double Value) const;
 
 	/** 
 	 * Set the value of any numeric type from a string point
@@ -1752,31 +1754,31 @@ class COREUOBJECT_API FNumericProperty : public FProperty
 	 * @param Value - Value (as a string) to set 
 	 * CAUTION: This routine does not do enum name conversion
 	**/
-	virtual void SetNumericPropertyValueFromString(void* Data, TCHAR const* Value) const;
-	virtual void SetNumericPropertyValueFromString_InContainer(void* Container, TCHAR const* Value) const;
+	UE_API virtual void SetNumericPropertyValueFromString(void* Data, TCHAR const* Value) const;
+	UE_API virtual void SetNumericPropertyValueFromString_InContainer(void* Container, TCHAR const* Value) const;
 
 	/** 
 	 * Gets the value of a signed integral property type
 	 * @param Data - pointer to property data to get
 	 * @return Data as a signed int
 	**/
-	virtual int64 GetSignedIntPropertyValue(void const* Data) const;
-	virtual int64 GetSignedIntPropertyValue_InContainer(void const* Container) const;
+	UE_API virtual int64 GetSignedIntPropertyValue(void const* Data) const;
+	UE_API virtual int64 GetSignedIntPropertyValue_InContainer(void const* Container) const;
 
 	/** 
 	 * Gets the value of an unsigned integral property type
 	 * @param Data - pointer to property data to get
 	 * @return Data as an unsigned int
 	**/
-	virtual uint64 GetUnsignedIntPropertyValue(void const* Data) const;
-	virtual uint64 GetUnsignedIntPropertyValue_InContainer(void const* Container) const;
+	UE_API virtual uint64 GetUnsignedIntPropertyValue(void const* Data) const;
+	UE_API virtual uint64 GetUnsignedIntPropertyValue_InContainer(void const* Container) const;
 
 	/** 
 	 * Gets the value of an floating point property type
 	 * @param Data - pointer to property data to get
 	 * @return Data as a double
 	**/
-	virtual double GetFloatingPointPropertyValue(void const* Data) const;
+	UE_API virtual double GetFloatingPointPropertyValue(void const* Data) const;
 
 	/** 
 	 * Get the value of any numeric type and return it as a string
@@ -1784,16 +1786,16 @@ class COREUOBJECT_API FNumericProperty : public FProperty
 	 * @return Data as a string
 	 * CAUTION: This routine does not do enum name conversion
 	**/
-	virtual FString GetNumericPropertyValueToString(void const* Data) const;
-	virtual FString GetNumericPropertyValueToString_InContainer(void const* Container) const;
+	UE_API virtual FString GetNumericPropertyValueToString(void const* Data) const;
+	UE_API virtual FString GetNumericPropertyValueToString_InContainer(void const* Container) const;
 	// End of FNumericProperty interface
 
-	static int64 ReadEnumAsInt64(FStructuredArchive::FSlot Slot, UStruct* DefaultsStruct, const FPropertyTag& Tag);
+	UE_API static int64 ReadEnumAsInt64(FStructuredArchive::FSlot Slot, UStruct* DefaultsStruct, const FPropertyTag& Tag);
 
 private:
-	virtual bool CanHoldDoubleValueInternal  (double Value) const PURE_VIRTUAL(FNumericProperty::CanHoldDoubleValueInternal,   return false;);
-	virtual bool CanHoldSignedValueInternal  (int64  Value) const PURE_VIRTUAL(FNumericProperty::CanHoldSignedValueInternal,   return false;);
-	virtual bool CanHoldUnsignedValueInternal(uint64 Value) const PURE_VIRTUAL(FNumericProperty::CanHoldUnsignedValueInternal, return false;);
+	UE_API virtual bool CanHoldDoubleValueInternal  (double Value) const PURE_VIRTUAL(FNumericProperty::CanHoldDoubleValueInternal,   return false;);
+	UE_API virtual bool CanHoldSignedValueInternal  (int64  Value) const PURE_VIRTUAL(FNumericProperty::CanHoldSignedValueInternal,   return false;);
+	UE_API virtual bool CanHoldUnsignedValueInternal(uint64 Value) const PURE_VIRTUAL(FNumericProperty::CanHoldUnsignedValueInternal, return false;);
 };
 
 template<typename InTCppType>
@@ -2055,61 +2057,61 @@ private:
 //
 // Describes an unsigned byte value or 255-value enumeration variable.
 //
-class COREUOBJECT_API FByteProperty : public TProperty_Numeric<uint8>
+class FByteProperty : public TProperty_Numeric<uint8>
 {
-	DECLARE_FIELD(FByteProperty, TProperty_Numeric<uint8>, CASTCLASS_FByteProperty)
+	DECLARE_FIELD_API(FByteProperty, TProperty_Numeric<uint8>, CASTCLASS_FByteProperty, UE_API)
 
 	// Variables.
 	TObjectPtr<UEnum> Enum;
 
-	FByteProperty(FFieldVariant InOwner, const FName& InName, EObjectFlags InObjectFlags);
+	UE_API FByteProperty(FFieldVariant InOwner, const FName& InName, EObjectFlags InObjectFlags);
 
 	/**
 	 * Constructor used for constructing compiled in properties
 	 * @param InOwner Owner of the property
 	 * @param PropBase Pointer to the compiled in structure describing the property
 	 **/
-	FByteProperty(FFieldVariant InOwner, const UECodeGen_Private::FBytePropertyParams& Prop);
+	UE_API FByteProperty(FFieldVariant InOwner, const UECodeGen_Private::FBytePropertyParams& Prop);
 
 #if WITH_EDITORONLY_DATA
-	explicit FByteProperty(UField* InField);
+	UE_API explicit FByteProperty(UField* InField);
 #endif // WITH_EDITORONLY_DATA
 
 	// UObject interface.
-	virtual void Serialize( FArchive& Ar ) override;
-	virtual void AddReferencedObjects(FReferenceCollector& Collector) override;
-	virtual void GetPreloadDependencies(TArray<UObject*>& OutDeps) override;
+	UE_API virtual void Serialize( FArchive& Ar ) override;
+	UE_API virtual void AddReferencedObjects(FReferenceCollector& Collector) override;
+	UE_API virtual void GetPreloadDependencies(TArray<UObject*>& OutDeps) override;
 	// End of UObject interface
 
 	// Field interface
-	virtual void PostDuplicate(const FField& InField) override;
+	UE_API virtual void PostDuplicate(const FField& InField) override;
 
 	// UHT interface
-	virtual FString GetCPPType( FString* ExtendedTypeText=NULL, uint32 CPPExportFlags=0 ) const override;
+	UE_API virtual FString GetCPPType( FString* ExtendedTypeText=NULL, uint32 CPPExportFlags=0 ) const override;
 	// End of UHT interface
 
 	// FProperty interface.
-	virtual void SerializeItem(FStructuredArchive::FSlot Slot, void* Value, void const* Defaults) const override;
-	virtual bool NetSerializeItem( FArchive& Ar, UPackageMap* Map, void* Data, TArray<uint8> * MetaData = NULL ) const override;
+	UE_API virtual void SerializeItem(FStructuredArchive::FSlot Slot, void* Value, void const* Defaults) const override;
+	UE_API virtual bool NetSerializeItem( FArchive& Ar, UPackageMap* Map, void* Data, TArray<uint8> * MetaData = NULL ) const override;
 protected:
-	virtual void ExportText_Internal( FString& ValueStr, const void* PropertyValueOrContainer, EPropertyPointerType PropertyPointerType, const void* DefaultValue, UObject* Parent, int32 PortFlags, UObject* ExportRootScope ) const override;
-	virtual const TCHAR* ImportText_Internal(const TCHAR* Buffer, void* ContainerOrPropertyPtr, EPropertyPointerType PropertyPointerType, UObject* Parent, int32 PortFlags, FOutputDevice* ErrorText) const override;
+	UE_API virtual void ExportText_Internal( FString& ValueStr, const void* PropertyValueOrContainer, EPropertyPointerType PropertyPointerType, const void* DefaultValue, UObject* Parent, int32 PortFlags, UObject* ExportRootScope ) const override;
+	UE_API virtual const TCHAR* ImportText_Internal(const TCHAR* Buffer, void* ContainerOrPropertyPtr, EPropertyPointerType PropertyPointerType, UObject* Parent, int32 PortFlags, FOutputDevice* ErrorText) const override;
 public:
-	virtual EConvertFromTypeResult ConvertFromType(const FPropertyTag& Tag, FStructuredArchive::FSlot Slot, uint8* Data, UStruct* DefaultsStruct, const uint8* Defaults) override;
+	UE_API virtual EConvertFromTypeResult ConvertFromType(const FPropertyTag& Tag, FStructuredArchive::FSlot Slot, uint8* Data, UStruct* DefaultsStruct, const uint8* Defaults) override;
 #if WITH_EDITORONLY_DATA
-	virtual void AppendSchemaHash(FBlake3& Builder, bool bSkipEditorOnly) const override;
+	UE_API virtual void AppendSchemaHash(FBlake3& Builder, bool bSkipEditorOnly) const override;
 #endif
-	virtual bool LoadTypeName(UE::FPropertyTypeName Type, const FPropertyTag* Tag = nullptr) override;
-	virtual void SaveTypeName(UE::FPropertyTypeNameBuilder& Type) const override;
-	virtual bool CanSerializeFromTypeName(UE::FPropertyTypeName Type) const override;
+	UE_API virtual bool LoadTypeName(UE::FPropertyTypeName Type, const FPropertyTag* Tag = nullptr) override;
+	UE_API virtual void SaveTypeName(UE::FPropertyTypeNameBuilder& Type) const override;
+	UE_API virtual bool CanSerializeFromTypeName(UE::FPropertyTypeName Type) const override;
 	// End of FProperty interface
 
 	// FNumericProperty interface.
-	virtual UEnum* GetIntPropertyEnum() const override;
+	UE_API virtual UEnum* GetIntPropertyEnum() const override;
 	// End of FNumericProperty interface
 
 	// Returns the number of bits required by NetSerializeItem to encode this property (may be fewer than 8 if this byte represents an enum)
-	uint64 GetMaxNetSerializeBits() const;
+	UE_API uint64 GetMaxNetSerializeBits() const;
 };
 
 /*-----------------------------------------------------------------------------
@@ -2119,21 +2121,21 @@ public:
 //
 // Describes a 8-bit signed integer variable.
 //
-class COREUOBJECT_API FInt8Property : public TProperty_Numeric<int8>
+class FInt8Property : public TProperty_Numeric<int8>
 {
-	DECLARE_FIELD(FInt8Property, TProperty_Numeric<int8>, CASTCLASS_FInt8Property)
+	DECLARE_FIELD_API(FInt8Property, TProperty_Numeric<int8>, CASTCLASS_FInt8Property, UE_API)
 
-	FInt8Property(FFieldVariant InOwner, const FName& InName, EObjectFlags InObjectFlags);
+	UE_API FInt8Property(FFieldVariant InOwner, const FName& InName, EObjectFlags InObjectFlags);
 
 	/**
 	 * Constructor used for constructing compiled in properties
 	 * @param InOwner Owner of the property
 	 * @param PropBase Pointer to the compiled in structure describing the property
 	 **/
-	FInt8Property(FFieldVariant InOwner, const UECodeGen_Private::FInt8PropertyParams& Prop);
+	UE_API FInt8Property(FFieldVariant InOwner, const UECodeGen_Private::FInt8PropertyParams& Prop);
 
 #if WITH_EDITORONLY_DATA
-	explicit FInt8Property(UField* InField);
+	UE_API explicit FInt8Property(UField* InField);
 #endif // WITH_EDITORONLY_DATA
 };
 
@@ -2144,21 +2146,21 @@ class COREUOBJECT_API FInt8Property : public TProperty_Numeric<int8>
 //
 // Describes a 16-bit signed integer variable.
 //
-class COREUOBJECT_API FInt16Property : public TProperty_Numeric<int16>
+class FInt16Property : public TProperty_Numeric<int16>
 {
-	DECLARE_FIELD(FInt16Property, TProperty_Numeric<int16>, CASTCLASS_FInt16Property)
+	DECLARE_FIELD_API(FInt16Property, TProperty_Numeric<int16>, CASTCLASS_FInt16Property, UE_API)
 
-	FInt16Property(FFieldVariant InOwner, const FName& InName, EObjectFlags InObjectFlags);
+	UE_API FInt16Property(FFieldVariant InOwner, const FName& InName, EObjectFlags InObjectFlags);
 
 	/**
 	 * Constructor used for constructing compiled in properties
 	 * @param InOwner Owner of the property
 	 * @param PropBase Pointer to the compiled in structure describing the property
 	 **/
-	FInt16Property(FFieldVariant InOwner, const UECodeGen_Private::FInt16PropertyParams& Prop);
+	UE_API FInt16Property(FFieldVariant InOwner, const UECodeGen_Private::FInt16PropertyParams& Prop);
 
 #if WITH_EDITORONLY_DATA
-	explicit FInt16Property(UField* InField);
+	UE_API explicit FInt16Property(UField* InField);
 #endif // WITH_EDITORONLY_DATA
 };
 
@@ -2170,21 +2172,21 @@ class COREUOBJECT_API FInt16Property : public TProperty_Numeric<int16>
 //
 // Describes a 32-bit signed integer variable.
 //
-class COREUOBJECT_API FIntProperty : public TProperty_Numeric<int32>
+class FIntProperty : public TProperty_Numeric<int32>
 {
-	DECLARE_FIELD(FIntProperty, TProperty_Numeric<int32>, CASTCLASS_FIntProperty)
+	DECLARE_FIELD_API(FIntProperty, TProperty_Numeric<int32>, CASTCLASS_FIntProperty, UE_API)
 
-	FIntProperty(FFieldVariant InOwner, const FName& InName, EObjectFlags InObjectFlags);
+	UE_API FIntProperty(FFieldVariant InOwner, const FName& InName, EObjectFlags InObjectFlags);
 
 	/**
 	 * Constructor used for constructing compiled in properties
 	 * @param InOwner Owner of the property
 	 * @param PropBase Pointer to the compiled in structure describing the property
 	 **/
-	FIntProperty(FFieldVariant InOwner, const UECodeGen_Private::FIntPropertyParams& Prop);
+	UE_API FIntProperty(FFieldVariant InOwner, const UECodeGen_Private::FIntPropertyParams& Prop);
 
 #if WITH_EDITORONLY_DATA
-	explicit FIntProperty(UField* InField);
+	UE_API explicit FIntProperty(UField* InField);
 #endif // WITH_EDITORONLY_DATA
 };
 
@@ -2195,21 +2197,21 @@ class COREUOBJECT_API FIntProperty : public TProperty_Numeric<int32>
 //
 // Describes a 64-bit signed integer variable.
 //
-class COREUOBJECT_API FInt64Property : public TProperty_Numeric<int64>
+class FInt64Property : public TProperty_Numeric<int64>
 {
-	DECLARE_FIELD(FInt64Property, TProperty_Numeric<int64>, CASTCLASS_FInt64Property)
+	DECLARE_FIELD_API(FInt64Property, TProperty_Numeric<int64>, CASTCLASS_FInt64Property, UE_API)
 
-	FInt64Property(FFieldVariant InOwner, const FName& InName, EObjectFlags InObjectFlags);
+	UE_API FInt64Property(FFieldVariant InOwner, const FName& InName, EObjectFlags InObjectFlags);
 
 	/**
 	 * Constructor used for constructing compiled in properties
 	 * @param InOwner Owner of the property
 	 * @param PropBase Pointer to the compiled in structure describing the property
 	 **/
-	FInt64Property(FFieldVariant InOwner, const UECodeGen_Private::FInt64PropertyParams& Prop);
+	UE_API FInt64Property(FFieldVariant InOwner, const UECodeGen_Private::FInt64PropertyParams& Prop);
 
 #if WITH_EDITORONLY_DATA
-	explicit FInt64Property(UField* InField);
+	UE_API explicit FInt64Property(UField* InField);
 #endif // WITH_EDITORONLY_DATA
 };
 
@@ -2220,21 +2222,21 @@ class COREUOBJECT_API FInt64Property : public TProperty_Numeric<int64>
 //
 // Describes a 16-bit unsigned integer variable.
 //
-class COREUOBJECT_API FUInt16Property : public TProperty_Numeric<uint16>
+class FUInt16Property : public TProperty_Numeric<uint16>
 {
-	DECLARE_FIELD(FUInt16Property, TProperty_Numeric<uint16>, CASTCLASS_FUInt16Property)
+	DECLARE_FIELD_API(FUInt16Property, TProperty_Numeric<uint16>, CASTCLASS_FUInt16Property, UE_API)
 
-	FUInt16Property(FFieldVariant InOwner, const FName& InName, EObjectFlags InObjectFlags);
+	UE_API FUInt16Property(FFieldVariant InOwner, const FName& InName, EObjectFlags InObjectFlags);
 
 	/**
 	 * Constructor used for constructing compiled in properties
 	 * @param InOwner Owner of the property
 	 * @param PropBase Pointer to the compiled in structure describing the property
 	 **/
-	FUInt16Property(FFieldVariant InOwner, const UECodeGen_Private::FUInt16PropertyParams& Prop);
+	UE_API FUInt16Property(FFieldVariant InOwner, const UECodeGen_Private::FUInt16PropertyParams& Prop);
 
 #if WITH_EDITORONLY_DATA
-	explicit FUInt16Property(UField* InField);
+	UE_API explicit FUInt16Property(UField* InField);
 #endif // WITH_EDITORONLY_DATA
 };
 
@@ -2245,21 +2247,21 @@ class COREUOBJECT_API FUInt16Property : public TProperty_Numeric<uint16>
 //
 // Describes a 32-bit unsigned integer variable.
 //
-class COREUOBJECT_API FUInt32Property : public TProperty_Numeric<uint32>
+class FUInt32Property : public TProperty_Numeric<uint32>
 {
-	DECLARE_FIELD(FUInt32Property, TProperty_Numeric<uint32>, CASTCLASS_FUInt32Property)
+	DECLARE_FIELD_API(FUInt32Property, TProperty_Numeric<uint32>, CASTCLASS_FUInt32Property, UE_API)
 
-	FUInt32Property(FFieldVariant InOwner, const FName& InName, EObjectFlags InObjectFlags);
+	UE_API FUInt32Property(FFieldVariant InOwner, const FName& InName, EObjectFlags InObjectFlags);
 
 	/**
 	 * Constructor used for constructing compiled in properties
 	 * @param InOwner Owner of the property
 	 * @param PropBase Pointer to the compiled in structure describing the property
 	 **/
-	FUInt32Property(FFieldVariant InOwner, const UECodeGen_Private::FUInt32PropertyParams& Prop);
+	UE_API FUInt32Property(FFieldVariant InOwner, const UECodeGen_Private::FUInt32PropertyParams& Prop);
 
 #if WITH_EDITORONLY_DATA
-	explicit FUInt32Property(UField* InField);
+	UE_API explicit FUInt32Property(UField* InField);
 #endif // WITH_EDITORONLY_DATA
 };
 
@@ -2270,21 +2272,21 @@ class COREUOBJECT_API FUInt32Property : public TProperty_Numeric<uint32>
 //
 // Describes a 64-bit unsigned integer variable.
 //
-class COREUOBJECT_API FUInt64Property : public TProperty_Numeric<uint64>
+class FUInt64Property : public TProperty_Numeric<uint64>
 {
-	DECLARE_FIELD(FUInt64Property, TProperty_Numeric<uint64>, CASTCLASS_FUInt64Property)
+	DECLARE_FIELD_API(FUInt64Property, TProperty_Numeric<uint64>, CASTCLASS_FUInt64Property, UE_API)
 
-	FUInt64Property(FFieldVariant InOwner, const FName& InName, EObjectFlags InObjectFlags);
+	UE_API FUInt64Property(FFieldVariant InOwner, const FName& InName, EObjectFlags InObjectFlags);
 
 	/**
 	 * Constructor used for constructing compiled in properties
 	 * @param InOwner Owner of the property
 	 * @param PropBase Pointer to the compiled in structure describing the property
 	 **/
-	FUInt64Property(FFieldVariant InOwner, const UECodeGen_Private::FUInt64PropertyParams& Prop);
+	UE_API FUInt64Property(FFieldVariant InOwner, const UECodeGen_Private::FUInt64PropertyParams& Prop);
 
 #if WITH_EDITORONLY_DATA
-	explicit FUInt64Property(UField* InField);
+	UE_API explicit FUInt64Property(UField* InField);
 #endif // WITH_EDITORONLY_DATA
 };
 
@@ -2296,28 +2298,28 @@ class COREUOBJECT_API FUInt64Property : public TProperty_Numeric<uint64>
 //
 // Describes an IEEE 32-bit floating point variable.
 //
-class COREUOBJECT_API FFloatProperty : public TProperty_Numeric<float>
+class FFloatProperty : public TProperty_Numeric<float>
 {
-	DECLARE_FIELD(FFloatProperty, TProperty_Numeric<float>, CASTCLASS_FFloatProperty)
+	DECLARE_FIELD_API(FFloatProperty, TProperty_Numeric<float>, CASTCLASS_FFloatProperty, UE_API)
 
-	FFloatProperty(FFieldVariant InOwner, const FName& InName, EObjectFlags InObjectFlags);
+	UE_API FFloatProperty(FFieldVariant InOwner, const FName& InName, EObjectFlags InObjectFlags);
 
 	/**
 	 * Constructor used for constructing compiled in properties
 	 * @param InOwner Owner of the property
 	 * @param PropBase Pointer to the compiled in structure describing the property
 	 **/
-	FFloatProperty(FFieldVariant InOwner, const UECodeGen_Private::FFloatPropertyParams& Prop);
+	UE_API FFloatProperty(FFieldVariant InOwner, const UECodeGen_Private::FFloatPropertyParams& Prop);
 
 #if WITH_EDITORONLY_DATA
-	explicit FFloatProperty(UField* InField);
+	UE_API explicit FFloatProperty(UField* InField);
 #endif // WITH_EDITORONLY_DATA
 
-	virtual bool Identical(const void* A, const void* B, uint32 PortFlags) const override;
+	UE_API virtual bool Identical(const void* A, const void* B, uint32 PortFlags) const override;
 
 protected:
 
-	virtual uint32 GetValueTypeHashInternal(const void* Src) const override;
+	UE_API virtual uint32 GetValueTypeHashInternal(const void* Src) const override;
 };
 
 /*-----------------------------------------------------------------------------
@@ -2327,34 +2329,34 @@ protected:
 //
 // Describes an IEEE 64-bit floating point variable.
 //
-class COREUOBJECT_API FDoubleProperty : public TProperty_Numeric<double>
+class FDoubleProperty : public TProperty_Numeric<double>
 {
-	DECLARE_FIELD(FDoubleProperty, TProperty_Numeric<double>, CASTCLASS_FDoubleProperty)
+	DECLARE_FIELD_API(FDoubleProperty, TProperty_Numeric<double>, CASTCLASS_FDoubleProperty, UE_API)
 
-	FDoubleProperty(FFieldVariant InOwner, const FName& InName, EObjectFlags InObjectFlags);
+	UE_API FDoubleProperty(FFieldVariant InOwner, const FName& InName, EObjectFlags InObjectFlags);
 
 	/**
 	 * Constructor used for constructing compiled in properties
 	 * @param InOwner Owner of the property
 	 * @param PropBase Pointer to the compiled in structure describing the property
 	 **/
-	FDoubleProperty(FFieldVariant InOwner, const UECodeGen_Private::FDoublePropertyParams& Prop);
+	UE_API FDoubleProperty(FFieldVariant InOwner, const UECodeGen_Private::FDoublePropertyParams& Prop);
 
 #if WITH_EDITORONLY_DATA
-	explicit FDoubleProperty(UField* InField);
+	UE_API explicit FDoubleProperty(UField* InField);
 #endif // WITH_EDITORONLY_DATA
 
-	virtual bool Identical(const void* A, const void* B, uint32 PortFlags) const override;
+	UE_API virtual bool Identical(const void* A, const void* B, uint32 PortFlags) const override;
 
 protected:
 
-	virtual uint32 GetValueTypeHashInternal(const void* Src) const override;
+	UE_API virtual uint32 GetValueTypeHashInternal(const void* Src) const override;
 };
 
 // Note: Stub only. Used to provide FNativeClassHeaderGenerator::PropertyNew with a way to track an as yet unaliased FFloatProperty/FDoubleProperty within the header tool, which will be resolved to the correct type at run time.
-class COREUOBJECT_API FLargeWorldCoordinatesRealProperty : public TProperty_Numeric<double>
+class FLargeWorldCoordinatesRealProperty : public TProperty_Numeric<double>
 {
-	DECLARE_FIELD(FLargeWorldCoordinatesRealProperty, TProperty_Numeric<double>, CASTCLASS_FLargeWorldCoordinatesRealProperty)
+	DECLARE_FIELD_API(FLargeWorldCoordinatesRealProperty, TProperty_Numeric<double>, CASTCLASS_FLargeWorldCoordinatesRealProperty, UE_API)
 
 	FLargeWorldCoordinatesRealProperty(FFieldVariant InOwner, const FName& InName, EObjectFlags InObjectFlags)
 		: TProperty_Numeric(InOwner, InName, InObjectFlags)
@@ -2369,9 +2371,9 @@ class COREUOBJECT_API FLargeWorldCoordinatesRealProperty : public TProperty_Nume
 //
 // Describes a single bit flag variable residing in a 32-bit unsigned double word.
 //
-class COREUOBJECT_API FBoolProperty : public FProperty
+class FBoolProperty : public FProperty
 {
-	DECLARE_FIELD(FBoolProperty, FProperty, CASTCLASS_FBoolProperty)
+	DECLARE_FIELD_API(FBoolProperty, FProperty, CASTCLASS_FBoolProperty, UE_API)
 
 	// Variables.
 private:
@@ -2387,47 +2389,47 @@ private:
 
 public:
 
-	FBoolProperty(FFieldVariant InOwner, const FName& InName, EObjectFlags InObjectFlags);
+	UE_API FBoolProperty(FFieldVariant InOwner, const FName& InName, EObjectFlags InObjectFlags);
 
 	/**
 	 * Constructor used for constructing compiled in properties
 	 * @param InOwner Owner of the property
 	 * @param PropBase Pointer to the compiled in structure describing the property
 	 **/
-	FBoolProperty(FFieldVariant InOwner, const UECodeGen_Private::FBoolPropertyParams& Prop);
+	UE_API FBoolProperty(FFieldVariant InOwner, const UECodeGen_Private::FBoolPropertyParams& Prop);
 
 #if WITH_EDITORONLY_DATA
-	explicit FBoolProperty(UField* InField);
+	UE_API explicit FBoolProperty(UField* InField);
 #endif // WITH_EDITORONLY_DATA
 
 	// UObject interface.
-	virtual void Serialize( FArchive& Ar ) override;
+	UE_API virtual void Serialize( FArchive& Ar ) override;
 	// End of UObject interface
 
 	// Field interface
-	virtual void PostDuplicate(const FField& InField) override;
+	UE_API virtual void PostDuplicate(const FField& InField) override;
 
 	// UHT interface
-	virtual FString GetCPPType( FString* ExtendedTypeText, uint32 CPPExportFlags ) const override;
-	virtual FString GetCPPMacroType( FString& ExtendedTypeText ) const override;
+	UE_API virtual FString GetCPPType( FString* ExtendedTypeText, uint32 CPPExportFlags ) const override;
+	UE_API virtual FString GetCPPMacroType( FString& ExtendedTypeText ) const override;
 	// End of UHT interface
 
 	// FProperty interface.
-	virtual void LinkInternal(FArchive& Ar) override;
-	virtual bool Identical( const void* A, const void* B, uint32 PortFlags ) const override;
-	virtual void SerializeItem( FStructuredArchive::FSlot Slot, void* Value, void const* Defaults) const override;
-	virtual bool NetSerializeItem( FArchive& Ar, UPackageMap* Map, void* Data, TArray<uint8> * MetaData = NULL ) const override;
+	UE_API virtual void LinkInternal(FArchive& Ar) override;
+	UE_API virtual bool Identical( const void* A, const void* B, uint32 PortFlags ) const override;
+	UE_API virtual void SerializeItem( FStructuredArchive::FSlot Slot, void* Value, void const* Defaults) const override;
+	UE_API virtual bool NetSerializeItem( FArchive& Ar, UPackageMap* Map, void* Data, TArray<uint8> * MetaData = NULL ) const override;
 protected:
-	virtual void ExportText_Internal( FString& ValueStr, const void* PropertyValueOrContainer, EPropertyPointerType PropertyPointerType, const void* DefaultValue, UObject* Parent, int32 PortFlags, UObject* ExportRootScope ) const override;
-	virtual const TCHAR* ImportText_Internal(const TCHAR* Buffer, void* ContainerOrPropertyPtr, EPropertyPointerType PropertyPointerType, UObject* Parent, int32 PortFlags, FOutputDevice* ErrorText) const override;
+	UE_API virtual void ExportText_Internal( FString& ValueStr, const void* PropertyValueOrContainer, EPropertyPointerType PropertyPointerType, const void* DefaultValue, UObject* Parent, int32 PortFlags, UObject* ExportRootScope ) const override;
+	UE_API virtual const TCHAR* ImportText_Internal(const TCHAR* Buffer, void* ContainerOrPropertyPtr, EPropertyPointerType PropertyPointerType, UObject* Parent, int32 PortFlags, FOutputDevice* ErrorText) const override;
 public:
-	virtual void CopyValuesInternal( void* Dest, void const* Src, int32 Count ) const override;
-	virtual void ClearValueInternal( void* Data ) const override;
-	virtual void InitializeValueInternal( void* Dest ) const override;
-	virtual int32 GetMinAlignment() const override;
-	virtual EConvertFromTypeResult ConvertFromType(const FPropertyTag& Tag, FStructuredArchive::FSlot Slot, uint8* Data, UStruct* DefaultsStruct, const uint8* Defaults) override;
+	UE_API virtual void CopyValuesInternal( void* Dest, void const* Src, int32 Count ) const override;
+	UE_API virtual void ClearValueInternal( void* Data ) const override;
+	UE_API virtual void InitializeValueInternal( void* Dest ) const override;
+	UE_API virtual int32 GetMinAlignment() const override;
+	UE_API virtual EConvertFromTypeResult ConvertFromType(const FPropertyTag& Tag, FStructuredArchive::FSlot Slot, uint8* Data, UStruct* DefaultsStruct, const uint8* Defaults) override;
 #if WITH_EDITORONLY_DATA
-	virtual void AppendSchemaHash(FBlake3& Builder, bool bSkipEditorOnly) const override;
+	UE_API virtual void AppendSchemaHash(FBlake3& Builder, bool bSkipEditorOnly) const override;
 #endif
 	// End of FProperty interface
 
@@ -2475,7 +2477,7 @@ public:
 	 * @param InSize size of the bitfield/bool type.
 	 * @param bIsNativeBool true if this property represents C++ bool type.
 	 */
-	void SetBoolSize( const uint32 InSize, const bool bIsNativeBool = false, const uint32 InBitMask = 0 );
+	UE_API void SetBoolSize( const uint32 InSize, const bool bIsNativeBool = false, const uint32 InBitMask = 0 );
 
 	/**
 	 * If the return value is true this FBoolProperty represents C++ bool type.
@@ -2497,7 +2499,7 @@ public:
 		return ByteOffset;
 	}
 
-	uint32 GetValueTypeHashInternal(const void* Src) const override;
+	UE_API uint32 GetValueTypeHashInternal(const void* Src) const override;
 
 	virtual bool HasIntrusiveUnsetOptionalState() const override 
 	{ 
@@ -2512,57 +2514,57 @@ public:
 //
 // Describes a reference variable to another object which may be nil.
 //
-class COREUOBJECT_API FObjectPropertyBase : public FProperty
+class FObjectPropertyBase : public FProperty
 {
-	DECLARE_FIELD(FObjectPropertyBase, FProperty, CASTCLASS_FObjectPropertyBase)
+	DECLARE_FIELD_API(FObjectPropertyBase, FProperty, CASTCLASS_FObjectPropertyBase, UE_API)
 
 public:
 
 	// Variables.
 	TObjectPtr<class UClass> PropertyClass;
 
-	FObjectPropertyBase(FFieldVariant InOwner, const FName& InName, EObjectFlags InObjectFlags);
+	UE_API FObjectPropertyBase(FFieldVariant InOwner, const FName& InName, EObjectFlags InObjectFlags);
 
 	/**
 	 * Constructor used for constructing compiled-in properties
 	 * @param InOwner Owner of the property
 	 * @param PropBase Pointer to the compiled in structure describing the property
 	 **/
-	FObjectPropertyBase(FFieldVariant InOwner, const UECodeGen_Private::FObjectPropertyParams& Prop, EPropertyFlags AdditionalPropertyFlags = CPF_None);
-	FObjectPropertyBase(FFieldVariant InOwner, const UECodeGen_Private::FObjectPropertyParamsWithoutClass& Prop, EPropertyFlags AdditionalPropertyFlags = CPF_None);
+	UE_API FObjectPropertyBase(FFieldVariant InOwner, const UECodeGen_Private::FObjectPropertyParams& Prop, EPropertyFlags AdditionalPropertyFlags = CPF_None);
+	UE_API FObjectPropertyBase(FFieldVariant InOwner, const UECodeGen_Private::FObjectPropertyParamsWithoutClass& Prop, EPropertyFlags AdditionalPropertyFlags = CPF_None);
 
 #if WITH_EDITORONLY_DATA
-	explicit FObjectPropertyBase(UField* InField);
+	UE_API explicit FObjectPropertyBase(UField* InField);
 #endif // WITH_EDITORONLY_DATA
 
 	// UObject interface
-	virtual void Serialize( FArchive& Ar ) override;
-	virtual void AddReferencedObjects(FReferenceCollector& Collector) override;
-	virtual void BeginDestroy() override;
+	UE_API virtual void Serialize( FArchive& Ar ) override;
+	UE_API virtual void AddReferencedObjects(FReferenceCollector& Collector) override;
+	UE_API virtual void BeginDestroy() override;
 	// End of UObject interface
 
 	// Field interface
-	virtual void PostDuplicate(const FField& InField) override;
+	UE_API virtual void PostDuplicate(const FField& InField) override;
 
 	// FProperty interface
-	virtual bool Identical( const void* A, const void* B, uint32 PortFlags ) const override;
-	virtual bool NetSerializeItem( FArchive& Ar, UPackageMap* Map, void* Data, TArray<uint8> * MetaData = NULL ) const override;
-	virtual bool SupportsNetSharedSerialization() const override { return false; }
+	UE_API virtual bool Identical( const void* A, const void* B, uint32 PortFlags ) const override;
+	UE_API virtual bool NetSerializeItem( FArchive& Ar, UPackageMap* Map, void* Data, TArray<uint8> * MetaData = NULL ) const override;
+	UE_API virtual bool SupportsNetSharedSerialization() const override { return false; }
 protected:
-	virtual void ExportText_Internal( FString& ValueStr, const void* PropertyValueOrContainer, EPropertyPointerType PropertyPointerType, const void* DefaultValue, UObject* Parent, int32 PortFlags, UObject* ExportRootScope ) const override;
-	virtual const TCHAR* ImportText_Internal(const TCHAR* Buffer, void* ContainerOrPropertyPtr, EPropertyPointerType PropertyPointerType, UObject* OwnerObject, int32 PortFlags, FOutputDevice* ErrorText) const override;
+	UE_API virtual void ExportText_Internal( FString& ValueStr, const void* PropertyValueOrContainer, EPropertyPointerType PropertyPointerType, const void* DefaultValue, UObject* Parent, int32 PortFlags, UObject* ExportRootScope ) const override;
+	UE_API virtual const TCHAR* ImportText_Internal(const TCHAR* Buffer, void* ContainerOrPropertyPtr, EPropertyPointerType PropertyPointerType, UObject* OwnerObject, int32 PortFlags, FOutputDevice* ErrorText) const override;
 public:
-	virtual FName GetID() const override;
-	virtual void InstanceSubobjects( void* Data, void const* DefaultData, TNotNull<UObject*> Owner, struct FObjectInstancingGraph* InstanceGraph ) override;
-	virtual bool SameType(const FProperty* Other) const override;
-	virtual EPropertyVisitorControlFlow Visit(FPropertyVisitorContext& Context, const TFunctionRef<EPropertyVisitorControlFlow(const FPropertyVisitorContext& /*Conmtext*/)> InFunc) const override;
-	virtual void* ResolveVisitedPathInfo(void* Data, const FPropertyVisitorInfo& Info) const override;
+	UE_API virtual FName GetID() const override;
+	UE_API virtual void InstanceSubobjects( void* Data, void const* DefaultData, TNotNull<UObject*> Owner, struct FObjectInstancingGraph* InstanceGraph ) override;
+	UE_API virtual bool SameType(const FProperty* Other) const override;
+	UE_API virtual EPropertyVisitorControlFlow Visit(FPropertyVisitorContext& Context, const TFunctionRef<EPropertyVisitorControlFlow(const FPropertyVisitorContext& /*Conmtext*/)> InFunc) const override;
+	UE_API virtual void* ResolveVisitedPathInfo(void* Data, const FPropertyVisitorInfo& Info) const override;
 	// End of FProperty interface
 
 	// FObjectPropertyBase interface
 public:
 
-	virtual FString GetCPPTypeCustom(FString* ExtendedTypeText, uint32 CPPExportFlags, const FString& InnerNativeTypeName) const PURE_VIRTUAL(FObjectPropertyBase::GetCPPTypeCustom, return TEXT(""););
+	UE_API virtual FString GetCPPTypeCustom(FString* ExtendedTypeText, uint32 CPPExportFlags, const FString& InnerNativeTypeName) const PURE_VIRTUAL(FObjectPropertyBase::GetCPPTypeCustom, return TEXT(""););
 
 	/**
 	 * Parses a text buffer into an object reference.
@@ -2579,8 +2581,8 @@ public:
 	 * @param bAllowAnyPackage		allows ignoring package name to find any object that happens to be loaded with the same name
 	 * @return	true if the text is successfully resolved into a valid object reference of the correct type, false otherwise.
 	 */
-	static bool ParseObjectPropertyValue( const FProperty* Property, UObject* OwnerObject, UClass* RequiredMetaClass, uint32 PortFlags, const TCHAR*& Buffer, TObjectPtr<UObject>& out_ResolvedValue, FUObjectSerializeContext* InSerializeContext = nullptr, bool bAllowAnyPackage = true );
-	static TObjectPtr<UObject> FindImportedObject( const FProperty* Property, UObject* OwnerObject, UClass* ObjectClass, UClass* RequiredMetaClass, const TCHAR* Text, uint32 PortFlags = 0, FUObjectSerializeContext* InSerializeContext = nullptr, bool bAllowAnyPackage = true );
+	UE_API static bool ParseObjectPropertyValue( const FProperty* Property, UObject* OwnerObject, UClass* RequiredMetaClass, uint32 PortFlags, const TCHAR*& Buffer, TObjectPtr<UObject>& out_ResolvedValue, FUObjectSerializeContext* InSerializeContext = nullptr, bool bAllowAnyPackage = true );
+	UE_API static TObjectPtr<UObject> FindImportedObject( const FProperty* Property, UObject* OwnerObject, UClass* ObjectClass, UClass* RequiredMetaClass, const TCHAR* Text, uint32 PortFlags = 0, FUObjectSerializeContext* InSerializeContext = nullptr, bool bAllowAnyPackage = true );
 	
 	/**
 	 * Returns the qualified export path for a given object, parent, and export root scope
@@ -2590,7 +2592,7 @@ public:
 	 * @param PortFlags Property port flags
 	 * @return A string representing the export path of an object, usually in the form of /ClassPackage.ClassName'/Package/Path.Object'
 	 */
-	static FString GetExportPath(const TObjectPtr<const UObject>& Object, const UObject* Parent = nullptr, const UObject* ExportRootScope = nullptr, const uint32 PortFlags = PPF_None);
+	UE_API static FString GetExportPath(const TObjectPtr<const UObject>& Object, const UObject* Parent = nullptr, const UObject* ExportRootScope = nullptr, const uint32 PortFlags = PPF_None);
 
 	/**
 	 * Returns the qualified export path given a class path name and object path name
@@ -2598,10 +2600,10 @@ public:
 	 * @param ObjectPathName of the Object used as a root object for generating the Objects path name
 	 * @return A string representing the export path of an object in the form of /ClassPackage.ClassName'/Package/Path.Object'
 	 */
-	static FString GetExportPath(FTopLevelAssetPath ClassPathName, const FString& ObjectPathName);
+	UE_API static FString GetExportPath(FTopLevelAssetPath ClassPathName, const FString& ObjectPathName);
 
 	// Helper method for sharing code with FObjectPtrProperty even though one doesn't inherit from the other
-	static bool StaticIdentical(UObject* A, UObject* B, uint32 PortFlags);
+	UE_API static bool StaticIdentical(UObject* A, UObject* B, uint32 PortFlags);
 
 	virtual UObject* LoadObjectPropertyValue(const void* PropertyValueAddress) const
 	{
@@ -2613,20 +2615,20 @@ public:
 	}
 
 protected:
-	virtual void SetObjectPropertyValueUnchecked(void* PropertyValueAddress, UObject* Value) const;
-	virtual void SetObjectPtrPropertyValueUnchecked(void* PropertyValueAddress, TObjectPtr<UObject> Ptr) const;
-	virtual void SetObjectPropertyValueUnchecked_InContainer(void* ContainerAddress, UObject* Value, int32 ArrayIndex = 0) const;
-	virtual void SetObjectPtrPropertyValueUnchecked_InContainer(void* ContainerAddress, TObjectPtr<UObject> Ptr, int32 ArrayIndex = 0) const;
+	UE_API virtual void SetObjectPropertyValueUnchecked(void* PropertyValueAddress, UObject* Value) const;
+	UE_API virtual void SetObjectPtrPropertyValueUnchecked(void* PropertyValueAddress, TObjectPtr<UObject> Ptr) const;
+	UE_API virtual void SetObjectPropertyValueUnchecked_InContainer(void* ContainerAddress, UObject* Value, int32 ArrayIndex = 0) const;
+	UE_API virtual void SetObjectPtrPropertyValueUnchecked_InContainer(void* ContainerAddress, TObjectPtr<UObject> Ptr, int32 ArrayIndex = 0) const;
 public:
-	virtual UObject* GetObjectPropertyValue(const void* PropertyValueAddress) const;
-	virtual TObjectPtr<UObject> GetObjectPtrPropertyValue(const void* PropertyValueAddress) const;
-	virtual UObject* GetObjectPropertyValue_InContainer(const void* ContainerAddress, int32 ArrayIndex = 0) const;
-	virtual TObjectPtr<UObject> GetObjectPtrPropertyValue_InContainer(const void* ContainerAddress, int32 ArrayIndex = 0) const;
+	UE_API virtual UObject* GetObjectPropertyValue(const void* PropertyValueAddress) const;
+	UE_API virtual TObjectPtr<UObject> GetObjectPtrPropertyValue(const void* PropertyValueAddress) const;
+	UE_API virtual UObject* GetObjectPropertyValue_InContainer(const void* ContainerAddress, int32 ArrayIndex = 0) const;
+	UE_API virtual TObjectPtr<UObject> GetObjectPtrPropertyValue_InContainer(const void* ContainerAddress, int32 ArrayIndex = 0) const;
 
-	void SetObjectPropertyValue(void* PropertyValueAddress, UObject* Value) const;
-	void SetObjectPtrPropertyValue(void* PropertyValueAddress, TObjectPtr<UObject> Ptr) const;
-	void SetObjectPropertyValue_InContainer(void* ContainerAddress, UObject* Value, int32 ArrayIndex = 0) const;
-	void SetObjectPtrPropertyValue_InContainer(void* ContainerAddress, TObjectPtr<UObject> Ptr, int32 ArrayIndex = 0) const;
+	UE_API void SetObjectPropertyValue(void* PropertyValueAddress, UObject* Value) const;
+	UE_API void SetObjectPtrPropertyValue(void* PropertyValueAddress, TObjectPtr<UObject> Ptr) const;
+	UE_API void SetObjectPropertyValue_InContainer(void* ContainerAddress, UObject* Value, int32 ArrayIndex = 0) const;
+	UE_API void SetObjectPtrPropertyValue_InContainer(void* ContainerAddress, TObjectPtr<UObject> Ptr, int32 ArrayIndex = 0) const;
 
 	/**
 	 * Setter function for this property's PropertyClass member. Favor this 
@@ -2638,16 +2640,16 @@ public:
 	 * @param  NewPropertyClass    The PropertyClass you want this property set with.
 	 */
 #if USE_CIRCULAR_DEPENDENCY_LOAD_DEFERRING
-	void SetPropertyClass(UClass* NewPropertyClass);
+	UE_API void SetPropertyClass(UClass* NewPropertyClass);
 #else
 	FORCEINLINE void SetPropertyClass(UClass* NewPropertyClass) { PropertyClass = NewPropertyClass; }
 #endif // USE_CIRCULAR_DEPENDENCY_LOAD_DEFERRING
 
-	virtual void CheckValidObject(void* ValueAddress, TObjectPtr<UObject> OldValue, const void* Defaults = nullptr) const;
-	virtual bool AllowObjectTypeReinterpretationTo(const FObjectPropertyBase* Other) const;
+	UE_API virtual void CheckValidObject(void* ValueAddress, TObjectPtr<UObject> OldValue, const void* Defaults = nullptr) const;
+	UE_API virtual bool AllowObjectTypeReinterpretationTo(const FObjectPropertyBase* Other) const;
 
 protected:
-	virtual bool AllowCrossLevel() const;
+	UE_API virtual bool AllowCrossLevel() const;
 	// End of FObjectPropertyBase interface
 
 	/**
@@ -2658,7 +2660,7 @@ protected:
 	 * @param Defaults Defaults object to use when constructing the object value (can be null).
 	 * @return Pointer that should be assigned to the property value, or null if a default value could not be created.
 	 */
-	UObject* ConstructDefaultObjectValueIfNecessary(UObject* ExistingValue, FString* OutFailureReason = nullptr, const void* Defaults = nullptr) const;
+	UE_API UObject* ConstructDefaultObjectValueIfNecessary(UObject* ExistingValue, FString* OutFailureReason = nullptr, const void* Defaults = nullptr) const;
 
 	// Disable false positive buffer overrun warning during pgoprofile linking step
 	PRAGMA_DISABLE_BUFFER_OVERRUN_WARNING
@@ -2781,7 +2783,7 @@ protected:
 };
 
 template<typename InTCppType>
-class COREUOBJECT_API TFObjectPropertyBase : public TProperty<InTCppType, FObjectPropertyBase>
+class TFObjectPropertyBase : public TProperty<InTCppType, FObjectPropertyBase>
 {
 public:
 	typedef TProperty<InTCppType, FObjectPropertyBase> Super;
@@ -2857,61 +2859,61 @@ ENUM_CLASS_FLAGS(EObjectPropertyOptions);
 //
 // Describes a reference variable to another object which may be nil.
 //
-class COREUOBJECT_API FObjectProperty : public TFObjectPropertyBase<TObjectPtr<UObject>>
+class FObjectProperty : public TFObjectPropertyBase<TObjectPtr<UObject>>
 {
-	DECLARE_FIELD(FObjectProperty, TFObjectPropertyBase<TObjectPtr<UObject>>, CASTCLASS_FObjectProperty)
+	DECLARE_FIELD_API(FObjectProperty, TFObjectPropertyBase<TObjectPtr<UObject>>, CASTCLASS_FObjectProperty, UE_API)
 
-	FObjectProperty(FFieldVariant InOwner, const FName& InName, EObjectFlags InObjectFlags);
+	UE_API FObjectProperty(FFieldVariant InOwner, const FName& InName, EObjectFlags InObjectFlags);
 
 	/**
 	 * Constructor used for constructing compiled in properties
 	 * @param InOwner Owner of the property
 	 * @param PropBase Pointer to the compiled in structure describing the property
 	 **/
-	FObjectProperty(FFieldVariant InOwner, const UECodeGen_Private::FObjectPropertyParams& Prop);
+	UE_API FObjectProperty(FFieldVariant InOwner, const UECodeGen_Private::FObjectPropertyParams& Prop);
 
 #if WITH_EDITORONLY_DATA
-	explicit FObjectProperty(UField* InField);
+	UE_API explicit FObjectProperty(UField* InField);
 #endif // WITH_EDITORONLY_DATA
 
 	// UHT interface
-	virtual FString GetCPPMacroType( FString& ExtendedTypeText ) const  override;
+	UE_API virtual FString GetCPPMacroType( FString& ExtendedTypeText ) const  override;
 	// End of UHT interface
 
 	// FProperty interface
-	virtual void SerializeItem(FStructuredArchive::FSlot Slot, void* Value, void const* Defaults) const override;
+	UE_API virtual void SerializeItem(FStructuredArchive::FSlot Slot, void* Value, void const* Defaults) const override;
 #if UE_WITH_REMOTE_OBJECT_HANDLE
-	virtual bool NetSerializeItem(FArchive& Ar, UPackageMap* Map, void* Data, TArray<uint8>* MetaData = nullptr) const override;
+	UE_API virtual bool NetSerializeItem(FArchive& Ar, UPackageMap* Map, void* Data, TArray<uint8>* MetaData = nullptr) const override;
 #endif // UE_WITH_REMOTE_OBJECT_HANDLE
-	virtual void EmitReferenceInfo(UE::GC::FSchemaBuilder& Schema, int32 BaseOffset, TArray<const FStructProperty*>& EncounteredStructProps, UE::GC::FPropertyStack& DebugPath) override;
-	virtual const TCHAR* ImportText_Internal(const TCHAR* Buffer, void* ContainerOrPropertyPtr, EPropertyPointerType PropertyPointerType, UObject* OwnerObject, int32 PortFlags, FOutputDevice* ErrorText) const override;
-	virtual EConvertFromTypeResult ConvertFromType(const FPropertyTag& Tag, FStructuredArchive::FSlot Slot, uint8* Data, UStruct* DefaultsStruct, const uint8* Defaults) override;
-	virtual bool AllowCrossLevel() const override;
+	UE_API virtual void EmitReferenceInfo(UE::GC::FSchemaBuilder& Schema, int32 BaseOffset, TArray<const FStructProperty*>& EncounteredStructProps, UE::GC::FPropertyStack& DebugPath) override;
+	UE_API virtual const TCHAR* ImportText_Internal(const TCHAR* Buffer, void* ContainerOrPropertyPtr, EPropertyPointerType PropertyPointerType, UObject* OwnerObject, int32 PortFlags, FOutputDevice* ErrorText) const override;
+	UE_API virtual EConvertFromTypeResult ConvertFromType(const FPropertyTag& Tag, FStructuredArchive::FSlot Slot, uint8* Data, UStruct* DefaultsStruct, const uint8* Defaults) override;
+	UE_API virtual bool AllowCrossLevel() const override;
 private:
-	virtual uint32 GetValueTypeHashInternal(const void* Src) const override;
-	virtual void CopyValuesInternal(void* Dest, void const* Src, int32 Count) const override;
+	UE_API virtual uint32 GetValueTypeHashInternal(const void* Src) const override;
+	UE_API virtual void CopyValuesInternal(void* Dest, void const* Src, int32 Count) const override;
 public:
-	virtual void CopySingleValueToScriptVM( void* Dest, void const* Src ) const override;
-	virtual void CopySingleValueFromScriptVM( void* Dest, void const* Src ) const override;
-	virtual void CopyCompleteValueToScriptVM( void* Dest, void const* Src ) const override;
-	virtual void CopyCompleteValueFromScriptVM( void* Dest, void const* Src ) const override;
-	virtual void CopyCompleteValueToScriptVM_InContainer( void* OutValue, void const* InContainer ) const override;
-	virtual void CopyCompleteValueFromScriptVM_InContainer( void* OutContainer, void const* InValue ) const override;
-	virtual bool Identical(const void* A, const void* B, uint32 PortFlags) const override;
+	UE_API virtual void CopySingleValueToScriptVM( void* Dest, void const* Src ) const override;
+	UE_API virtual void CopySingleValueFromScriptVM( void* Dest, void const* Src ) const override;
+	UE_API virtual void CopyCompleteValueToScriptVM( void* Dest, void const* Src ) const override;
+	UE_API virtual void CopyCompleteValueFromScriptVM( void* Dest, void const* Src ) const override;
+	UE_API virtual void CopyCompleteValueToScriptVM_InContainer( void* OutValue, void const* InContainer ) const override;
+	UE_API virtual void CopyCompleteValueFromScriptVM_InContainer( void* OutContainer, void const* InValue ) const override;
+	UE_API virtual bool Identical(const void* A, const void* B, uint32 PortFlags) const override;
 	// End of FProperty interface
 
 	// FObjectPropertyBase interface
-	virtual UObject* GetObjectPropertyValue(const void* PropertyValueAddress) const override;
-	virtual TObjectPtr<UObject> GetObjectPtrPropertyValue(const void* PropertyValueAddress) const override;
-	virtual UObject* GetObjectPropertyValue_InContainer(const void* ContainerAddress, int32 ArrayIndex = 0) const override;
-	virtual TObjectPtr<UObject> GetObjectPtrPropertyValue_InContainer(const void* ContainerAddress, int32 ArrayIndex = 0) const override;
+	UE_API virtual UObject* GetObjectPropertyValue(const void* PropertyValueAddress) const override;
+	UE_API virtual TObjectPtr<UObject> GetObjectPtrPropertyValue(const void* PropertyValueAddress) const override;
+	UE_API virtual UObject* GetObjectPropertyValue_InContainer(const void* ContainerAddress, int32 ArrayIndex = 0) const override;
+	UE_API virtual TObjectPtr<UObject> GetObjectPtrPropertyValue_InContainer(const void* ContainerAddress, int32 ArrayIndex = 0) const override;
 protected:
-	virtual void SetObjectPropertyValueUnchecked(void* PropertyValueAddress, UObject* Value) const override;
-	virtual void SetObjectPtrPropertyValueUnchecked(void* PropertyValueAddress, TObjectPtr<UObject> Ptr) const override;
-	virtual void SetObjectPropertyValueUnchecked_InContainer(void* ContainerAddress, UObject* Value, int32 ArrayIndex = 0) const override;
-	virtual void SetObjectPtrPropertyValueUnchecked_InContainer(void* ContainerAddress, TObjectPtr<UObject> Ptr, int32 ArrayIndex = 0) const override;
+	UE_API virtual void SetObjectPropertyValueUnchecked(void* PropertyValueAddress, UObject* Value) const override;
+	UE_API virtual void SetObjectPtrPropertyValueUnchecked(void* PropertyValueAddress, TObjectPtr<UObject> Ptr) const override;
+	UE_API virtual void SetObjectPropertyValueUnchecked_InContainer(void* ContainerAddress, UObject* Value, int32 ArrayIndex = 0) const override;
+	UE_API virtual void SetObjectPtrPropertyValueUnchecked_InContainer(void* ContainerAddress, TObjectPtr<UObject> Ptr, int32 ArrayIndex = 0) const override;
 public:
-	virtual FString GetCPPTypeCustom(FString* ExtendedTypeText, uint32 CPPExportFlags, const FString& InnerNativeTypeName)  const override;
+	UE_API virtual FString GetCPPTypeCustom(FString* ExtendedTypeText, uint32 CPPExportFlags, const FString& InnerNativeTypeName)  const override;
 	// End of FObjectPropertyBase interface
 	
 	/**
@@ -2922,7 +2924,7 @@ public:
 	 * @param ObjectValue Deserialized Object value
 	 * @param Defaults Defaults used during serialization
 	 */
-	void PostSerializeObjectItem(FArchive& SerializingArchive, void* Value, UObject* CurrentValue, UObject* ObjectValue, EObjectPropertyOptions Options = EObjectPropertyOptions::None, const void* Defaults = nullptr) const;
+	UE_API void PostSerializeObjectItem(FArchive& SerializingArchive, void* Value, UObject* CurrentValue, UObject* ObjectValue, EObjectPropertyOptions Options = EObjectPropertyOptions::None, const void* Defaults = nullptr) const;
 
 	inline TObjectPtr<UObject>* GetObjectPtrPropertyValuePtr(const void* PropertyValueAddress) const
 	{
@@ -2959,7 +2961,7 @@ public:
 		ClearValue(Data);
 	}
 
-	virtual void EmitIntrusiveOptionalReferenceInfo(UE::GC::FSchemaBuilder& Schema, int32 BaseOffset, TArray<const FStructProperty*>& EncounteredStructProps, UE::GC::FPropertyStack& DebugPath) override;
+	UE_API virtual void EmitIntrusiveOptionalReferenceInfo(UE::GC::FSchemaBuilder& Schema, int32 BaseOffset, TArray<const FStructProperty*>& EncounteredStructProps, UE::GC::FPropertyStack& DebugPath) override;
 };
 
 using FObjectPtrProperty UE_DEPRECATED(5.4, "FObjectPtrProperty is deprecated using FObjectProperty instead.")  = FObjectProperty;
@@ -2967,58 +2969,58 @@ using FObjectPtrProperty UE_DEPRECATED(5.4, "FObjectPtrProperty is deprecated us
 //
 // Describes a reference variable to another object which may be nil, and may turn nil at any point
 //
-class COREUOBJECT_API FWeakObjectProperty : public TFObjectPropertyBase<FWeakObjectPtr>
+class FWeakObjectProperty : public TFObjectPropertyBase<FWeakObjectPtr>
 {
-	DECLARE_FIELD(FWeakObjectProperty, TFObjectPropertyBase<FWeakObjectPtr>, CASTCLASS_FWeakObjectProperty)
+	DECLARE_FIELD_API(FWeakObjectProperty, TFObjectPropertyBase<FWeakObjectPtr>, CASTCLASS_FWeakObjectProperty, UE_API)
 
-	FWeakObjectProperty(FFieldVariant InOwner, const FName& InName, EObjectFlags InObjectFlags);
+	UE_API FWeakObjectProperty(FFieldVariant InOwner, const FName& InName, EObjectFlags InObjectFlags);
 
 	/**
 	 * Constructor used for constructing compiled in properties
 	 * @param InOwner Owner of the property
 	 * @param PropBase Pointer to the compiled in structure describing the property
 	 **/
-	FWeakObjectProperty(FFieldVariant InOwner, const UECodeGen_Private::FWeakObjectPropertyParams& Prop);
+	UE_API FWeakObjectProperty(FFieldVariant InOwner, const UECodeGen_Private::FWeakObjectPropertyParams& Prop);
 
 #if WITH_EDITORONLY_DATA
-	explicit FWeakObjectProperty(UField* InField);
+	UE_API explicit FWeakObjectProperty(UField* InField);
 #endif // WITH_EDITORONLY_DATA
 
 	// UHT interface
-	virtual FString GetCPPMacroType( FString& ExtendedTypeText ) const  override;
-	virtual FString GetCPPType( FString* ExtendedTypeText, uint32 CPPExportFlags ) const override;
-	virtual FString GetCPPTypeCustom(FString* ExtendedTypeText, uint32 CPPExportFlags, const FString& InnerNativeTypeName) const override;
+	UE_API virtual FString GetCPPMacroType( FString& ExtendedTypeText ) const  override;
+	UE_API virtual FString GetCPPType( FString* ExtendedTypeText, uint32 CPPExportFlags ) const override;
+	UE_API virtual FString GetCPPTypeCustom(FString* ExtendedTypeText, uint32 CPPExportFlags, const FString& InnerNativeTypeName) const override;
 	// End of UHT interface
 
 	// FProperty interface
-	virtual void SerializeItem(FStructuredArchive::FSlot Slot, void* Value, void const* Defaults) const override;
+	UE_API virtual void SerializeItem(FStructuredArchive::FSlot Slot, void* Value, void const* Defaults) const override;
 #if UE_WITH_REMOTE_OBJECT_HANDLE
-	virtual bool Identical(const void* A, const void* B, uint32 PortFlags) const override;
-	virtual bool NetSerializeItem(FArchive& Ar, UPackageMap* Map, void* Data, TArray<uint8>* MetaData = nullptr) const override;
+	UE_API virtual bool Identical(const void* A, const void* B, uint32 PortFlags) const override;
+	UE_API virtual bool NetSerializeItem(FArchive& Ar, UPackageMap* Map, void* Data, TArray<uint8>* MetaData = nullptr) const override;
 #endif
 protected:
-	virtual void LinkInternal(FArchive& Ar) override;
+	UE_API virtual void LinkInternal(FArchive& Ar) override;
 private:
-	virtual uint32 GetValueTypeHashInternal(const void* Src) const override;
+	UE_API virtual uint32 GetValueTypeHashInternal(const void* Src) const override;
 public:
-	virtual void CopySingleValueToScriptVM( void* Dest, void const* Src ) const override;
-	virtual void CopySingleValueFromScriptVM( void* Dest, void const* Src ) const override;
-	virtual void CopyCompleteValueToScriptVM( void* Dest, void const* Src ) const override;
-	virtual void CopyCompleteValueFromScriptVM( void* Dest, void const* Src ) const override;
-	virtual void CopyCompleteValueToScriptVM_InContainer( void* OutValue, void const* InContainer ) const override;
-	virtual void CopyCompleteValueFromScriptVM_InContainer( void* OutContainer, void const* InValue ) const override;
+	UE_API virtual void CopySingleValueToScriptVM( void* Dest, void const* Src ) const override;
+	UE_API virtual void CopySingleValueFromScriptVM( void* Dest, void const* Src ) const override;
+	UE_API virtual void CopyCompleteValueToScriptVM( void* Dest, void const* Src ) const override;
+	UE_API virtual void CopyCompleteValueFromScriptVM( void* Dest, void const* Src ) const override;
+	UE_API virtual void CopyCompleteValueToScriptVM_InContainer( void* OutValue, void const* InContainer ) const override;
+	UE_API virtual void CopyCompleteValueFromScriptVM_InContainer( void* OutContainer, void const* InValue ) const override;
 	// End of FProperty interface
 
 	// FObjectProperty interface
-	virtual UObject* GetObjectPropertyValue(const void* PropertyValueAddress) const override;
-	virtual TObjectPtr<UObject> GetObjectPtrPropertyValue(const void* PropertyValueAddress) const override;
-	virtual UObject* GetObjectPropertyValue_InContainer(const void* ContainerAddress, int32 ArrayIndex = 0) const override;
-	virtual TObjectPtr<UObject> GetObjectPtrPropertyValue_InContainer(const void* ContainerAddress, int32 ArrayIndex = 0) const override;
+	UE_API virtual UObject* GetObjectPropertyValue(const void* PropertyValueAddress) const override;
+	UE_API virtual TObjectPtr<UObject> GetObjectPtrPropertyValue(const void* PropertyValueAddress) const override;
+	UE_API virtual UObject* GetObjectPropertyValue_InContainer(const void* ContainerAddress, int32 ArrayIndex = 0) const override;
+	UE_API virtual TObjectPtr<UObject> GetObjectPtrPropertyValue_InContainer(const void* ContainerAddress, int32 ArrayIndex = 0) const override;
 protected:
-	virtual void SetObjectPropertyValueUnchecked(void* PropertyValueAddress, UObject* Value) const override;
-	virtual void SetObjectPtrPropertyValueUnchecked(void* PropertyValueAddress, TObjectPtr<UObject> Ptr) const override;
-	virtual void SetObjectPropertyValueUnchecked_InContainer(void* ContainerAddress, UObject* Value, int32 ArrayIndex = 0) const override;
-	virtual void SetObjectPtrPropertyValueUnchecked_InContainer(void* ContainerAddress, TObjectPtr<UObject> Ptr, int32 ArrayIndex = 0) const override;
+	UE_API virtual void SetObjectPropertyValueUnchecked(void* PropertyValueAddress, UObject* Value) const override;
+	UE_API virtual void SetObjectPtrPropertyValueUnchecked(void* PropertyValueAddress, TObjectPtr<UObject> Ptr) const override;
+	UE_API virtual void SetObjectPropertyValueUnchecked_InContainer(void* ContainerAddress, UObject* Value, int32 ArrayIndex = 0) const override;
+	UE_API virtual void SetObjectPtrPropertyValueUnchecked_InContainer(void* ContainerAddress, TObjectPtr<UObject> Ptr, int32 ArrayIndex = 0) const override;
 public:
 	// End of FObjectProperty interface
 };
@@ -3026,55 +3028,55 @@ public:
 //
 // Describes a reference variable to another object which may be nil, and will become valid or invalid at any point
 //
-class COREUOBJECT_API FLazyObjectProperty : public TFObjectPropertyBase<FLazyObjectPtr>
+class FLazyObjectProperty : public TFObjectPropertyBase<FLazyObjectPtr>
 {
-	DECLARE_FIELD(FLazyObjectProperty, TFObjectPropertyBase<FLazyObjectPtr>, CASTCLASS_FLazyObjectProperty)
+	DECLARE_FIELD_API(FLazyObjectProperty, TFObjectPropertyBase<FLazyObjectPtr>, CASTCLASS_FLazyObjectProperty, UE_API)
 
-	FLazyObjectProperty(FFieldVariant InOwner, const FName& InName, EObjectFlags InObjectFlags);
+	UE_API FLazyObjectProperty(FFieldVariant InOwner, const FName& InName, EObjectFlags InObjectFlags);
 
 	/**
 	 * Constructor used for constructing compiled in properties
 	 * @param InOwner Owner of the property
 	 * @param PropBase Pointer to the compiled in structure describing the property
 	 **/
-	FLazyObjectProperty(FFieldVariant InOwner, const UECodeGen_Private::FLazyObjectPropertyParams& Prop);
+	UE_API FLazyObjectProperty(FFieldVariant InOwner, const UECodeGen_Private::FLazyObjectPropertyParams& Prop);
 
 #if WITH_EDITORONLY_DATA
-	explicit FLazyObjectProperty(UField* InField);
+	UE_API explicit FLazyObjectProperty(UField* InField);
 #endif // WITH_EDITORONLY_DATA
 
 	// UHT interface
-	virtual FString GetCPPMacroType( FString& ExtendedTypeText ) const  override;
-	virtual FString GetCPPType( FString* ExtendedTypeText, uint32 CPPExportFlags ) const override;
-	virtual FString GetCPPTypeCustom(FString* ExtendedTypeText, uint32 CPPExportFlags, const FString& InnerNativeTypeName) const override;
+	UE_API virtual FString GetCPPMacroType( FString& ExtendedTypeText ) const  override;
+	UE_API virtual FString GetCPPType( FString* ExtendedTypeText, uint32 CPPExportFlags ) const override;
+	UE_API virtual FString GetCPPTypeCustom(FString* ExtendedTypeText, uint32 CPPExportFlags, const FString& InnerNativeTypeName) const override;
 	// End of UHT interface
 
 	// FProperty interface
-	virtual FName GetID() const override;
-	virtual bool Identical( const void* A, const void* B, uint32 PortFlags ) const override;
-	virtual void SerializeItem( FStructuredArchive::FSlot Slot, void* Value, void const* Defaults ) const override;
-	virtual void CopySingleValueToScriptVM( void* Dest, void const* Src ) const override;
-	virtual void CopySingleValueFromScriptVM( void* Dest, void const* Src ) const override;
-	virtual void CopyCompleteValueToScriptVM( void* Dest, void const* Src ) const override;
-	virtual void CopyCompleteValueFromScriptVM( void* Dest, void const* Src ) const override;
-	virtual void CopyCompleteValueToScriptVM_InContainer( void* OutValue, void const* InContainer ) const override;
-	virtual void CopyCompleteValueFromScriptVM_InContainer( void* OutContainer, void const* InValue ) const override;
+	UE_API virtual FName GetID() const override;
+	UE_API virtual bool Identical( const void* A, const void* B, uint32 PortFlags ) const override;
+	UE_API virtual void SerializeItem( FStructuredArchive::FSlot Slot, void* Value, void const* Defaults ) const override;
+	UE_API virtual void CopySingleValueToScriptVM( void* Dest, void const* Src ) const override;
+	UE_API virtual void CopySingleValueFromScriptVM( void* Dest, void const* Src ) const override;
+	UE_API virtual void CopyCompleteValueToScriptVM( void* Dest, void const* Src ) const override;
+	UE_API virtual void CopyCompleteValueFromScriptVM( void* Dest, void const* Src ) const override;
+	UE_API virtual void CopyCompleteValueToScriptVM_InContainer( void* OutValue, void const* InContainer ) const override;
+	UE_API virtual void CopyCompleteValueFromScriptVM_InContainer( void* OutContainer, void const* InValue ) const override;
 	// End of FProperty interface
 
 	// FObjectProperty interface
-	virtual UObject* GetObjectPropertyValue(const void* PropertyValueAddress) const override;
-	virtual TObjectPtr<UObject> GetObjectPtrPropertyValue(const void* PropertyValueAddress) const override;
-	virtual UObject* GetObjectPropertyValue_InContainer(const void* ContainerAddress, int32 ArrayIndex = 0) const override;
-	virtual TObjectPtr<UObject> GetObjectPtrPropertyValue_InContainer(const void* ContainerAddress, int32 ArrayIndex = 0) const override;
+	UE_API virtual UObject* GetObjectPropertyValue(const void* PropertyValueAddress) const override;
+	UE_API virtual TObjectPtr<UObject> GetObjectPtrPropertyValue(const void* PropertyValueAddress) const override;
+	UE_API virtual UObject* GetObjectPropertyValue_InContainer(const void* ContainerAddress, int32 ArrayIndex = 0) const override;
+	UE_API virtual TObjectPtr<UObject> GetObjectPtrPropertyValue_InContainer(const void* ContainerAddress, int32 ArrayIndex = 0) const override;
 protected:
-	virtual void SetObjectPropertyValueUnchecked(void* PropertyValueAddress, UObject* Value) const override;
-	virtual void SetObjectPtrPropertyValueUnchecked(void* PropertyValueAddress, TObjectPtr<UObject> Ptr) const override;
-	virtual void SetObjectPropertyValueUnchecked_InContainer(void* ContainerAddress, UObject* Value, int32 ArrayIndex = 0) const override;
-	virtual void SetObjectPtrPropertyValueUnchecked_InContainer(void* ContainerAddress, TObjectPtr<UObject> Ptr, int32 ArrayIndex = 0) const override;
+	UE_API virtual void SetObjectPropertyValueUnchecked(void* PropertyValueAddress, UObject* Value) const override;
+	UE_API virtual void SetObjectPtrPropertyValueUnchecked(void* PropertyValueAddress, TObjectPtr<UObject> Ptr) const override;
+	UE_API virtual void SetObjectPropertyValueUnchecked_InContainer(void* ContainerAddress, UObject* Value, int32 ArrayIndex = 0) const override;
+	UE_API virtual void SetObjectPtrPropertyValueUnchecked_InContainer(void* ContainerAddress, TObjectPtr<UObject> Ptr, int32 ArrayIndex = 0) const override;
 public:
-	virtual bool AllowCrossLevel() const override;
+	UE_API virtual bool AllowCrossLevel() const override;
 private:
-	virtual uint32 GetValueTypeHashInternal(const void* Src) const override;
+	UE_API virtual uint32 GetValueTypeHashInternal(const void* Src) const override;
 public:
 	// End of FObjectProperty interface
 };
@@ -3082,18 +3084,18 @@ public:
 //
 // Describes a reference variable to another object which may be nil, and will become valid or invalid at any point
 //
-class COREUOBJECT_API FSoftObjectProperty : public TFObjectPropertyBase<FSoftObjectPtr>
+class FSoftObjectProperty : public TFObjectPropertyBase<FSoftObjectPtr>
 {
-	DECLARE_FIELD(FSoftObjectProperty, TFObjectPropertyBase<FSoftObjectPtr>, CASTCLASS_FSoftObjectProperty)
+	DECLARE_FIELD_API(FSoftObjectProperty, TFObjectPropertyBase<FSoftObjectPtr>, CASTCLASS_FSoftObjectProperty, UE_API)
 
-	FSoftObjectProperty(FFieldVariant InOwner, const FName& InName, EObjectFlags InObjectFlags);
+	UE_API FSoftObjectProperty(FFieldVariant InOwner, const FName& InName, EObjectFlags InObjectFlags);
 
 	/**
 	 * Constructor used for constructing compiled-in properties
 	 * @param InOwner Owner of the property
 	 * @param PropBase Pointer to the compiled in structure describing the property
 	 **/
-	FSoftObjectProperty(FFieldVariant InOwner, const UECodeGen_Private::FSoftObjectPropertyParams& Prop);
+	UE_API FSoftObjectProperty(FFieldVariant InOwner, const UECodeGen_Private::FSoftObjectPropertyParams& Prop);
 
 	/**
 	 * Constructor used for constructing compiled-in properties
@@ -3101,48 +3103,48 @@ class COREUOBJECT_API FSoftObjectProperty : public TFObjectPropertyBase<FSoftObj
 	 * @param PropBase Pointer to the compiled in structure describing the property
 	 * @param Class Class of the object this property represents
 	 **/
-	FSoftObjectProperty(FFieldVariant InOwner, const UECodeGen_Private::FObjectPropertyParamsWithoutClass& Prop, UClass* InClass);
+	UE_API FSoftObjectProperty(FFieldVariant InOwner, const UECodeGen_Private::FObjectPropertyParamsWithoutClass& Prop, UClass* InClass);
 
 #if WITH_EDITORONLY_DATA
-	explicit FSoftObjectProperty(UField* InField);
+	UE_API explicit FSoftObjectProperty(UField* InField);
 #endif // WITH_EDITORONLY_DATA
 
 	// UHT interface
-	virtual FString GetCPPMacroType( FString& ExtendedTypeText ) const  override;
+	UE_API virtual FString GetCPPMacroType( FString& ExtendedTypeText ) const  override;
 	// End of UHT interface
 
 	// FProperty interface
-	virtual FName GetID() const override;
-	virtual bool Identical( const void* A, const void* B, uint32 PortFlags ) const override;
-	virtual void SerializeItem( FStructuredArchive::FSlot Slot, void* Value, void const* Defaults) const override;
-	virtual bool NetSerializeItem(FArchive& Ar, UPackageMap* Map, void* Data, TArray<uint8> * MetaData = NULL) const override;
+	UE_API virtual FName GetID() const override;
+	UE_API virtual bool Identical( const void* A, const void* B, uint32 PortFlags ) const override;
+	UE_API virtual void SerializeItem( FStructuredArchive::FSlot Slot, void* Value, void const* Defaults) const override;
+	UE_API virtual bool NetSerializeItem(FArchive& Ar, UPackageMap* Map, void* Data, TArray<uint8> * MetaData = NULL) const override;
 protected:
-	virtual void LinkInternal(FArchive& Ar) override;
-	virtual void ExportText_Internal( FString& ValueStr, const void* PropertyValueOrContainer, EPropertyPointerType PropertyPointerType, const void* DefaultValue, UObject* Parent, int32 PortFlags, UObject* ExportRootScope ) const override;
-	virtual const TCHAR* ImportText_Internal(const TCHAR* Buffer, void* ContainerOrPropertyPtr, EPropertyPointerType PropertyPointerType, UObject* OwnerObject, int32 PortFlags, FOutputDevice* ErrorText) const override;
+	UE_API virtual void LinkInternal(FArchive& Ar) override;
+	UE_API virtual void ExportText_Internal( FString& ValueStr, const void* PropertyValueOrContainer, EPropertyPointerType PropertyPointerType, const void* DefaultValue, UObject* Parent, int32 PortFlags, UObject* ExportRootScope ) const override;
+	UE_API virtual const TCHAR* ImportText_Internal(const TCHAR* Buffer, void* ContainerOrPropertyPtr, EPropertyPointerType PropertyPointerType, UObject* OwnerObject, int32 PortFlags, FOutputDevice* ErrorText) const override;
 public:
-	virtual EConvertFromTypeResult ConvertFromType(const FPropertyTag& Tag, FStructuredArchive::FSlot Slot, uint8* Data, UStruct* DefaultsStruct, const uint8* Defaults) override;
+	UE_API virtual EConvertFromTypeResult ConvertFromType(const FPropertyTag& Tag, FStructuredArchive::FSlot Slot, uint8* Data, UStruct* DefaultsStruct, const uint8* Defaults) override;
 	// End of FProperty interface
 
 	// FObjectProperty interface
-	virtual UObject* LoadObjectPropertyValue(const void* PropertyValueAddress) const override;
+	UE_API virtual UObject* LoadObjectPropertyValue(const void* PropertyValueAddress) const override;
 
-	virtual UObject* GetObjectPropertyValue(const void* PropertyValueAddress) const override;
-	virtual TObjectPtr<UObject> GetObjectPtrPropertyValue(const void* PropertyValueAddress) const override;
-	virtual UObject* GetObjectPropertyValue_InContainer(const void* ContainerAddress, int32 ArrayIndex = 0) const override;
-	virtual TObjectPtr<UObject> GetObjectPtrPropertyValue_InContainer(const void* ContainerAddress, int32 ArrayIndex = 0) const override;
+	UE_API virtual UObject* GetObjectPropertyValue(const void* PropertyValueAddress) const override;
+	UE_API virtual TObjectPtr<UObject> GetObjectPtrPropertyValue(const void* PropertyValueAddress) const override;
+	UE_API virtual UObject* GetObjectPropertyValue_InContainer(const void* ContainerAddress, int32 ArrayIndex = 0) const override;
+	UE_API virtual TObjectPtr<UObject> GetObjectPtrPropertyValue_InContainer(const void* ContainerAddress, int32 ArrayIndex = 0) const override;
 protected:
-	virtual void SetObjectPropertyValueUnchecked(void* PropertyValueAddress, UObject* Value) const override;
-	virtual void SetObjectPtrPropertyValueUnchecked(void* PropertyValueAddress, TObjectPtr<UObject> Ptr) const override;
-	virtual void SetObjectPropertyValueUnchecked_InContainer(void* ContainerAddress, UObject* Value, int32 ArrayIndex = 0) const override;
-	virtual void SetObjectPtrPropertyValueUnchecked_InContainer(void* ContainerAddress, TObjectPtr<UObject> Ptr, int32 ArrayIndex = 0) const override;
+	UE_API virtual void SetObjectPropertyValueUnchecked(void* PropertyValueAddress, UObject* Value) const override;
+	UE_API virtual void SetObjectPtrPropertyValueUnchecked(void* PropertyValueAddress, TObjectPtr<UObject> Ptr) const override;
+	UE_API virtual void SetObjectPropertyValueUnchecked_InContainer(void* ContainerAddress, UObject* Value, int32 ArrayIndex = 0) const override;
+	UE_API virtual void SetObjectPtrPropertyValueUnchecked_InContainer(void* ContainerAddress, TObjectPtr<UObject> Ptr, int32 ArrayIndex = 0) const override;
 public:
-	virtual bool AllowCrossLevel() const override;
-	virtual FString GetCPPTypeCustom(FString* ExtendedTypeText, uint32 CPPExportFlags, const FString& InnerNativeTypeName) const override;
-	virtual FString GetCPPType(FString* ExtendedTypeText, uint32 CPPExportFlags) const override;
+	UE_API virtual bool AllowCrossLevel() const override;
+	UE_API virtual FString GetCPPTypeCustom(FString* ExtendedTypeText, uint32 CPPExportFlags, const FString& InnerNativeTypeName) const override;
+	UE_API virtual FString GetCPPType(FString* ExtendedTypeText, uint32 CPPExportFlags) const override;
 
 private:
-	virtual uint32 GetValueTypeHashInternal(const void* Src) const override;
+	UE_API virtual uint32 GetValueTypeHashInternal(const void* Src) const override;
 public:
 	// Note: FSoftObjectProperty does not override the Copy*VM functions, as ScriptVM should store Asset as a FSoftObjectPtr not as a UObject*.
 
@@ -3156,51 +3158,51 @@ public:
 //
 // Describes a reference variable to another object which may be nil.
 //
-class COREUOBJECT_API FClassProperty : public FObjectProperty
+class FClassProperty : public FObjectProperty
 {
-	DECLARE_FIELD(FClassProperty, FObjectProperty, CASTCLASS_FClassProperty)
+	DECLARE_FIELD_API(FClassProperty, FObjectProperty, CASTCLASS_FClassProperty, UE_API)
 
 	// Variables.
 	TObjectPtr<class UClass> MetaClass;
 public:
 
-	FClassProperty(FFieldVariant InOwner, const FName& InName, EObjectFlags InObjectFlags);
+	UE_API FClassProperty(FFieldVariant InOwner, const FName& InName, EObjectFlags InObjectFlags);
 
 	/**
 	 * Constructor used for constructing compiled in properties
 	 * @param InOwner Owner of the property
 	 * @param PropBase Pointer to the compiled in structure describing the property
 	 **/
-	FClassProperty(FFieldVariant InOwner, const UECodeGen_Private::FClassPropertyParams& Prop);
+	UE_API FClassProperty(FFieldVariant InOwner, const UECodeGen_Private::FClassPropertyParams& Prop);
 
 #if WITH_EDITORONLY_DATA
-	explicit FClassProperty(UField* InField);
+	UE_API explicit FClassProperty(UField* InField);
 #endif // WITH_EDITORONLY_DATA
 
 	// UObject interface
-	virtual void Serialize( FArchive& Ar ) override;
-	virtual void AddReferencedObjects(FReferenceCollector& Collector) override;
-	virtual void BeginDestroy() override;
+	UE_API virtual void Serialize( FArchive& Ar ) override;
+	UE_API virtual void AddReferencedObjects(FReferenceCollector& Collector) override;
+	UE_API virtual void BeginDestroy() override;
 	// End of UObject interface
 
 	// Field Interface
-	virtual void PostDuplicate(const FField& InField) override;
+	UE_API virtual void PostDuplicate(const FField& InField) override;
 
 	// UHT interface
-	virtual FString GetCPPType(FString* ExtendedTypeText, uint32 CPPExportFlags)  const override;
-	virtual FString GetCPPMacroType( FString& ExtendedTypeText ) const  override;
+	UE_API virtual FString GetCPPType(FString* ExtendedTypeText, uint32 CPPExportFlags)  const override;
+	UE_API virtual FString GetCPPMacroType( FString& ExtendedTypeText ) const  override;
 	// End of UHT interface
 
 	// FProperty interface
-	virtual const TCHAR* ImportText_Internal(const TCHAR* Buffer, void* ContainerOrPropertyPtr, EPropertyPointerType PropertyPointerType, UObject* OwnerObject, int32 PortFlags, FOutputDevice* ErrorText) const override;
-	virtual bool SameType(const FProperty* Other) const override;
-	virtual bool Identical( const void* A, const void* B, uint32 PortFlags ) const override;
+	UE_API virtual const TCHAR* ImportText_Internal(const TCHAR* Buffer, void* ContainerOrPropertyPtr, EPropertyPointerType PropertyPointerType, UObject* OwnerObject, int32 PortFlags, FOutputDevice* ErrorText) const override;
+	UE_API virtual bool SameType(const FProperty* Other) const override;
+	UE_API virtual bool Identical( const void* A, const void* B, uint32 PortFlags ) const override;
 #if WITH_EDITORONLY_DATA
-	virtual void AppendSchemaHash(FBlake3& Builder, bool bSkipEditorOnly) const override;
+	UE_API virtual void AppendSchemaHash(FBlake3& Builder, bool bSkipEditorOnly) const override;
 #endif
 	// End of FProperty interface
 
-	virtual FString GetCPPTypeCustom(FString* ExtendedTypeText, uint32 CPPExportFlags, const FString& InnerNativeTypeName)  const override;
+	UE_API virtual FString GetCPPTypeCustom(FString* ExtendedTypeText, uint32 CPPExportFlags, const FString& InnerNativeTypeName)  const override;
 
 	/**
 	 * Setter function for this property's MetaClass member. Favor this function 
@@ -3212,7 +3214,7 @@ public:
 	 * @param  NewMetaClass    The MetaClass you want this property set with.
 	 */
 #if USE_CIRCULAR_DEPENDENCY_LOAD_DEFERRING
-	void SetMetaClass(UClass* NewMetaClass);
+	UE_API void SetMetaClass(UClass* NewMetaClass);
 #else
 	FORCEINLINE void SetMetaClass(UClass* NewMetaClass) { MetaClass = NewMetaClass; }
 #endif // USE_CIRCULAR_DEPENDENCY_LOAD_DEFERRING
@@ -3226,46 +3228,46 @@ using FClassPtrProperty UE_DEPRECATED(5.4, "FClassPtrProperty is deprecated use
 //
 // Describes a reference variable to another class which may be nil, and will become valid or invalid at any point
 //
-class COREUOBJECT_API FSoftClassProperty : public FSoftObjectProperty
+class FSoftClassProperty : public FSoftObjectProperty
 {
-	DECLARE_FIELD(FSoftClassProperty, FSoftObjectProperty, CASTCLASS_FSoftClassProperty)
+	DECLARE_FIELD_API(FSoftClassProperty, FSoftObjectProperty, CASTCLASS_FSoftClassProperty, UE_API)
 
 	// Variables.
 	TObjectPtr<class UClass> MetaClass;
 public:
 
-	FSoftClassProperty(FFieldVariant InOwner, const FName& InName, EObjectFlags InObjectFlags);
+	UE_API FSoftClassProperty(FFieldVariant InOwner, const FName& InName, EObjectFlags InObjectFlags);
 
 	/**
 	 * Constructor used for constructing compiled in properties
 	 * @param InOwner Owner of the property
 	 * @param PropBase Pointer to the compiled in structure describing the property
 	 **/
-	FSoftClassProperty(FFieldVariant InOwner, const UECodeGen_Private::FSoftClassPropertyParams& Prop);
+	UE_API FSoftClassProperty(FFieldVariant InOwner, const UECodeGen_Private::FSoftClassPropertyParams& Prop);
 
 #if WITH_EDITORONLY_DATA
-	explicit FSoftClassProperty(UField* InField);
+	UE_API explicit FSoftClassProperty(UField* InField);
 #endif // WITH_EDITORONLY_DATA
 
 	// UHT interface
-	virtual FString GetCPPType(FString* ExtendedTypeText, uint32 CPPExportFlags) const override;
-	virtual FString GetCPPMacroType( FString& ExtendedTypeText ) const  override;
+	UE_API virtual FString GetCPPType(FString* ExtendedTypeText, uint32 CPPExportFlags) const override;
+	UE_API virtual FString GetCPPMacroType( FString& ExtendedTypeText ) const  override;
 	// End of UHT interface
 
 	// Field Interface
-	virtual void PostDuplicate(const FField& InField) override;
+	UE_API virtual void PostDuplicate(const FField& InField) override;
 
 	// UObject interface
-	virtual void Serialize( FArchive& Ar ) override;
-	virtual void AddReferencedObjects(FReferenceCollector& Collector) override;
-	virtual void BeginDestroy() override;
+	UE_API virtual void Serialize( FArchive& Ar ) override;
+	UE_API virtual void AddReferencedObjects(FReferenceCollector& Collector) override;
+	UE_API virtual void BeginDestroy() override;
 	// End of UObject interface
 
 	// FProperty interface
-	virtual bool SameType(const FProperty* Other) const override;
+	UE_API virtual bool SameType(const FProperty* Other) const override;
 	// End of FProperty interface
 
-	virtual FString GetCPPTypeCustom(FString* ExtendedTypeText, uint32 CPPExportFlags, const FString& InnerNativeTypeName)  const override;
+	UE_API virtual FString GetCPPTypeCustom(FString* ExtendedTypeText, uint32 CPPExportFlags, const FString& InnerNativeTypeName)  const override;
 
 	/**
 	 * Setter function for this property's MetaClass member. Favor this function 
@@ -3277,7 +3279,7 @@ public:
 	 * @param  NewMetaClass    The MetaClass you want this property set with.
 	 */
 #if USE_CIRCULAR_DEPENDENCY_LOAD_DEFERRING
-	void SetMetaClass(UClass* NewMetaClass);
+	UE_API void SetMetaClass(UClass* NewMetaClass);
 #else
 	FORCEINLINE void SetMetaClass(UClass* NewMetaClass) { MetaClass = NewMetaClass; }
 #endif // USE_CIRCULAR_DEPENDENCY_LOAD_DEFERRING
@@ -3292,9 +3294,9 @@ public:
  * script header files as a TScriptInterface.
  */
 
-class COREUOBJECT_API FInterfaceProperty : public TProperty<FScriptInterface, FProperty>
+class FInterfaceProperty : public TProperty<FScriptInterface, FProperty>
 {
-	DECLARE_FIELD(FInterfaceProperty, (TProperty<FScriptInterface, FProperty>), CASTCLASS_FInterfaceProperty)
+	DECLARE_FIELD_API(FInterfaceProperty, (TProperty<FScriptInterface, FProperty>), CASTCLASS_FInterfaceProperty, UE_API)
 
 	/** The native interface class that this interface property refers to */
 	TObjectPtr<class	UClass>		InterfaceClass;
@@ -3303,49 +3305,49 @@ public:
 	typedef Super::TTypeFundamentals TTypeFundamentals;
 	typedef TTypeFundamentals::TCppType TCppType;
 
-	FInterfaceProperty(FFieldVariant InOwner, const FName& InName, EObjectFlags InObjectFlags);
+	UE_API FInterfaceProperty(FFieldVariant InOwner, const FName& InName, EObjectFlags InObjectFlags);
 
 	/**
 	 * Constructor used for constructing compiled in properties
 	 * @param InOwner Owner of the property
 	 * @param PropBase Pointer to the compiled in structure describing the property
 	 **/
-	FInterfaceProperty(FFieldVariant InOwner, const UECodeGen_Private::FInterfacePropertyParams& Prop);
+	UE_API FInterfaceProperty(FFieldVariant InOwner, const UECodeGen_Private::FInterfacePropertyParams& Prop);
 
 #if WITH_EDITORONLY_DATA
-	explicit FInterfaceProperty(UField* InField);
+	UE_API explicit FInterfaceProperty(UField* InField);
 #endif // WITH_EDITORONLY_DATA
 
 	// UHT interface
-	virtual FString GetCPPMacroType( FString& ExtendedTypeText ) const  override;
-	virtual FString GetCPPType( FString* ExtendedTypeText, uint32 CPPExportFlags ) const override;
+	UE_API virtual FString GetCPPMacroType( FString& ExtendedTypeText ) const  override;
+	UE_API virtual FString GetCPPType( FString* ExtendedTypeText, uint32 CPPExportFlags ) const override;
 	// End of UHT interface
 
 	// Field interface
-	virtual void PostDuplicate(const FField& InField) override;
+	UE_API virtual void PostDuplicate(const FField& InField) override;
 
 	// FProperty interface
-	virtual void LinkInternal(FArchive& Ar) override;
-	virtual bool Identical( const void* A, const void* B, uint32 PortFlags ) const override;
-	virtual void SerializeItem( FStructuredArchive::FSlot Slot, void* Value, void const* Defaults) const override;
-	virtual bool NetSerializeItem( FArchive& Ar, UPackageMap* Map, void* Data, TArray<uint8> * MetaData = NULL ) const override;
+	UE_API virtual void LinkInternal(FArchive& Ar) override;
+	UE_API virtual bool Identical( const void* A, const void* B, uint32 PortFlags ) const override;
+	UE_API virtual void SerializeItem( FStructuredArchive::FSlot Slot, void* Value, void const* Defaults) const override;
+	UE_API virtual bool NetSerializeItem( FArchive& Ar, UPackageMap* Map, void* Data, TArray<uint8> * MetaData = NULL ) const override;
 	virtual bool SupportsNetSharedSerialization() const override { return false; }
 protected:
-	virtual void ExportText_Internal( FString& ValueStr, const void* PropertyValueOrContainer, EPropertyPointerType PropertyPointerType, const void* DefaultValue, UObject* Parent, int32 PortFlags, UObject* ExportRootScope ) const override;
-	virtual const TCHAR* ImportText_Internal(const TCHAR* Buffer, void* ContainerOrPropertyPtr, EPropertyPointerType PropertyPointerType, UObject* OwnerObject, int32 PortFlags, FOutputDevice* ErrorText) const override;
+	UE_API virtual void ExportText_Internal( FString& ValueStr, const void* PropertyValueOrContainer, EPropertyPointerType PropertyPointerType, const void* DefaultValue, UObject* Parent, int32 PortFlags, UObject* ExportRootScope ) const override;
+	UE_API virtual const TCHAR* ImportText_Internal(const TCHAR* Buffer, void* ContainerOrPropertyPtr, EPropertyPointerType PropertyPointerType, UObject* OwnerObject, int32 PortFlags, FOutputDevice* ErrorText) const override;
 public:
-	virtual bool ContainsObjectReference(TArray<const FStructProperty*>& EncounteredStructProps, EPropertyObjectReferenceType InReferenceType = EPropertyObjectReferenceType::Strong) const override;
-	virtual bool SameType(const FProperty* Other) const override;
+	UE_API virtual bool ContainsObjectReference(TArray<const FStructProperty*>& EncounteredStructProps, EPropertyObjectReferenceType InReferenceType = EPropertyObjectReferenceType::Strong) const override;
+	UE_API virtual bool SameType(const FProperty* Other) const override;
 #if WITH_EDITORONLY_DATA
-	virtual void AppendSchemaHash(FBlake3& Builder, bool bSkipEditorOnly) const override;
+	UE_API virtual void AppendSchemaHash(FBlake3& Builder, bool bSkipEditorOnly) const override;
 #endif
 	// End of FProperty interface
 
 	// UObject interface
-	virtual void Serialize( FArchive& Ar ) override;
-	virtual void EmitReferenceInfo(UE::GC::FSchemaBuilder& Schema, int32 BaseOffset, TArray<const FStructProperty*>& EncounteredStructProps, UE::GC::FPropertyStack& DebugPath) override;
-	virtual void BeginDestroy() override;
-	virtual void AddReferencedObjects(FReferenceCollector& Collector) override;
+	UE_API virtual void Serialize( FArchive& Ar ) override;
+	UE_API virtual void EmitReferenceInfo(UE::GC::FSchemaBuilder& Schema, int32 BaseOffset, TArray<const FStructProperty*>& EncounteredStructProps, UE::GC::FPropertyStack& DebugPath) override;
+	UE_API virtual void BeginDestroy() override;
+	UE_API virtual void AddReferencedObjects(FReferenceCollector& Collector) override;
 	// End of UObject interface
 
 	/**
@@ -3358,7 +3360,7 @@ public:
 	 * @param  NewInterfaceClass    The InterfaceClass you want this property set with.
 	 */
 #if USE_CIRCULAR_DEPENDENCY_LOAD_DEFERRING
-	void SetInterfaceClass(UClass* NewInterfaceClass);
+	UE_API void SetInterfaceClass(UClass* NewInterfaceClass);
 #else
 	FORCEINLINE void SetInterfaceClass(UClass* NewInterfaceClass) { InterfaceClass = NewInterfaceClass; }
 #endif // USE_CIRCULAR_DEPENDENCY_LOAD_DEFERRING
@@ -3372,33 +3374,33 @@ public:
 // Describes a name variable pointing into the global name table.
 //
 
-class COREUOBJECT_API FNameProperty : public TProperty_WithEqualityAndSerializer<FName, FProperty>
+class FNameProperty : public TProperty_WithEqualityAndSerializer<FName, FProperty>
 {
-	DECLARE_FIELD(FNameProperty, (TProperty_WithEqualityAndSerializer<FName, FProperty>), CASTCLASS_FNameProperty)
+	DECLARE_FIELD_API(FNameProperty, (TProperty_WithEqualityAndSerializer<FName, FProperty>), CASTCLASS_FNameProperty, UE_API)
 public:
 	typedef Super::TTypeFundamentals TTypeFundamentals;
 	typedef TTypeFundamentals::TCppType TCppType;
 
-	FNameProperty(FFieldVariant InOwner, const FName& InName, EObjectFlags InObjectFlags);
+	UE_API FNameProperty(FFieldVariant InOwner, const FName& InName, EObjectFlags InObjectFlags);
 
 	/**
 	 * Constructor used for constructing compiled in properties
 	 * @param InOwner Owner of the property
 	 * @param PropBase Pointer to the compiled in structure describing the property
 	 **/
-	FNameProperty(FFieldVariant InOwner, const UECodeGen_Private::FNamePropertyParams& Prop);
+	UE_API FNameProperty(FFieldVariant InOwner, const UECodeGen_Private::FNamePropertyParams& Prop);
 
 #if WITH_EDITORONLY_DATA
-	explicit FNameProperty(UField* InField);
+	UE_API explicit FNameProperty(UField* InField);
 #endif // WITH_EDITORONLY_DATA
 
 	// FProperty interface
 protected:
-	virtual void ExportText_Internal( FString& ValueStr, const void* PropertyValueOrContainer, EPropertyPointerType PropertyPointerType, const void* DefaultValue, UObject* Parent, int32 PortFlags, UObject* ExportRootScope ) const override;
-	virtual const TCHAR* ImportText_Internal(const TCHAR* Buffer, void* ContainerOrPropertyPtr, EPropertyPointerType PropertyPointerType, UObject* OwnerObject, int32 PortFlags, FOutputDevice* ErrorText) const override;
+	UE_API virtual void ExportText_Internal( FString& ValueStr, const void* PropertyValueOrContainer, EPropertyPointerType PropertyPointerType, const void* DefaultValue, UObject* Parent, int32 PortFlags, UObject* ExportRootScope ) const override;
+	UE_API virtual const TCHAR* ImportText_Internal(const TCHAR* Buffer, void* ContainerOrPropertyPtr, EPropertyPointerType PropertyPointerType, UObject* OwnerObject, int32 PortFlags, FOutputDevice* ErrorText) const override;
 public:
-	virtual EConvertFromTypeResult ConvertFromType(const FPropertyTag& Tag, FStructuredArchive::FSlot Slot, uint8* Data, UStruct* DefaultsStruct, const uint8* Defaults) override;
-	uint32 GetValueTypeHashInternal(const void* Src) const override;
+	UE_API virtual EConvertFromTypeResult ConvertFromType(const FPropertyTag& Tag, FStructuredArchive::FSlot Slot, uint8* Data, UStruct* DefaultsStruct, const uint8* Defaults) override;
+	UE_API uint32 GetValueTypeHashInternal(const void* Src) const override;
 	// End of FProperty interface
 };
 
@@ -3418,9 +3420,9 @@ using FFreezableScriptArray = TScriptArray<TMemoryImageAllocator<DEFAULT_ALIGNME
 
 class FScriptArrayHelper;
 
-class COREUOBJECT_API FArrayProperty : public TProperty<FScriptArray, FProperty>
+class FArrayProperty : public TProperty<FScriptArray, FProperty>
 {
-	DECLARE_FIELD(FArrayProperty, (TProperty<FScriptArray, FProperty>), CASTCLASS_FArrayProperty)
+	DECLARE_FIELD_API(FArrayProperty, (TProperty<FScriptArray, FProperty>), CASTCLASS_FArrayProperty, UE_API)
 
 	// Variables.
 	EArrayPropertyFlags ArrayFlags;
@@ -3438,49 +3440,49 @@ public:
 	using TTypeFundamentals = Super::TTypeFundamentals;
 	using TCppType = TTypeFundamentals::TCppType;
 
-	FArrayProperty(FFieldVariant InOwner, const FName& InName, EObjectFlags InObjectFlags, EArrayPropertyFlags InArrayPropertyFlags=EArrayPropertyFlags::None);
+	UE_API FArrayProperty(FFieldVariant InOwner, const FName& InName, EObjectFlags InObjectFlags, EArrayPropertyFlags InArrayPropertyFlags=EArrayPropertyFlags::None);
 
 	/** 
 	 * Constructor used for constructing compiled in properties
 	 * @param InOwner Owner of the property
 	 * @param PropBase Pointer to the compiled in structure describing the property
 	 **/
-	FArrayProperty(FFieldVariant InOwner, const UECodeGen_Private::FArrayPropertyParams& Prop);
+	UE_API FArrayProperty(FFieldVariant InOwner, const UECodeGen_Private::FArrayPropertyParams& Prop);
 
-	virtual ~FArrayProperty();
+	UE_API virtual ~FArrayProperty();
 
 #if WITH_EDITORONLY_DATA
-	explicit FArrayProperty(UField* InField);
+	UE_API explicit FArrayProperty(UField* InField);
 #endif // WITH_EDITORONLY_DATA
 
 	// UObject interface
-	virtual void Serialize( FArchive& Ar ) override;
-	virtual void AddReferencedObjects(FReferenceCollector& Collector) override;
-	virtual void GetPreloadDependencies(TArray<UObject*>& OutDeps) override;
+	UE_API virtual void Serialize( FArchive& Ar ) override;
+	UE_API virtual void AddReferencedObjects(FReferenceCollector& Collector) override;
+	UE_API virtual void GetPreloadDependencies(TArray<UObject*>& OutDeps) override;
 	// End of UObject interface
 
 	// Field interface
-	virtual void PostDuplicate(const FField& InField) override;
+	UE_API virtual void PostDuplicate(const FField& InField) override;
 
 	// UField interface
-	virtual void AddCppProperty(FProperty* Property) override;
-	virtual FField* GetInnerFieldByName(const FName& InName) override;
-	virtual void GetInnerFields(TArray<FField*>& OutFields) override;
+	UE_API virtual void AddCppProperty(FProperty* Property) override;
+	UE_API virtual FField* GetInnerFieldByName(const FName& InName) override;
+	UE_API virtual void GetInnerFields(TArray<FField*>& OutFields) override;
 
 	// End of UField interface
 
 	// FProperty interface
-	virtual FString GetCPPMacroType( FString& ExtendedTypeText ) const  override;
-	virtual FString GetCPPType( FString* ExtendedTypeText, uint32 CPPExportFlags ) const override;
-	virtual void LinkInternal(FArchive& Ar) override;
-	virtual bool Identical( const void* A, const void* B, uint32 PortFlags ) const override;
-	virtual void SerializeItem(FStructuredArchive::FSlot Slot, void* Value, void const* Defaults) const override;
-	virtual bool NetSerializeItem( FArchive& Ar, UPackageMap* Map, void* Data, TArray<uint8> * MetaData = NULL ) const override;
+	UE_API virtual FString GetCPPMacroType( FString& ExtendedTypeText ) const  override;
+	UE_API virtual FString GetCPPType( FString* ExtendedTypeText, uint32 CPPExportFlags ) const override;
+	UE_API virtual void LinkInternal(FArchive& Ar) override;
+	UE_API virtual bool Identical( const void* A, const void* B, uint32 PortFlags ) const override;
+	UE_API virtual void SerializeItem(FStructuredArchive::FSlot Slot, void* Value, void const* Defaults) const override;
+	UE_API virtual bool NetSerializeItem( FArchive& Ar, UPackageMap* Map, void* Data, TArray<uint8> * MetaData = NULL ) const override;
 protected:
-	virtual void ExportText_Internal( FString& ValueStr, const void* PropertyValueOrContainer, EPropertyPointerType PropertyPointerType, const void* DefaultValue, UObject* Parent, int32 PortFlags, UObject* ExportRootScope ) const override;
-	virtual const TCHAR* ImportText_Internal(const TCHAR* Buffer, void* ContainerOrPropertyPtr, EPropertyPointerType PropertyPointerType, UObject* OwnerObject, int32 PortFlags, FOutputDevice* ErrorText) const override;
-	virtual bool ContainsClearOnFinishDestroyInternal(TArray<const FStructProperty*>& EncounteredStructProps) const override;
-	virtual void FinishDestroyInternal(void* Data) const override;
+	UE_API virtual void ExportText_Internal( FString& ValueStr, const void* PropertyValueOrContainer, EPropertyPointerType PropertyPointerType, const void* DefaultValue, UObject* Parent, int32 PortFlags, UObject* ExportRootScope ) const override;
+	UE_API virtual const TCHAR* ImportText_Internal(const TCHAR* Buffer, void* ContainerOrPropertyPtr, EPropertyPointerType PropertyPointerType, UObject* OwnerObject, int32 PortFlags, FOutputDevice* ErrorText) const override;
+	UE_API virtual bool ContainsClearOnFinishDestroyInternal(TArray<const FStructProperty*>& EncounteredStructProps) const override;
+	UE_API virtual void FinishDestroyInternal(void* Data) const override;
 public:
 	virtual void InitializeValueInternal(void* Dest) const override
 	{
@@ -3501,14 +3503,14 @@ public:
 			}
 		}
 	}
-	virtual void CopyValuesInternal( void* Dest, void const* Src, int32 Count  ) const override;
-	virtual void ClearValueInternal( void* Data ) const override;
-	virtual void DestroyValueInternal( void* Dest ) const override;
-	virtual void InstanceSubobjects( void* Data, void const* DefaultData, TNotNull<UObject*> Owner, struct FObjectInstancingGraph* InstanceGraph ) override;
-	virtual bool ContainsObjectReference(TArray<const FStructProperty*>& EncounteredStructProps, EPropertyObjectReferenceType InReferenceType = EPropertyObjectReferenceType::Strong) const override;
-	virtual void EmitReferenceInfo(UE::GC::FSchemaBuilder& Schema, int32 BaseOffset, TArray<const FStructProperty*>& EncounteredStructProps, UE::GC::FPropertyStack& DebugPath) override;
-	virtual bool SameType(const FProperty* Other) const override;
-	virtual EConvertFromTypeResult ConvertFromType(const FPropertyTag& Tag, FStructuredArchive::FSlot Slot, uint8* Data, UStruct* DefaultsStruct, const uint8* Defaults) override;
+	UE_API virtual void CopyValuesInternal( void* Dest, void const* Src, int32 Count  ) const override;
+	UE_API virtual void ClearValueInternal( void* Data ) const override;
+	UE_API virtual void DestroyValueInternal( void* Dest ) const override;
+	UE_API virtual void InstanceSubobjects( void* Data, void const* DefaultData, TNotNull<UObject*> Owner, struct FObjectInstancingGraph* InstanceGraph ) override;
+	UE_API virtual bool ContainsObjectReference(TArray<const FStructProperty*>& EncounteredStructProps, EPropertyObjectReferenceType InReferenceType = EPropertyObjectReferenceType::Strong) const override;
+	UE_API virtual void EmitReferenceInfo(UE::GC::FSchemaBuilder& Schema, int32 BaseOffset, TArray<const FStructProperty*>& EncounteredStructProps, UE::GC::FPropertyStack& DebugPath) override;
+	UE_API virtual bool SameType(const FProperty* Other) const override;
+	UE_API virtual EConvertFromTypeResult ConvertFromType(const FPropertyTag& Tag, FStructuredArchive::FSlot Slot, uint8* Data, UStruct* DefaultsStruct, const uint8* Defaults) override;
 
 	virtual int32 GetMinAlignment() const override
 	{
@@ -3516,31 +3518,31 @@ public:
 		return alignof(FScriptArray);
 	}
 
-	virtual void* GetValueAddressAtIndex_Direct(const FProperty* Inner, void* InValueAddress, int32 Index) const override;
-	virtual bool UseBinaryOrNativeSerialization(const FArchive& Ar) const override;
-	virtual bool LoadTypeName(UE::FPropertyTypeName Type, const FPropertyTag* Tag = nullptr) override;
-	virtual void SaveTypeName(UE::FPropertyTypeNameBuilder& Type) const override;
-	virtual bool CanSerializeFromTypeName(UE::FPropertyTypeName Type) const override;
-	virtual EPropertyVisitorControlFlow Visit(FPropertyVisitorContext& Context, const TFunctionRef<EPropertyVisitorControlFlow(const FPropertyVisitorContext& /*Conmtext*/)> InFunc) const override;
-	virtual void* ResolveVisitedPathInfo(void* Data, const FPropertyVisitorInfo& Info) const override;
+	UE_API virtual void* GetValueAddressAtIndex_Direct(const FProperty* Inner, void* InValueAddress, int32 Index) const override;
+	UE_API virtual bool UseBinaryOrNativeSerialization(const FArchive& Ar) const override;
+	UE_API virtual bool LoadTypeName(UE::FPropertyTypeName Type, const FPropertyTag* Tag = nullptr) override;
+	UE_API virtual void SaveTypeName(UE::FPropertyTypeNameBuilder& Type) const override;
+	UE_API virtual bool CanSerializeFromTypeName(UE::FPropertyTypeName Type) const override;
+	UE_API virtual EPropertyVisitorControlFlow Visit(FPropertyVisitorContext& Context, const TFunctionRef<EPropertyVisitorControlFlow(const FPropertyVisitorContext& /*Conmtext*/)> InFunc) const override;
+	UE_API virtual void* ResolveVisitedPathInfo(void* Data, const FPropertyVisitorInfo& Info) const override;
 	// End of FProperty interface
 
-	FString GetCPPTypeCustom(FString* ExtendedTypeText, uint32 CPPExportFlags, const FString& InnerTypeText, const FString& InInnerExtendedTypeText) const;
+	UE_API FString GetCPPTypeCustom(FString* ExtendedTypeText, uint32 CPPExportFlags, const FString& InnerTypeText, const FString& InInnerExtendedTypeText) const;
 
 	/** Called by ExportText_Internal, but can also be used by a non-ArrayProperty whose ArrayDim is > 1. */
-	static void ExportTextInnerItem(FString& ValueStr, const FProperty* Inner, const void* PropertyValue, int32 PropertySize, const void* DefaultValue, int32 DefaultSize, UObject* Parent = nullptr, int32 PortFlags = 0, UObject* ExportRootScope = nullptr);
+	UE_API static void ExportTextInnerItem(FString& ValueStr, const FProperty* Inner, const void* PropertyValue, int32 PropertySize, const void* DefaultValue, int32 DefaultSize, UObject* Parent = nullptr, int32 PortFlags = 0, UObject* ExportRootScope = nullptr);
 
 	/** Called by ImportTextItem, but can also be used by a non-ArrayProperty whose ArrayDim is > 1. ArrayHelper should be supplied by ArrayProperties and nullptr for fixed-size arrays. */
-	static const TCHAR* ImportTextInnerItem(const TCHAR* Buffer, const FProperty* Inner, void* Data, int32 PortFlags, UObject* OwnerObject, FScriptArrayHelper* ArrayHelper = nullptr, FOutputDevice* ErrorText = (FOutputDevice*)GWarn);
+	UE_API static const TCHAR* ImportTextInnerItem(const TCHAR* Buffer, const FProperty* Inner, void* Data, int32 PortFlags, UObject* OwnerObject, FScriptArrayHelper* ArrayHelper = nullptr, FOutputDevice* ErrorText = (FOutputDevice*)GWarn);
 #if WITH_EDITORONLY_DATA
-	virtual void AppendSchemaHash(FBlake3& Builder, bool bSkipEditorOnly) const override;
+	UE_API virtual void AppendSchemaHash(FBlake3& Builder, bool bSkipEditorOnly) const override;
 #endif
 
-	virtual bool HasIntrusiveUnsetOptionalState() const override;
-	virtual void InitializeIntrusiveUnsetOptionalValue(void* Data) const override;
-	virtual bool IsIntrusiveOptionalValueSet(const void* Data) const override;
-	virtual void ClearIntrusiveOptionalValue(void* Data) const override;
-	virtual void EmitIntrusiveOptionalReferenceInfo(UE::GC::FSchemaBuilder& Schema, int32 BaseOffset, TArray<const FStructProperty*>& EncounteredStructProps, UE::GC::FPropertyStack& DebugPath) override;
+	UE_API virtual bool HasIntrusiveUnsetOptionalState() const override;
+	UE_API virtual void InitializeIntrusiveUnsetOptionalValue(void* Data) const override;
+	UE_API virtual bool IsIntrusiveOptionalValueSet(const void* Data) const override;
+	UE_API virtual void ClearIntrusiveOptionalValue(void* Data) const override;
+	UE_API virtual void EmitIntrusiveOptionalReferenceInfo(UE::GC::FSchemaBuilder& Schema, int32 BaseOffset, TArray<const FStructProperty*>& EncounteredStructProps, UE::GC::FPropertyStack& DebugPath) override;
 };
 
 using FFreezableScriptMap = TScriptMap<FMemoryImageSetAllocator>;
@@ -3548,9 +3550,9 @@ using FFreezableScriptMap = TScriptMap<FMemoryImageSetAllocator>;
 //@todo stever sizeof(FScriptMap) is 80 bytes, while sizeof(FFreezableScriptMap) is 56 bytes atm
 //static_assert(sizeof(FScriptMap) == sizeof(FFreezableScriptMap) && alignof(FScriptMap) == alignof(FFreezableScriptMap), "FScriptMap and FFreezableScriptMap are expected to be layout-compatible");
 
-class COREUOBJECT_API FMapProperty : public TProperty<FScriptMap, FProperty>
+class FMapProperty : public TProperty<FScriptMap, FProperty>
 {
-	DECLARE_FIELD(FMapProperty, (TProperty<FScriptMap, FProperty>), CASTCLASS_FMapProperty)
+	DECLARE_FIELD_API(FMapProperty, (TProperty<FScriptMap, FProperty>), CASTCLASS_FMapProperty, UE_API)
 
 	// Properties representing the key type and value type of the contained pairs
 	FProperty*       KeyProp;
@@ -3575,48 +3577,48 @@ public:
 	using TTypeFundamentals = Super::TTypeFundamentals;
 	using TCppType = TTypeFundamentals::TCppType;
 
-	FMapProperty(FFieldVariant InOwner, const FName& InName, EObjectFlags InObjectFlags, EMapPropertyFlags InMapFlags=EMapPropertyFlags::None);
+	UE_API FMapProperty(FFieldVariant InOwner, const FName& InName, EObjectFlags InObjectFlags, EMapPropertyFlags InMapFlags=EMapPropertyFlags::None);
 
 	/**
 	 * Constructor used for constructing compiled in properties
 	 * @param InOwner Owner of the property
 	 * @param PropBase Pointer to the compiled in structure describing the property
 	 **/
-	FMapProperty(FFieldVariant InOwner, const UECodeGen_Private::FMapPropertyParams& Prop);
+	UE_API FMapProperty(FFieldVariant InOwner, const UECodeGen_Private::FMapPropertyParams& Prop);
 
 #if WITH_EDITORONLY_DATA
-	explicit FMapProperty(UField* InField);
+	UE_API explicit FMapProperty(UField* InField);
 #endif // WITH_EDITORONLY_DATA
 
-	virtual ~FMapProperty();
+	UE_API virtual ~FMapProperty();
 
 	// UObject interface
-	virtual void Serialize(FArchive& Ar) override;
-	virtual void AddReferencedObjects(FReferenceCollector& Collector) override;
-	virtual void GetPreloadDependencies(TArray<UObject*>& OutDeps) override;
+	UE_API virtual void Serialize(FArchive& Ar) override;
+	UE_API virtual void AddReferencedObjects(FReferenceCollector& Collector) override;
+	UE_API virtual void GetPreloadDependencies(TArray<UObject*>& OutDeps) override;
 	// End of UObject interface
 
 	// Field Interface
-	virtual void PostDuplicate(const FField& InField) override;
-	virtual FField* GetInnerFieldByName(const FName& InName) override;
-	virtual void GetInnerFields(TArray<FField*>& OutFields) override;
+	UE_API virtual void PostDuplicate(const FField& InField) override;
+	UE_API virtual FField* GetInnerFieldByName(const FName& InName) override;
+	UE_API virtual void GetInnerFields(TArray<FField*>& OutFields) override;
 
 	// UField interface
-	virtual void AddCppProperty(FProperty* Property) override;
+	UE_API virtual void AddCppProperty(FProperty* Property) override;
 	// End of UField interface
 
 	// FProperty interface
-	virtual FString GetCPPMacroType(FString& ExtendedTypeText) const  override;
-	virtual FString GetCPPType(FString* ExtendedTypeText, uint32 CPPExportFlags) const override;
-	virtual void LinkInternal(FArchive& Ar) override;
-	virtual bool Identical(const void* A, const void* B, uint32 PortFlags) const override;
-	virtual void SerializeItem(FStructuredArchive::FSlot Slot, void* Value, void const* Defaults) const override;
-	virtual bool NetSerializeItem(FArchive& Ar, UPackageMap* Map, void* Data, TArray<uint8> * MetaData = NULL) const override;
+	UE_API virtual FString GetCPPMacroType(FString& ExtendedTypeText) const  override;
+	UE_API virtual FString GetCPPType(FString* ExtendedTypeText, uint32 CPPExportFlags) const override;
+	UE_API virtual void LinkInternal(FArchive& Ar) override;
+	UE_API virtual bool Identical(const void* A, const void* B, uint32 PortFlags) const override;
+	UE_API virtual void SerializeItem(FStructuredArchive::FSlot Slot, void* Value, void const* Defaults) const override;
+	UE_API virtual bool NetSerializeItem(FArchive& Ar, UPackageMap* Map, void* Data, TArray<uint8> * MetaData = NULL) const override;
 protected:
-	virtual void ExportText_Internal(FString& ValueStr, const void* PropertyValueOrContainer, EPropertyPointerType PropertyPointerType, const void* DefaultValue, UObject* Parent, int32 PortFlags, UObject* ExportRootScope) const override;
-	virtual const TCHAR* ImportText_Internal(const TCHAR* Buffer, void* ContainerOrPropertyPtr, EPropertyPointerType PropertyPointerType, UObject* OwnerObject, int32 PortFlags, FOutputDevice* ErrorText) const override;
-	virtual bool ContainsClearOnFinishDestroyInternal(TArray<const FStructProperty*>& EncounteredStructProps) const override;
-	virtual void FinishDestroyInternal(void* Data) const override;
+	UE_API virtual void ExportText_Internal(FString& ValueStr, const void* PropertyValueOrContainer, EPropertyPointerType PropertyPointerType, const void* DefaultValue, UObject* Parent, int32 PortFlags, UObject* ExportRootScope) const override;
+	UE_API virtual const TCHAR* ImportText_Internal(const TCHAR* Buffer, void* ContainerOrPropertyPtr, EPropertyPointerType PropertyPointerType, UObject* OwnerObject, int32 PortFlags, FOutputDevice* ErrorText) const override;
+	UE_API virtual bool ContainsClearOnFinishDestroyInternal(TArray<const FStructProperty*>& EncounteredStructProps) const override;
+	UE_API virtual void FinishDestroyInternal(void* Data) const override;
 public:
 	virtual void InitializeValueInternal(void* Dest) const override
 	{
@@ -3637,24 +3639,24 @@ public:
 			}
 		}
 	}
-	virtual void CopyValuesInternal(void* Dest, void const* Src, int32 Count) const override;
-	virtual void ClearValueInternal(void* Data) const override;
-	virtual void DestroyValueInternal(void* Dest) const override;
-	virtual void InstanceSubobjects(void* Data, void const* DefaultData, TNotNull<UObject*> Owner, struct FObjectInstancingGraph* InstanceGraph) override;
-	virtual bool ContainsObjectReference(TArray<const FStructProperty*>& EncounteredStructProps, EPropertyObjectReferenceType InReferenceType = EPropertyObjectReferenceType::Strong) const override;
-	virtual void EmitReferenceInfo(UE::GC::FSchemaBuilder& Schema, int32 BaseOffset, TArray<const FStructProperty*>& EncounteredStructProps, UE::GC::FPropertyStack& DebugPath) override;
-	virtual bool SameType(const FProperty* Other) const override;
-	virtual EConvertFromTypeResult ConvertFromType(const FPropertyTag& Tag, FStructuredArchive::FSlot Slot, uint8* Data, UStruct* DefaultsStruct, const uint8* Defaults) override;
-	virtual void* GetValueAddressAtIndex_Direct(const FProperty* Inner, void* InValueAddress, int32 LogicalIndex) const override;
-	virtual bool UseBinaryOrNativeSerialization(const FArchive& Ar) const override;
-	virtual bool LoadTypeName(UE::FPropertyTypeName Type, const FPropertyTag* Tag = nullptr) override;
-	virtual void SaveTypeName(UE::FPropertyTypeNameBuilder& Type) const override;
-	virtual bool CanSerializeFromTypeName(UE::FPropertyTypeName Type) const override;
-	virtual EPropertyVisitorControlFlow Visit(FPropertyVisitorContext& Context, const TFunctionRef<EPropertyVisitorControlFlow(const FPropertyVisitorContext& /*Conmtext*/)> InFunc) const override;
-	virtual void* ResolveVisitedPathInfo(void* Data, const FPropertyVisitorInfo& Info) const override;
+	UE_API virtual void CopyValuesInternal(void* Dest, void const* Src, int32 Count) const override;
+	UE_API virtual void ClearValueInternal(void* Data) const override;
+	UE_API virtual void DestroyValueInternal(void* Dest) const override;
+	UE_API virtual void InstanceSubobjects(void* Data, void const* DefaultData, TNotNull<UObject*> Owner, struct FObjectInstancingGraph* InstanceGraph) override;
+	UE_API virtual bool ContainsObjectReference(TArray<const FStructProperty*>& EncounteredStructProps, EPropertyObjectReferenceType InReferenceType = EPropertyObjectReferenceType::Strong) const override;
+	UE_API virtual void EmitReferenceInfo(UE::GC::FSchemaBuilder& Schema, int32 BaseOffset, TArray<const FStructProperty*>& EncounteredStructProps, UE::GC::FPropertyStack& DebugPath) override;
+	UE_API virtual bool SameType(const FProperty* Other) const override;
+	UE_API virtual EConvertFromTypeResult ConvertFromType(const FPropertyTag& Tag, FStructuredArchive::FSlot Slot, uint8* Data, UStruct* DefaultsStruct, const uint8* Defaults) override;
+	UE_API virtual void* GetValueAddressAtIndex_Direct(const FProperty* Inner, void* InValueAddress, int32 LogicalIndex) const override;
+	UE_API virtual bool UseBinaryOrNativeSerialization(const FArchive& Ar) const override;
+	UE_API virtual bool LoadTypeName(UE::FPropertyTypeName Type, const FPropertyTag* Tag = nullptr) override;
+	UE_API virtual void SaveTypeName(UE::FPropertyTypeNameBuilder& Type) const override;
+	UE_API virtual bool CanSerializeFromTypeName(UE::FPropertyTypeName Type) const override;
+	UE_API virtual EPropertyVisitorControlFlow Visit(FPropertyVisitorContext& Context, const TFunctionRef<EPropertyVisitorControlFlow(const FPropertyVisitorContext& /*Conmtext*/)> InFunc) const override;
+	UE_API virtual void* ResolveVisitedPathInfo(void* Data, const FPropertyVisitorInfo& Info) const override;
 	// End of FProperty interface
 
-	FString GetCPPTypeCustom(FString* ExtendedTypeText, uint32 CPPExportFlags, const FString& KeyTypeText, const FString& InKeyExtendedTypeText, const FString& ValueTypeText, const FString& InValueExtendedTypeText) const;
+	UE_API FString GetCPPTypeCustom(FString* ExtendedTypeText, uint32 CPPExportFlags, const FString& KeyTypeText, const FString& InKeyExtendedTypeText, const FString& ValueTypeText, const FString& InValueExtendedTypeText) const;
 
 	/*
 	 * Helper function to get the number of key/value pairs inside of a map. 
@@ -3702,19 +3704,19 @@ public:
 		return ValueProp;
 	}
 #if WITH_EDITORONLY_DATA
-	virtual void AppendSchemaHash(FBlake3& Builder, bool bSkipEditorOnly) const override;
+	UE_API virtual void AppendSchemaHash(FBlake3& Builder, bool bSkipEditorOnly) const override;
 #endif
 
-	virtual bool HasIntrusiveUnsetOptionalState() const override;
-	virtual void InitializeIntrusiveUnsetOptionalValue(void* Data) const override;
-	virtual bool IsIntrusiveOptionalValueSet(const void* Data) const override;
-	virtual void ClearIntrusiveOptionalValue(void* Data) const override;
-	virtual void EmitIntrusiveOptionalReferenceInfo(UE::GC::FSchemaBuilder& Schema, int32 BaseOffset, TArray<const FStructProperty*>& EncounteredStructProps, UE::GC::FPropertyStack& DebugPath) override;
+	UE_API virtual bool HasIntrusiveUnsetOptionalState() const override;
+	UE_API virtual void InitializeIntrusiveUnsetOptionalValue(void* Data) const override;
+	UE_API virtual bool IsIntrusiveOptionalValueSet(const void* Data) const override;
+	UE_API virtual void ClearIntrusiveOptionalValue(void* Data) const override;
+	UE_API virtual void EmitIntrusiveOptionalReferenceInfo(UE::GC::FSchemaBuilder& Schema, int32 BaseOffset, TArray<const FStructProperty*>& EncounteredStructProps, UE::GC::FPropertyStack& DebugPath) override;
 };
 
-class COREUOBJECT_API FSetProperty : public TProperty<FScriptSet, FProperty>
+class FSetProperty : public TProperty<FScriptSet, FProperty>
 {
-	DECLARE_FIELD(FSetProperty, (TProperty<FScriptSet, FProperty>), CASTCLASS_FSetProperty)
+	DECLARE_FIELD_API(FSetProperty, (TProperty<FScriptSet, FProperty>), CASTCLASS_FSetProperty, UE_API)
 
 	// Properties representing the key type and value type of the contained pairs
 	FProperty*       ElementProp;
@@ -3724,67 +3726,67 @@ public:
 	using TTypeFundamentals = Super::TTypeFundamentals;
 	using TCppType = TTypeFundamentals::TCppType;
 
-	FSetProperty(FFieldVariant InOwner, const FName& InName, EObjectFlags InObjectFlags);
+	UE_API FSetProperty(FFieldVariant InOwner, const FName& InName, EObjectFlags InObjectFlags);
 
 	/**
 	 * Constructor used for constructing compiled in properties
 	 * @param InOwner Owner of the property
 	 * @param PropBase Pointer to the compiled in structure describing the property
 	 **/
-	FSetProperty(FFieldVariant InOwner, const UECodeGen_Private::FSetPropertyParams& Prop);
+	UE_API FSetProperty(FFieldVariant InOwner, const UECodeGen_Private::FSetPropertyParams& Prop);
 
 #if WITH_EDITORONLY_DATA
-	explicit FSetProperty(UField* InField);
+	UE_API explicit FSetProperty(UField* InField);
 #endif // WITH_EDITORONLY_DATA
 
-	virtual ~FSetProperty();
+	UE_API virtual ~FSetProperty();
 
 	// UObject interface
-	virtual void Serialize(FArchive& Ar) override;
-	virtual void AddReferencedObjects(FReferenceCollector& Collector) override;
-	virtual void GetPreloadDependencies(TArray<UObject*>& OutDeps) override;
+	UE_API virtual void Serialize(FArchive& Ar) override;
+	UE_API virtual void AddReferencedObjects(FReferenceCollector& Collector) override;
+	UE_API virtual void GetPreloadDependencies(TArray<UObject*>& OutDeps) override;
 	// End of UObject interface
 
 	// Field interface
-	virtual void PostDuplicate(const FField& InField) override;
-	virtual FField* GetInnerFieldByName(const FName& InName) override;
-	virtual void GetInnerFields(TArray<FField*>& OutFields) override;
+	UE_API virtual void PostDuplicate(const FField& InField) override;
+	UE_API virtual FField* GetInnerFieldByName(const FName& InName) override;
+	UE_API virtual void GetInnerFields(TArray<FField*>& OutFields) override;
 
 	// UField interface
-	virtual void AddCppProperty(FProperty* Property) override;
+	UE_API virtual void AddCppProperty(FProperty* Property) override;
 	// End of UField interface
 
 	// FProperty interface
-	virtual FString GetCPPMacroType(FString& ExtendedTypeText) const  override;
-	virtual FString GetCPPType(FString* ExtendedTypeText, uint32 CPPExportFlags) const override;
-	virtual void LinkInternal(FArchive& Ar) override;
-	virtual bool Identical(const void* A, const void* B, uint32 PortFlags) const override;
-	virtual void SerializeItem(FStructuredArchive::FSlot Slot, void* Value, void const* Defaults) const override;
-	virtual bool NetSerializeItem(FArchive& Ar, UPackageMap* Map, void* Data, TArray<uint8> * MetaData = NULL) const override;
+	UE_API virtual FString GetCPPMacroType(FString& ExtendedTypeText) const  override;
+	UE_API virtual FString GetCPPType(FString* ExtendedTypeText, uint32 CPPExportFlags) const override;
+	UE_API virtual void LinkInternal(FArchive& Ar) override;
+	UE_API virtual bool Identical(const void* A, const void* B, uint32 PortFlags) const override;
+	UE_API virtual void SerializeItem(FStructuredArchive::FSlot Slot, void* Value, void const* Defaults) const override;
+	UE_API virtual bool NetSerializeItem(FArchive& Ar, UPackageMap* Map, void* Data, TArray<uint8> * MetaData = NULL) const override;
 protected:
-	virtual void ExportText_Internal(FString& ValueStr, const void* PropertyValueOrContainer, EPropertyPointerType PropertyPointerType, const void* DefaultValue, UObject* Parent, int32 PortFlags, UObject* ExportRootScope) const override;
-	virtual const TCHAR* ImportText_Internal(const TCHAR* Buffer, void* ContainerOrPropertyPtr, EPropertyPointerType PropertyPointerType, UObject* OwnerObject, int32 PortFlags, FOutputDevice* ErrorText) const override;
-	virtual bool ContainsClearOnFinishDestroyInternal(TArray<const FStructProperty*>& EncounteredStructProps) const override;
-	virtual void FinishDestroyInternal(void* Data) const override;
+	UE_API virtual void ExportText_Internal(FString& ValueStr, const void* PropertyValueOrContainer, EPropertyPointerType PropertyPointerType, const void* DefaultValue, UObject* Parent, int32 PortFlags, UObject* ExportRootScope) const override;
+	UE_API virtual const TCHAR* ImportText_Internal(const TCHAR* Buffer, void* ContainerOrPropertyPtr, EPropertyPointerType PropertyPointerType, UObject* OwnerObject, int32 PortFlags, FOutputDevice* ErrorText) const override;
+	UE_API virtual bool ContainsClearOnFinishDestroyInternal(TArray<const FStructProperty*>& EncounteredStructProps) const override;
+	UE_API virtual void FinishDestroyInternal(void* Data) const override;
 public:
-	virtual void CopyValuesInternal(void* Dest, void const* Src, int32 Count) const override;
-	virtual void ClearValueInternal(void* Data) const override;
-	virtual void DestroyValueInternal(void* Dest) const override;
-	virtual void InstanceSubobjects(void* Data, void const* DefaultData, TNotNull<UObject*> Owner, struct FObjectInstancingGraph* InstanceGraph) override;
-	virtual bool ContainsObjectReference(TArray<const FStructProperty*>& EncounteredStructProps, EPropertyObjectReferenceType InReferenceType = EPropertyObjectReferenceType::Strong) const override;
-	virtual void EmitReferenceInfo(UE::GC::FSchemaBuilder& Schema, int32 BaseOffset, TArray<const FStructProperty*>& EncounteredStructProps, UE::GC::FPropertyStack& DebugPath) override;
-	virtual bool SameType(const FProperty* Other) const override;
-	virtual EConvertFromTypeResult ConvertFromType(const FPropertyTag& Tag, FStructuredArchive::FSlot Slot, uint8* Data, UStruct* DefaultsStruct, const uint8* Defaults) override;
-	virtual void* GetValueAddressAtIndex_Direct(const FProperty* Inner, void* InValueAddress, int32 LogicalIndex) const override;
-	virtual bool UseBinaryOrNativeSerialization(const FArchive& Ar) const override;
-	virtual bool LoadTypeName(UE::FPropertyTypeName Type, const FPropertyTag* Tag = nullptr) override;
-	virtual void SaveTypeName(UE::FPropertyTypeNameBuilder& Type) const override;
-	virtual bool CanSerializeFromTypeName(UE::FPropertyTypeName Type) const override;
-	virtual EPropertyVisitorControlFlow Visit(FPropertyVisitorContext& Context, const TFunctionRef<EPropertyVisitorControlFlow(const FPropertyVisitorContext& /*Conmtext*/)> InFunc) const override;
-	virtual void* ResolveVisitedPathInfo(void* Data, const FPropertyVisitorInfo& Info) const override;
+	UE_API virtual void CopyValuesInternal(void* Dest, void const* Src, int32 Count) const override;
+	UE_API virtual void ClearValueInternal(void* Data) const override;
+	UE_API virtual void DestroyValueInternal(void* Dest) const override;
+	UE_API virtual void InstanceSubobjects(void* Data, void const* DefaultData, TNotNull<UObject*> Owner, struct FObjectInstancingGraph* InstanceGraph) override;
+	UE_API virtual bool ContainsObjectReference(TArray<const FStructProperty*>& EncounteredStructProps, EPropertyObjectReferenceType InReferenceType = EPropertyObjectReferenceType::Strong) const override;
+	UE_API virtual void EmitReferenceInfo(UE::GC::FSchemaBuilder& Schema, int32 BaseOffset, TArray<const FStructProperty*>& EncounteredStructProps, UE::GC::FPropertyStack& DebugPath) override;
+	UE_API virtual bool SameType(const FProperty* Other) const override;
+	UE_API virtual EConvertFromTypeResult ConvertFromType(const FPropertyTag& Tag, FStructuredArchive::FSlot Slot, uint8* Data, UStruct* DefaultsStruct, const uint8* Defaults) override;
+	UE_API virtual void* GetValueAddressAtIndex_Direct(const FProperty* Inner, void* InValueAddress, int32 LogicalIndex) const override;
+	UE_API virtual bool UseBinaryOrNativeSerialization(const FArchive& Ar) const override;
+	UE_API virtual bool LoadTypeName(UE::FPropertyTypeName Type, const FPropertyTag* Tag = nullptr) override;
+	UE_API virtual void SaveTypeName(UE::FPropertyTypeNameBuilder& Type) const override;
+	UE_API virtual bool CanSerializeFromTypeName(UE::FPropertyTypeName Type) const override;
+	UE_API virtual EPropertyVisitorControlFlow Visit(FPropertyVisitorContext& Context, const TFunctionRef<EPropertyVisitorControlFlow(const FPropertyVisitorContext& /*Conmtext*/)> InFunc) const override;
+	UE_API virtual void* ResolveVisitedPathInfo(void* Data, const FPropertyVisitorInfo& Info) const override;
 	// End of FProperty interface
 
-	FString GetCPPTypeCustom(FString* ExtendedTypeText, uint32 CPPExportFlags, const FString& ElementTypeText, const FString& InElementExtendedTypeText) const;
+	UE_API FString GetCPPTypeCustom(FString* ExtendedTypeText, uint32 CPPExportFlags, const FString& ElementTypeText, const FString& InElementExtendedTypeText) const;
 
 	/*
 	 * Helper function to get the number of elements inside of a set.
@@ -3831,14 +3833,14 @@ public:
 	}
 
 #if WITH_EDITORONLY_DATA
-	virtual void AppendSchemaHash(FBlake3& Builder, bool bSkipEditorOnly) const override;
+	UE_API virtual void AppendSchemaHash(FBlake3& Builder, bool bSkipEditorOnly) const override;
 #endif
 
-	virtual bool HasIntrusiveUnsetOptionalState() const override;
-	virtual void InitializeIntrusiveUnsetOptionalValue(void* Data) const override;
-	virtual bool IsIntrusiveOptionalValueSet(const void* Data) const override;
-	virtual void ClearIntrusiveOptionalValue(void* Data) const override;
-	virtual void EmitIntrusiveOptionalReferenceInfo(UE::GC::FSchemaBuilder& Schema, int32 BaseOffset, TArray<const FStructProperty*>& EncounteredStructProps, UE::GC::FPropertyStack& DebugPath) override;
+	UE_API virtual bool HasIntrusiveUnsetOptionalState() const override;
+	UE_API virtual void InitializeIntrusiveUnsetOptionalValue(void* Data) const override;
+	UE_API virtual bool IsIntrusiveOptionalValueSet(const void* Data) const override;
+	UE_API virtual void ClearIntrusiveOptionalValue(void* Data) const override;
+	UE_API virtual void EmitIntrusiveOptionalReferenceInfo(UE::GC::FSchemaBuilder& Schema, int32 BaseOffset, TArray<const FStructProperty*>& EncounteredStructProps, UE::GC::FPropertyStack& DebugPath) override;
 };
 
 /**
@@ -5922,80 +5924,80 @@ public:
 // Describes a structure variable embedded in (as opposed to referenced by) 
 // an object.
 //
-class COREUOBJECT_API FStructProperty : public FProperty
+class FStructProperty : public FProperty
 {
-	DECLARE_FIELD(FStructProperty, FProperty, CASTCLASS_FStructProperty)
+	DECLARE_FIELD_API(FStructProperty, FProperty, CASTCLASS_FStructProperty, UE_API)
 
 	// Variables.
 	TObjectPtr<class UScriptStruct> Struct;
 public:
-	FStructProperty(FFieldVariant InOwner, const FName& InName, EObjectFlags InObjectFlags);
+	UE_API FStructProperty(FFieldVariant InOwner, const FName& InName, EObjectFlags InObjectFlags);
 
 	/**
 	 * Constructor used for constructing compiled in properties
 	 * @param InOwner Owner of the property
 	 * @param PropBase Pointer to the compiled in structure describing the property
 	 **/
-	FStructProperty(FFieldVariant InOwner, const UECodeGen_Private::FStructPropertyParams& Prop);
+	UE_API FStructProperty(FFieldVariant InOwner, const UECodeGen_Private::FStructPropertyParams& Prop);
 
 #if WITH_EDITORONLY_DATA
-	explicit FStructProperty(UField* InField);
+	UE_API explicit FStructProperty(UField* InField);
 #endif // WITH_EDITORONLY_DATA
 
 	// UObject interface
-	virtual void Serialize( FArchive& Ar ) override;
-	virtual void AddReferencedObjects(FReferenceCollector& Collector) override;
-	virtual void GetPreloadDependencies(TArray<UObject*>& OutDeps) override;
+	UE_API virtual void Serialize( FArchive& Ar ) override;
+	UE_API virtual void AddReferencedObjects(FReferenceCollector& Collector) override;
+	UE_API virtual void GetPreloadDependencies(TArray<UObject*>& OutDeps) override;
 	// End of UObject interface
 
 	// Field interface
-	virtual void PostDuplicate(const FField& InField) override;
+	UE_API virtual void PostDuplicate(const FField& InField) override;
 
 	// FProperty interface
-	virtual FString GetCPPMacroType( FString& ExtendedTypeText ) const  override;
-	virtual FString GetCPPType( FString* ExtendedTypeText, uint32 CPPExportFlags ) const override;
-	virtual void LinkInternal(FArchive& Ar) override;
-	virtual bool Identical( const void* A, const void* B, uint32 PortFlags ) const override;
-	virtual void SerializeItem(FStructuredArchive::FSlot Slot, void* Value, void const* Defaults) const override;
-	virtual bool NetSerializeItem( FArchive& Ar, UPackageMap* Map, void* Data, TArray<uint8> * MetaData = NULL ) const override;
-	virtual bool SupportsNetSharedSerialization() const override;
+	UE_API virtual FString GetCPPMacroType( FString& ExtendedTypeText ) const  override;
+	UE_API virtual FString GetCPPType( FString* ExtendedTypeText, uint32 CPPExportFlags ) const override;
+	UE_API virtual void LinkInternal(FArchive& Ar) override;
+	UE_API virtual bool Identical( const void* A, const void* B, uint32 PortFlags ) const override;
+	UE_API virtual void SerializeItem(FStructuredArchive::FSlot Slot, void* Value, void const* Defaults) const override;
+	UE_API virtual bool NetSerializeItem( FArchive& Ar, UPackageMap* Map, void* Data, TArray<uint8> * MetaData = NULL ) const override;
+	UE_API virtual bool SupportsNetSharedSerialization() const override;
 protected:
-	virtual void ExportText_Internal( FString& ValueStr, const void* PropertyValueOrContainer, EPropertyPointerType PropertyPointerType, const void* DefaultValue, UObject* Parent, int32 PortFlags, UObject* ExportRootScope ) const override;
-	virtual const TCHAR* ImportText_Internal(const TCHAR* Buffer, void* ContainerOrPropertyPtr, EPropertyPointerType PropertyPointerType, UObject* OwnerObject, int32 PortFlags, FOutputDevice* ErrorText) const override;
-	virtual bool ContainsClearOnFinishDestroyInternal(TArray<const FStructProperty*>& EncounteredStructProps) const override;
-	virtual void FinishDestroyInternal(void* Data) const override;
+	UE_API virtual void ExportText_Internal( FString& ValueStr, const void* PropertyValueOrContainer, EPropertyPointerType PropertyPointerType, const void* DefaultValue, UObject* Parent, int32 PortFlags, UObject* ExportRootScope ) const override;
+	UE_API virtual const TCHAR* ImportText_Internal(const TCHAR* Buffer, void* ContainerOrPropertyPtr, EPropertyPointerType PropertyPointerType, UObject* OwnerObject, int32 PortFlags, FOutputDevice* ErrorText) const override;
+	UE_API virtual bool ContainsClearOnFinishDestroyInternal(TArray<const FStructProperty*>& EncounteredStructProps) const override;
+	UE_API virtual void FinishDestroyInternal(void* Data) const override;
 public:
-	virtual void CopyValuesInternal( void* Dest, void const* Src, int32 Count  ) const override;
-	virtual void ClearValueInternal( void* Data ) const override;
-	virtual void DestroyValueInternal( void* Dest ) const override;
-	virtual void InitializeValueInternal( void* Dest ) const override;
-	virtual void InstanceSubobjects( void* Data, void const* DefaultData, TNotNull<UObject*> Owner, struct FObjectInstancingGraph* InstanceGraph ) override;
-	virtual int32 GetMinAlignment() const override;
-	virtual bool ContainsObjectReference(TArray<const FStructProperty*>& EncounteredStructProps, EPropertyObjectReferenceType InReferenceType = EPropertyObjectReferenceType::Strong) const override;
-	virtual void EmitReferenceInfo(UE::GC::FSchemaBuilder& Schema, int32 BaseOffset, TArray<const FStructProperty*>& EncounteredStructProps, UE::GC::FPropertyStack& DebugPath) override;
-	virtual bool SameType(const FProperty* Other) const override;
-	virtual EConvertFromTypeResult ConvertFromType(const FPropertyTag& Tag, FStructuredArchive::FSlot Slot, uint8* Data, UStruct* DefaultsStruct, const uint8* Defaults) override;
+	UE_API virtual void CopyValuesInternal( void* Dest, void const* Src, int32 Count  ) const override;
+	UE_API virtual void ClearValueInternal( void* Data ) const override;
+	UE_API virtual void DestroyValueInternal( void* Dest ) const override;
+	UE_API virtual void InitializeValueInternal( void* Dest ) const override;
+	UE_API virtual void InstanceSubobjects( void* Data, void const* DefaultData, TNotNull<UObject*> Owner, struct FObjectInstancingGraph* InstanceGraph ) override;
+	UE_API virtual int32 GetMinAlignment() const override;
+	UE_API virtual bool ContainsObjectReference(TArray<const FStructProperty*>& EncounteredStructProps, EPropertyObjectReferenceType InReferenceType = EPropertyObjectReferenceType::Strong) const override;
+	UE_API virtual void EmitReferenceInfo(UE::GC::FSchemaBuilder& Schema, int32 BaseOffset, TArray<const FStructProperty*>& EncounteredStructProps, UE::GC::FPropertyStack& DebugPath) override;
+	UE_API virtual bool SameType(const FProperty* Other) const override;
+	UE_API virtual EConvertFromTypeResult ConvertFromType(const FPropertyTag& Tag, FStructuredArchive::FSlot Slot, uint8* Data, UStruct* DefaultsStruct, const uint8* Defaults) override;
 #if WITH_EDITORONLY_DATA
-	virtual void AppendSchemaHash(FBlake3& Builder, bool bSkipEditorOnly) const override;
+	UE_API virtual void AppendSchemaHash(FBlake3& Builder, bool bSkipEditorOnly) const override;
 #endif
-	virtual bool UseBinaryOrNativeSerialization(const FArchive& Ar) const override;
-	virtual bool LoadTypeName(UE::FPropertyTypeName Type, const FPropertyTag* Tag = nullptr) override;
-	virtual void SaveTypeName(UE::FPropertyTypeNameBuilder& Type) const override;
-	virtual bool CanSerializeFromTypeName(UE::FPropertyTypeName Type) const override;
-	virtual EPropertyVisitorControlFlow Visit(FPropertyVisitorContext& Context, const TFunctionRef<EPropertyVisitorControlFlow(const FPropertyVisitorContext& /*Conmtext*/)> InFunc) const override;
-	virtual void* ResolveVisitedPathInfo(void* Data, const FPropertyVisitorInfo& Info) const override;
+	UE_API virtual bool UseBinaryOrNativeSerialization(const FArchive& Ar) const override;
+	UE_API virtual bool LoadTypeName(UE::FPropertyTypeName Type, const FPropertyTag* Tag = nullptr) override;
+	UE_API virtual void SaveTypeName(UE::FPropertyTypeNameBuilder& Type) const override;
+	UE_API virtual bool CanSerializeFromTypeName(UE::FPropertyTypeName Type) const override;
+	UE_API virtual EPropertyVisitorControlFlow Visit(FPropertyVisitorContext& Context, const TFunctionRef<EPropertyVisitorControlFlow(const FPropertyVisitorContext& /*Conmtext*/)> InFunc) const override;
+	UE_API virtual void* ResolveVisitedPathInfo(void* Data, const FPropertyVisitorInfo& Info) const override;
 	// End of FProperty interface
 
-	bool FindInnerPropertyInstance(FName PropertyName, const void* Data, const FProperty*& OutProp, const void*& OutData) const;
+	UE_API bool FindInnerPropertyInstance(FName PropertyName, const void* Data, const FProperty*& OutProp, const void*& OutData) const;
 
-	virtual bool HasIntrusiveUnsetOptionalState() const override;
-	virtual void InitializeIntrusiveUnsetOptionalValue(void* Data) const override;
-	virtual bool IsIntrusiveOptionalValueSet(const void* Data) const override;
-	virtual void ClearIntrusiveOptionalValue(void* Data) const override;
-	virtual void EmitIntrusiveOptionalReferenceInfo(UE::GC::FSchemaBuilder& Schema, int32 BaseOffset, TArray<const FStructProperty*>& EncounteredStructProps, UE::GC::FPropertyStack& DebugPath) override;
+	UE_API virtual bool HasIntrusiveUnsetOptionalState() const override;
+	UE_API virtual void InitializeIntrusiveUnsetOptionalValue(void* Data) const override;
+	UE_API virtual bool IsIntrusiveOptionalValueSet(const void* Data) const override;
+	UE_API virtual void ClearIntrusiveOptionalValue(void* Data) const override;
+	UE_API virtual void EmitIntrusiveOptionalReferenceInfo(UE::GC::FSchemaBuilder& Schema, int32 BaseOffset, TArray<const FStructProperty*>& EncounteredStructProps, UE::GC::FPropertyStack& DebugPath) override;
 
 private:
-	virtual uint32 GetValueTypeHashInternal(const void* Src) const;
+	UE_API virtual uint32 GetValueTypeHashInternal(const void* Src) const;
 };
 
 /*-----------------------------------------------------------------------------
@@ -6005,9 +6007,9 @@ private:
 /**
  * Describes a pointer to a function bound to an Object.
  */
-class COREUOBJECT_API FDelegateProperty : public TProperty<FScriptDelegate, FProperty>
+class FDelegateProperty : public TProperty<FScriptDelegate, FProperty>
 {
-	DECLARE_FIELD(FDelegateProperty, (TProperty<FScriptDelegate, FProperty>), CASTCLASS_FDelegateProperty)
+	DECLARE_FIELD_API(FDelegateProperty, (TProperty<FScriptDelegate, FProperty>), CASTCLASS_FDelegateProperty, UE_API)
 
 	/** Points to the source delegate function (the function declared with the delegate keyword) used in the declaration of this delegate property. */
 	TObjectPtr<UFunction> SignatureFunction;
@@ -6016,42 +6018,42 @@ public:
 	using TTypeFundamentals = Super::TTypeFundamentals;
 	using TCppType = TTypeFundamentals::TCppType;
 
-	FDelegateProperty(FFieldVariant InOwner, const FName& InName, EObjectFlags InObjectFlags);
+	UE_API FDelegateProperty(FFieldVariant InOwner, const FName& InName, EObjectFlags InObjectFlags);
 
 	/**
 	 * Constructor used for constructing compiled in properties
 	 * @param InOwner Owner of the property
 	 * @param PropBase Pointer to the compiled in structure describing the property
 	 **/
-	FDelegateProperty(FFieldVariant InOwner, const UECodeGen_Private::FDelegatePropertyParams& Prop);
+	UE_API FDelegateProperty(FFieldVariant InOwner, const UECodeGen_Private::FDelegatePropertyParams& Prop);
 
 #if WITH_EDITORONLY_DATA
-	explicit FDelegateProperty(UField* InField);
+	UE_API explicit FDelegateProperty(UField* InField);
 #endif // WITH_EDITORONLY_DATA
 
 	// UObject interface
-	virtual void Serialize( FArchive& Ar ) override;
-	virtual void AddReferencedObjects(FReferenceCollector& Collector) override;
-	virtual void BeginDestroy() override;
+	UE_API virtual void Serialize( FArchive& Ar ) override;
+	UE_API virtual void AddReferencedObjects(FReferenceCollector& Collector) override;
+	UE_API virtual void BeginDestroy() override;
 	// End of UObject interface
 
 	// Field interface
-	virtual void PostDuplicate(const FField& InField) override;
+	UE_API virtual void PostDuplicate(const FField& InField) override;
 
 	// FProperty interface
-	virtual FString GetCPPType( FString* ExtendedTypeText, uint32 CPPExportFlags ) const override;
-	virtual bool Identical( const void* A, const void* B, uint32 PortFlags ) const override;
-	virtual void SerializeItem(FStructuredArchive::FSlot Slot, void* Value, void const* Defaults) const override;
-	virtual bool NetSerializeItem( FArchive& Ar, UPackageMap* Map, void* Data, TArray<uint8> * MetaData = NULL ) const override;
+	UE_API virtual FString GetCPPType( FString* ExtendedTypeText, uint32 CPPExportFlags ) const override;
+	UE_API virtual bool Identical( const void* A, const void* B, uint32 PortFlags ) const override;
+	UE_API virtual void SerializeItem(FStructuredArchive::FSlot Slot, void* Value, void const* Defaults) const override;
+	UE_API virtual bool NetSerializeItem( FArchive& Ar, UPackageMap* Map, void* Data, TArray<uint8> * MetaData = NULL ) const override;
 protected:
-	virtual void ExportText_Internal( FString& ValueStr, const void* PropertyValueOrContainer, EPropertyPointerType PropertyPointerType, const void* DefaultValue, UObject* Parent, int32 PortFlags, UObject* ExportRootScope ) const override;
-	virtual const TCHAR* ImportText_Internal(const TCHAR* Buffer, void* ContainerOrPropertyPtr, EPropertyPointerType PropertyPointerType, UObject* OwnerObject, int32 PortFlags, FOutputDevice* ErrorText) const override;
+	UE_API virtual void ExportText_Internal( FString& ValueStr, const void* PropertyValueOrContainer, EPropertyPointerType PropertyPointerType, const void* DefaultValue, UObject* Parent, int32 PortFlags, UObject* ExportRootScope ) const override;
+	UE_API virtual const TCHAR* ImportText_Internal(const TCHAR* Buffer, void* ContainerOrPropertyPtr, EPropertyPointerType PropertyPointerType, UObject* OwnerObject, int32 PortFlags, FOutputDevice* ErrorText) const override;
 public:
-	virtual bool ContainsObjectReference(TArray<const FStructProperty*>& EncounteredStructProps, EPropertyObjectReferenceType InReferenceType = EPropertyObjectReferenceType::Strong) const override;
-	virtual void InstanceSubobjects( void* Data, void const* DefaultData, TNotNull<UObject*> Owner, struct FObjectInstancingGraph* InstanceGraph ) override;
-	virtual bool SameType(const FProperty* Other) const override;
+	UE_API virtual bool ContainsObjectReference(TArray<const FStructProperty*>& EncounteredStructProps, EPropertyObjectReferenceType InReferenceType = EPropertyObjectReferenceType::Strong) const override;
+	UE_API virtual void InstanceSubobjects( void* Data, void const* DefaultData, TNotNull<UObject*> Owner, struct FObjectInstancingGraph* InstanceGraph ) override;
+	UE_API virtual bool SameType(const FProperty* Other) const override;
 #if WITH_EDITORONLY_DATA
-	virtual void AppendSchemaHash(FBlake3& Builder, bool bSkipEditorOnly) const override;
+	UE_API virtual void AppendSchemaHash(FBlake3& Builder, bool bSkipEditorOnly) const override;
 #endif
 	// End of FProperty interface
 };
@@ -6064,72 +6066,72 @@ public:
 /**
  * Describes a list of functions bound to an Object.
  */
-class COREUOBJECT_API FMulticastDelegateProperty : public FProperty
+class FMulticastDelegateProperty : public FProperty
 {
-	DECLARE_FIELD(FMulticastDelegateProperty, FProperty, CASTCLASS_FMulticastDelegateProperty)
+	DECLARE_FIELD_API(FMulticastDelegateProperty, FProperty, CASTCLASS_FMulticastDelegateProperty, UE_API)
 
 	/** Points to the source delegate function (the function declared with the delegate keyword) used in the declaration of this delegate property. */
 	TObjectPtr<UFunction> SignatureFunction;
 
 public:
 
-	FMulticastDelegateProperty(FFieldVariant InOwner, const FName& InName, EObjectFlags InObjectFlags);
+	UE_API FMulticastDelegateProperty(FFieldVariant InOwner, const FName& InName, EObjectFlags InObjectFlags);
 
 	/**
 	 * Constructor used for constructing compiled in properties
 	 * @param InOwner Owner of the property
 	 * @param PropBase Pointer to the compiled in structure describing the property
 	 **/
-	FMulticastDelegateProperty(FFieldVariant InOwner, const UECodeGen_Private::FMulticastDelegatePropertyParams& Prop, EPropertyFlags AdditionalPropertyFlags = CPF_None);
+	UE_API FMulticastDelegateProperty(FFieldVariant InOwner, const UECodeGen_Private::FMulticastDelegatePropertyParams& Prop, EPropertyFlags AdditionalPropertyFlags = CPF_None);
 
 #if WITH_EDITORONLY_DATA
-	explicit FMulticastDelegateProperty(UField* InField);
+	UE_API explicit FMulticastDelegateProperty(UField* InField);
 #endif // WITH_EDITORONLY_DATA
 
 	// UObject interface
-	virtual void Serialize( FArchive& Ar ) override;
-	virtual void BeginDestroy() override;
-	virtual void AddReferencedObjects(FReferenceCollector& Collector) override;
+	UE_API virtual void Serialize( FArchive& Ar ) override;
+	UE_API virtual void BeginDestroy() override;
+	UE_API virtual void AddReferencedObjects(FReferenceCollector& Collector) override;
 	// End of UObject interface
 
 	// Field interface
-	virtual void PostDuplicate(const FField& InField) override;
+	UE_API virtual void PostDuplicate(const FField& InField) override;
 
 	// FProperty interface
-	virtual FString GetCPPType( FString* ExtendedTypeText, uint32 CPPExportFlags ) const override;
-	virtual bool Identical( const void* A, const void* B, uint32 PortFlags ) const override;
-	virtual bool NetSerializeItem( FArchive& Ar, UPackageMap* Map, void* Data, TArray<uint8> * MetaData = NULL ) const override;
+	UE_API virtual FString GetCPPType( FString* ExtendedTypeText, uint32 CPPExportFlags ) const override;
+	UE_API virtual bool Identical( const void* A, const void* B, uint32 PortFlags ) const override;
+	UE_API virtual bool NetSerializeItem( FArchive& Ar, UPackageMap* Map, void* Data, TArray<uint8> * MetaData = NULL ) const override;
 protected:
-	virtual void ExportText_Internal( FString& ValueStr, const void* PropertyValueOrContainer, EPropertyPointerType PropertyPointerType, const void* DefaultValue, UObject* Parent, int32 PortFlags, UObject* ExportRootScope ) const override;
+	UE_API virtual void ExportText_Internal( FString& ValueStr, const void* PropertyValueOrContainer, EPropertyPointerType PropertyPointerType, const void* DefaultValue, UObject* Parent, int32 PortFlags, UObject* ExportRootScope ) const override;
 public:
-	virtual bool ContainsObjectReference(TArray<const FStructProperty*>& EncounteredStructProps, EPropertyObjectReferenceType InReferenceType = EPropertyObjectReferenceType::Strong) const override;
-	virtual void InstanceSubobjects( void* Data, void const* DefaultData, TNotNull<UObject*> Owner, struct FObjectInstancingGraph* InstanceGraph ) override;
-	virtual bool SameType(const FProperty* Other) const override;
-	virtual EConvertFromTypeResult ConvertFromType(const FPropertyTag& Tag, FStructuredArchive::FSlot Slot, uint8* Data, UStruct* DefaultsStruct, const uint8* Defaults) override;
+	UE_API virtual bool ContainsObjectReference(TArray<const FStructProperty*>& EncounteredStructProps, EPropertyObjectReferenceType InReferenceType = EPropertyObjectReferenceType::Strong) const override;
+	UE_API virtual void InstanceSubobjects( void* Data, void const* DefaultData, TNotNull<UObject*> Owner, struct FObjectInstancingGraph* InstanceGraph ) override;
+	UE_API virtual bool SameType(const FProperty* Other) const override;
+	UE_API virtual EConvertFromTypeResult ConvertFromType(const FPropertyTag& Tag, FStructuredArchive::FSlot Slot, uint8* Data, UStruct* DefaultsStruct, const uint8* Defaults) override;
 	// End of FProperty interface
 
-	virtual const FMulticastScriptDelegate* GetMulticastDelegate(const void* PropertyValue) const PURE_VIRTUAL(FMulticastDelegateProperty::GetMulticastDelegate, return nullptr;);
-	virtual void SetMulticastDelegate(void* PropertyValue, FMulticastScriptDelegate ScriptDelegate) const PURE_VIRTUAL(FMulticastDelegateProperty::SetMulticastDelegate, );
+	UE_API virtual const FMulticastScriptDelegate* GetMulticastDelegate(const void* PropertyValue) const PURE_VIRTUAL(FMulticastDelegateProperty::GetMulticastDelegate, return nullptr;);
+	UE_API virtual void SetMulticastDelegate(void* PropertyValue, FMulticastScriptDelegate ScriptDelegate) const PURE_VIRTUAL(FMulticastDelegateProperty::SetMulticastDelegate, );
 
-	virtual void AddDelegate(FScriptDelegate ScriptDelegate, UObject* Parent = nullptr, void* PropertyValue = nullptr) const PURE_VIRTUAL(FMulticastDelegateProperty::AddDelegate, );
-	virtual void RemoveDelegate(const FScriptDelegate& ScriptDelegate, UObject* Parent = nullptr, void* PropertyValue = nullptr) const PURE_VIRTUAL(FMulticastDelegateProperty::RemoveDelegate, );
-	virtual void ClearDelegate(UObject* Parent = nullptr, void* PropertyValue = nullptr)  const PURE_VIRTUAL(FMulticastDelegateProperty::ClearDelegate, );
+	UE_API virtual void AddDelegate(FScriptDelegate ScriptDelegate, UObject* Parent = nullptr, void* PropertyValue = nullptr) const PURE_VIRTUAL(FMulticastDelegateProperty::AddDelegate, );
+	UE_API virtual void RemoveDelegate(const FScriptDelegate& ScriptDelegate, UObject* Parent = nullptr, void* PropertyValue = nullptr) const PURE_VIRTUAL(FMulticastDelegateProperty::RemoveDelegate, );
+	UE_API virtual void ClearDelegate(UObject* Parent = nullptr, void* PropertyValue = nullptr)  const PURE_VIRTUAL(FMulticastDelegateProperty::ClearDelegate, );
 
 protected:
 	friend class FProperty;
 
-	static FMulticastScriptDelegate EmptyDelegate;
-	virtual FMulticastScriptDelegate& GetMulticastScriptDelegate(const void* PropertyValue, int32 Index) const PURE_VIRTUAL(FMulticastDelegateProperty::GetMulticastScriptDelegate, return EmptyDelegate;);
+	UE_API static FMulticastScriptDelegate EmptyDelegate;
+	UE_API virtual FMulticastScriptDelegate& GetMulticastScriptDelegate(const void* PropertyValue, int32 Index) const PURE_VIRTUAL(FMulticastDelegateProperty::GetMulticastScriptDelegate, return EmptyDelegate;);
 
 
-	const TCHAR* ImportText_Add( const TCHAR* Buffer, void* PropertyValue, int32 PortFlags, UObject* Parent, FOutputDevice* ErrorText ) const;
-	const TCHAR* ImportText_Remove( const TCHAR* Buffer, void* PropertyValue, int32 PortFlags, UObject* Parent, FOutputDevice* ErrorText ) const;
+	UE_API const TCHAR* ImportText_Add( const TCHAR* Buffer, void* PropertyValue, int32 PortFlags, UObject* Parent, FOutputDevice* ErrorText ) const;
+	UE_API const TCHAR* ImportText_Remove( const TCHAR* Buffer, void* PropertyValue, int32 PortFlags, UObject* Parent, FOutputDevice* ErrorText ) const;
 
-	const TCHAR* ImportDelegateFromText(FMulticastScriptDelegate& MulticastDelegate, const TCHAR* Buffer, UObject* OwnerObject, FOutputDevice* ErrorText) const;
+	UE_API const TCHAR* ImportDelegateFromText(FMulticastScriptDelegate& MulticastDelegate, const TCHAR* Buffer, UObject* OwnerObject, FOutputDevice* ErrorText) const;
 };
 
 template<class InTCppType>
-class COREUOBJECT_API TProperty_MulticastDelegate : public TProperty<InTCppType, FMulticastDelegateProperty>
+class TProperty_MulticastDelegate : public TProperty<InTCppType, FMulticastDelegateProperty>
 {
 public:
 	typedef TProperty<InTCppType, FMulticastDelegateProperty> Super;
@@ -6178,9 +6180,9 @@ public:
 	// End of FProperty interface
 };
 
-class COREUOBJECT_API FMulticastInlineDelegateProperty : public TProperty_MulticastDelegate<FMulticastScriptDelegate>
+class FMulticastInlineDelegateProperty : public TProperty_MulticastDelegate<FMulticastScriptDelegate>
 {
-	DECLARE_FIELD(FMulticastInlineDelegateProperty, TProperty_MulticastDelegate<FMulticastScriptDelegate>, CASTCLASS_FMulticastInlineDelegateProperty)
+	DECLARE_FIELD_API(FMulticastInlineDelegateProperty, TProperty_MulticastDelegate<FMulticastScriptDelegate>, CASTCLASS_FMulticastInlineDelegateProperty, UE_API)
 
 public:
 
@@ -6194,7 +6196,7 @@ public:
 	 * @param InOwner Owner of the property
 	 * @param PropBase Pointer to the compiled in structure describing the property
 	 **/
-	FMulticastInlineDelegateProperty(FFieldVariant InOwner, const UECodeGen_Private::FMulticastDelegatePropertyParams& Prop);
+	UE_API FMulticastInlineDelegateProperty(FFieldVariant InOwner, const UECodeGen_Private::FMulticastDelegatePropertyParams& Prop);
 
 #if WITH_EDITORONLY_DATA
 	explicit FMulticastInlineDelegateProperty(UField* InField)
@@ -6204,26 +6206,26 @@ public:
 #endif // WITH_EDITORONLY_DATA
 
 	// FProperty interface
-	virtual void SerializeItem(FStructuredArchive::FSlot Slot, void* Value, void const* Defaults) const override;
-	virtual const TCHAR* ImportText_Internal(const TCHAR* Buffer, void* ContainerOrPropertyPtr, EPropertyPointerType PropertyPointerType, UObject* OwnerObject, int32 PortFlags, FOutputDevice* ErrorText) const override;
+	UE_API virtual void SerializeItem(FStructuredArchive::FSlot Slot, void* Value, void const* Defaults) const override;
+	UE_API virtual const TCHAR* ImportText_Internal(const TCHAR* Buffer, void* ContainerOrPropertyPtr, EPropertyPointerType PropertyPointerType, UObject* OwnerObject, int32 PortFlags, FOutputDevice* ErrorText) const override;
 	// End of FProperty interface
 
 	// FMulticastDelegateProperty interface
-	virtual const FMulticastScriptDelegate* GetMulticastDelegate(const void* PropertyValue) const override;
-	virtual void SetMulticastDelegate(void* PropertyValue, FMulticastScriptDelegate ScriptDelegate) const override;
+	UE_API virtual const FMulticastScriptDelegate* GetMulticastDelegate(const void* PropertyValue) const override;
+	UE_API virtual void SetMulticastDelegate(void* PropertyValue, FMulticastScriptDelegate ScriptDelegate) const override;
 
-	virtual void AddDelegate(FScriptDelegate ScriptDelegate, UObject* Parent = nullptr, void* PropertyValue = nullptr) const override;
-	virtual void RemoveDelegate(const FScriptDelegate& ScriptDelegate, UObject* Parent = nullptr, void* PropertyValue = nullptr) const override;
-	virtual void ClearDelegate(UObject* Parent = nullptr, void* PropertyValue = nullptr) const override;
+	UE_API virtual void AddDelegate(FScriptDelegate ScriptDelegate, UObject* Parent = nullptr, void* PropertyValue = nullptr) const override;
+	UE_API virtual void RemoveDelegate(const FScriptDelegate& ScriptDelegate, UObject* Parent = nullptr, void* PropertyValue = nullptr) const override;
+	UE_API virtual void ClearDelegate(UObject* Parent = nullptr, void* PropertyValue = nullptr) const override;
 
 protected:
-	virtual FMulticastScriptDelegate& GetMulticastScriptDelegate(const void* PropertyValue, int32 Index) const;
+	UE_API virtual FMulticastScriptDelegate& GetMulticastScriptDelegate(const void* PropertyValue, int32 Index) const;
 	// End of FMulticastDelegateProperty interface
 };
 
-class COREUOBJECT_API FMulticastSparseDelegateProperty : public TProperty_MulticastDelegate<FSparseDelegate> 
+class FMulticastSparseDelegateProperty : public TProperty_MulticastDelegate<FSparseDelegate> 
 {
-	DECLARE_FIELD(FMulticastSparseDelegateProperty, TProperty_MulticastDelegate<FSparseDelegate>, CASTCLASS_FMulticastSparseDelegateProperty)
+	DECLARE_FIELD_API(FMulticastSparseDelegateProperty, TProperty_MulticastDelegate<FSparseDelegate>, CASTCLASS_FMulticastSparseDelegateProperty, UE_API)
 
 public:
 
@@ -6237,7 +6239,7 @@ public:
 	 * @param InOwner Owner of the property
 	 * @param PropBase Pointer to the compiled in structure describing the property
 	 **/
-	FMulticastSparseDelegateProperty(FFieldVariant InOwner, const UECodeGen_Private::FMulticastDelegatePropertyParams& Prop);
+	UE_API FMulticastSparseDelegateProperty(FFieldVariant InOwner, const UECodeGen_Private::FMulticastDelegatePropertyParams& Prop);
 
 #if WITH_EDITORONLY_DATA
 	explicit FMulticastSparseDelegateProperty(UField* InField)
@@ -6247,29 +6249,29 @@ public:
 #endif // WITH_EDITORONLY_DATA
 
 	// FProperty interface
-	virtual void SerializeItem(FStructuredArchive::FSlot Slot, void* Value, void const* Defaults) const override;
-	virtual const TCHAR* ImportText_Internal(const TCHAR* Buffer, void* ContainerOrPropertyPtr, EPropertyPointerType PropertyPointerType, UObject* OwnerObject, int32 PortFlags, FOutputDevice* ErrorText) const override;
-	virtual bool LoadTypeName(UE::FPropertyTypeName Type, const FPropertyTag* Tag = nullptr) override;
+	UE_API virtual void SerializeItem(FStructuredArchive::FSlot Slot, void* Value, void const* Defaults) const override;
+	UE_API virtual const TCHAR* ImportText_Internal(const TCHAR* Buffer, void* ContainerOrPropertyPtr, EPropertyPointerType PropertyPointerType, UObject* OwnerObject, int32 PortFlags, FOutputDevice* ErrorText) const override;
+	UE_API virtual bool LoadTypeName(UE::FPropertyTypeName Type, const FPropertyTag* Tag = nullptr) override;
 	// End of FProperty interface
 
 	// FMulticastDelegateProperty interface
-	virtual const FMulticastScriptDelegate* GetMulticastDelegate(const void* PropertyValue) const override;
-	virtual void SetMulticastDelegate(void* PropertyValue, FMulticastScriptDelegate ScriptDelegate) const override;
+	UE_API virtual const FMulticastScriptDelegate* GetMulticastDelegate(const void* PropertyValue) const override;
+	UE_API virtual void SetMulticastDelegate(void* PropertyValue, FMulticastScriptDelegate ScriptDelegate) const override;
 
-	virtual void AddDelegate(FScriptDelegate ScriptDelegate, UObject* Parent = nullptr, void* PropertyValue = nullptr) const override;
-	virtual void RemoveDelegate(const FScriptDelegate& ScriptDelegate, UObject* Parent = nullptr, void* PropertyValue = nullptr) const override;
-	virtual void ClearDelegate(UObject* Parent = nullptr, void* PropertyValue = nullptr) const override;
+	UE_API virtual void AddDelegate(FScriptDelegate ScriptDelegate, UObject* Parent = nullptr, void* PropertyValue = nullptr) const override;
+	UE_API virtual void RemoveDelegate(const FScriptDelegate& ScriptDelegate, UObject* Parent = nullptr, void* PropertyValue = nullptr) const override;
+	UE_API virtual void ClearDelegate(UObject* Parent = nullptr, void* PropertyValue = nullptr) const override;
 
 protected:
-	virtual FMulticastScriptDelegate& GetMulticastScriptDelegate(const void* PropertyValue, int32 Index) const;
+	UE_API virtual FMulticastScriptDelegate& GetMulticastScriptDelegate(const void* PropertyValue, int32 Index) const;
 	// End of FMulticastDelegateProperty interface
 
 private:
-	virtual void SerializeItemInternal(FArchive& Ar, void* Value, void const* Defaults) const;
+	UE_API virtual void SerializeItemInternal(FArchive& Ar, void* Value, void const* Defaults) const;
 };
 
 /** Describes a single node in a custom property list. */
-struct COREUOBJECT_API FCustomPropertyListNode
+struct FCustomPropertyListNode
 {
 	/** The property that's being referenced at this node. */
 	FProperty* Property;
@@ -6309,7 +6311,7 @@ struct COREUOBJECT_API FCustomPropertyListNode
  * This class represents the chain of member properties leading to an internal struct property.  It is used
  * for tracking which member property corresponds to the UScriptStruct that owns a particular property.
  */
-class COREUOBJECT_API FEditPropertyChain : public TDoubleLinkedList<FProperty*>
+class FEditPropertyChain : public TDoubleLinkedList<FProperty*>
 {
 
 public:
@@ -6324,7 +6326,7 @@ public:
 	 * @return	true if the ActivePropertyNode was successfully changed to the node associated with the property
 	 *			specified.  false if there was no node corresponding to that property.
 	 */
-	bool SetActivePropertyNode( FProperty* NewActiveProperty );
+	UE_API bool SetActivePropertyNode( FProperty* NewActiveProperty );
 
 	/**
 	 * Sets the ActiveMemberPropertyNode to the node associated with the property specified.
@@ -6335,7 +6337,7 @@ public:
 	 * @return	true if the ActiveMemberPropertyNode was successfully changed to the node associated with the
 	 *			property specified.  false if there was no node corresponding to that property.
 	 */
-	bool SetActiveMemberPropertyNode( FProperty* NewActiveMemberProperty );
+	UE_API bool SetActiveMemberPropertyNode( FProperty* NewActiveMemberProperty );
 
 	/**
 	 * Specify the set of archetype instances that will be affected by the property change.
@@ -6350,19 +6352,19 @@ public:
 	/**
 	 * Returns whether the specified archetype instance will be affected by the property change.
 	 */
-	bool IsArchetypeInstanceAffected( UObject* InInstance ) const;
+	UE_API bool IsArchetypeInstanceAffected( UObject* InInstance ) const;
 
 	/**
 	 * Returns the node corresponding to the currently active property.
 	 */
-	TDoubleLinkedListNode* GetActiveNode() const;
+	UE_API TDoubleLinkedListNode* GetActiveNode() const;
 
 	/**
 	 * Returns the node corresponding to the currently active property, or if the currently active property
 	 * is not a member variable (i.e. inside of a struct/array), the node corresponding to the member variable
 	 * which contains the currently active property.
 	 */
-	TDoubleLinkedListNode* GetActiveMemberNode() const;
+	UE_API TDoubleLinkedListNode* GetActiveMemberNode() const;
 
 protected:
 	/**
@@ -6398,7 +6400,7 @@ protected:
 	 *
 	 * @param	NewListSize		the new size for this list
 	 */
-	virtual void SetListSize( int32 NewListSize );
+	UE_API virtual void SetListSize( int32 NewListSize );
 };
 
 
@@ -7298,9 +7300,9 @@ static_assert(sizeof(bool) == sizeof(uint8), "Bool is not one byte.");
  * Property wrappers are owned by UStruct that owns the property they wrap and are tied to its lifetime
  * so that weak object pointer functionality works as expected.
  */
-class COREUOBJECT_API UPropertyWrapper : public UObject
+class UPropertyWrapper : public UObject
 {
-	DECLARE_CLASS_INTRINSIC(UPropertyWrapper, UObject, CLASS_Transient, TEXT("/Script/CoreUObject"));
+	DECLARE_CLASS_INTRINSIC_API(UPropertyWrapper, UObject, CLASS_Transient, TEXT("/Script/CoreUObject"), UE_API);
 
 protected:
 	/** Cached property object */
@@ -7323,16 +7325,18 @@ public:
 	}
 };
 
-class COREUOBJECT_API UMulticastDelegatePropertyWrapper : public UPropertyWrapper
+class UMulticastDelegatePropertyWrapper : public UPropertyWrapper
 {
-	DECLARE_CLASS_INTRINSIC(UMulticastDelegatePropertyWrapper, UPropertyWrapper, CLASS_Transient, TEXT("/Script/CoreUObject"));
+	DECLARE_CLASS_INTRINSIC_API(UMulticastDelegatePropertyWrapper, UPropertyWrapper, CLASS_Transient, TEXT("/Script/CoreUObject"), UE_API);
 };
 
-class COREUOBJECT_API UMulticastInlineDelegatePropertyWrapper : public UMulticastDelegatePropertyWrapper
+class UMulticastInlineDelegatePropertyWrapper : public UMulticastDelegatePropertyWrapper
 {
-	DECLARE_CLASS_INTRINSIC(UMulticastInlineDelegatePropertyWrapper, UMulticastDelegatePropertyWrapper, CLASS_Transient, TEXT("/Script/CoreUObject"));
+	DECLARE_CLASS_INTRINSIC_API(UMulticastInlineDelegatePropertyWrapper, UMulticastDelegatePropertyWrapper, CLASS_Transient, TEXT("/Script/CoreUObject"), UE_API);
 };
 
+#undef UE_API
+
 #if UE_ENABLE_INCLUDE_ORDER_DEPRECATED_IN_5_5
 #include "Templates/IsTriviallyDestructible.h"
 #endif
