Add a new config string array ini:Editor:[BulkDataRegistry]CommandletsUsingBulkDataRegistry which allows individual commandlets to opt into having the BulkDataRegistry  enabled when they run.

diff --git a/Engine/Source/Runtime/CoreUObject/Private/Serialization/BulkDataRegistry.cpp b/Engine/Source/Runtime/CoreUObject/Private/Serialization/BulkDataRegistry.cpp
--- a/Engine/Source/Runtime/CoreUObject/Private/Serialization/BulkDataRegistry.cpp
+++ b/Engine/Source/Runtime/CoreUObject/Private/Serialization/BulkDataRegistry.cpp
@@ -8,10 +8,12 @@
 #include "Containers/Set.h"
 #include "HAL/CriticalSection.h"
 #include "Interfaces/IPluginManager.h"
+#include "Misc/CommandLine.h"
 #include "Misc/ConfigCacheIni.h"
 #include "Misc/CoreDelegates.h"
 #include "Misc/ScopeRWLock.h"
 #include "Serialization/EditorBulkData.h"
+#include "UObject/Class.h"
 #include "UObject/Package.h"
 #include "UObject/PackageResourceManager.h"
 #include "UObject/UObjectGlobals.h"
@@ -102,11 +104,52 @@ private:
 	FResaveSizeTracker ResaveSizeTracker;
 };
 
+/** Returns if the current running commandlet requires the BulkDataRegistry or not */
+bool IsAllowedForCurrentCommandlet()
+{
+	/** Cache the results so we don't have to redo this every time the function is called */
+	struct FConfigEnabled
+	{
+		FConfigEnabled()
+		{
+			FStringView CommandletToken(TEXTVIEW("commandlet"));
+
+			FString CommandletName = GetCommandletNameFromCmdline();
+			if (CommandletName.EndsWith(CommandletToken))
+			{
+				CommandletName.LeftChopInline(CommandletToken.Len(), EAllowShrinking::No);
+			}
+
+			TArray<FString> CommandletsUsingRegistry;
+			GConfig->GetArray(TEXT("BulkDataRegistry"), TEXT("CommandletsUsingBulkDataRegistry"), CommandletsUsingRegistry, GEditorIni);
+
+			for (FString& CommandletUsingRegistry : CommandletsUsingRegistry)
+			{
+				if (CommandletUsingRegistry.EndsWith(CommandletToken))
+				{
+					CommandletUsingRegistry.LeftChopInline(CommandletToken.Len(), EAllowShrinking::No);
+				}
+
+				if (CommandletUsingRegistry == CommandletName)
+				{
+					bEnabled = true;
+					break;
+				}
+			}
+		}
+
+		bool bEnabled = false;
+
+	} static ConfigEnabled;
+
+	return ConfigEnabled.bEnabled;
+}
+
 }
 
 bool IBulkDataRegistry::IsEnabled()
 {
-	if (IsRunningCommandlet() && !IsRunningCookCommandlet())
+	if (IsRunningCommandlet() && !IsRunningCookCommandlet() && !UE::BulkDataRegistry::Private::IsAllowedForCurrentCommandlet())
 	{
 		// Disabled in other commandlets because they do not tick the BulkDataRegistry.
 		// When not ticked, the BulkDataRegistry will never free memory and will hold onto all memory
