Serialize structs with defaults when supported

diff --git a/Engine/Source/Runtime/CoreUObject/Private/StructUtils/InstancedStruct.cpp b/Engine/Source/Runtime/CoreUObject/Private/StructUtils/InstancedStruct.cpp
--- a/Engine/Source/Runtime/CoreUObject/Private/StructUtils/InstancedStruct.cpp
+++ b/Engine/Source/Runtime/CoreUObject/Private/StructUtils/InstancedStruct.cpp
@@ -119,7 +119,7 @@ void FInstancedStruct::Reset()
 	ResetStructData();
 }
 
-bool FInstancedStruct::Serialize(FArchive& Ar)
+bool FInstancedStruct::Serialize(FArchive& Ar, UStruct* DefaultsStruct, const void* Defaults)
 {
 	Ar.UsingCustomVersion(GInstancedStructCustomVersion.GUID);
 
diff --git a/Engine/Source/Runtime/CoreUObject/Private/UObject/Class.cpp b/Engine/Source/Runtime/CoreUObject/Private/UObject/Class.cpp
--- a/Engine/Source/Runtime/CoreUObject/Private/UObject/Class.cpp
+++ b/Engine/Source/Runtime/CoreUObject/Private/UObject/Class.cpp
@@ -3311,7 +3311,7 @@ void UScriptStruct::SerializeItem(FStructuredArchive::FSlot Slot, void* Value, v
 
 		if (TheCppStructOps->HasStructuredSerializer())
 		{
-			bItemSerialized = TheCppStructOps->Serialize(Slot, Value);
+			bItemSerialized = TheCppStructOps->Serialize(Slot, Value, this, Defaults);
 		}
 		else
 		{
@@ -3320,7 +3320,7 @@ void UScriptStruct::SerializeItem(FStructuredArchive::FSlot Slot, void* Value, v
 			{
 				FArchiveUObjectFromStructuredArchive Adapter(Slot);
 				FArchive& Ar = Adapter.GetArchive();
-				bItemSerialized = TheCppStructOps->Serialize(Ar, Value);
+				bItemSerialized = TheCppStructOps->Serialize(Ar, Value, this, Defaults);
 				if (bItemSerialized && !Slot.IsFilled())
 				{
 					// The struct said that serialization succeeded but it didn't actually write anything.
@@ -3331,7 +3331,7 @@ void UScriptStruct::SerializeItem(FStructuredArchive::FSlot Slot, void* Value, v
 			else
 #endif
 			{
-				bItemSerialized = TheCppStructOps->Serialize(Slot.GetUnderlyingArchive(), Value);
+				bItemSerialized = TheCppStructOps->Serialize(Slot.GetUnderlyingArchive(), Value, this, Defaults);
 				if (bItemSerialized && !Slot.IsFilled())
 				{
 					// The struct said that serialization succeeded but it didn't actually write anything.
diff --git a/Engine/Source/Runtime/CoreUObject/Public/Concepts/StructSerializableWithDefaults.h b/Engine/Source/Runtime/CoreUObject/Public/Concepts/StructSerializableWithDefaults.h
new file mode 100644
--- /dev/null
+++ b/Engine/Source/Runtime/CoreUObject/Public/Concepts/StructSerializableWithDefaults.h
@@ -0,0 +1,23 @@
+// Copyright Epic Games, Inc. All Rights Reserved.
+
+#pragma once
+
+class FArchive;
+class FStructuredArchiveSlot;
+class UStruct;
+
+/** Describes a struct that can be serialized with defaults. */
+struct CStructSerializableWithDefaults
+{
+	template <typename StructType>
+	auto Requires(FArchive& Ar, StructType* Data, UStruct* DefaultsStruct, const void* Defaults)
+		-> decltype(Data->Serialize(Ar, DefaultsStruct, Defaults));
+};
+
+/** Describes a struct that can be serialized with defaults. */
+struct CStructStructuredSerializableWithDefaults
+{
+	template <typename StructType>
+	auto Requires(FStructuredArchiveSlot Slot, StructType* Data, UStruct* DefaultsStruct, const void* Defaults)
+		-> decltype(Data->Serialize(Slot, DefaultsStruct, Defaults));
+};
diff --git a/Engine/Source/Runtime/CoreUObject/Public/StructUtils/InstancedStruct.h b/Engine/Source/Runtime/CoreUObject/Public/StructUtils/InstancedStruct.h
--- a/Engine/Source/Runtime/CoreUObject/Public/StructUtils/InstancedStruct.h
+++ b/Engine/Source/Runtime/CoreUObject/Public/StructUtils/InstancedStruct.h
@@ -151,7 +151,7 @@ public:
 	}
 
 	/** For StructOpsTypeTraits */
-	UE_API bool Serialize(FArchive& Ar);
+	UE_API bool Serialize(FArchive& Ar, UStruct* DefaultsStruct = nullptr, const void* Defaults = nullptr);
 	UE_API bool Identical(const FInstancedStruct* Other, uint32 PortFlags) const;
 	UE_API void AddStructReferencedObjects(class FReferenceCollector& Collector);
 	UE_API bool ExportTextItem(FString& ValueStr, FInstancedStruct const& DefaultValue, class UObject* Parent, int32 PortFlags, class UObject* ExportRootScope) const;
diff --git a/Engine/Source/Runtime/CoreUObject/Public/UObject/Class.h b/Engine/Source/Runtime/CoreUObject/Public/UObject/Class.h
--- a/Engine/Source/Runtime/CoreUObject/Public/UObject/Class.h
+++ b/Engine/Source/Runtime/CoreUObject/Public/UObject/Class.h
@@ -9,6 +9,7 @@
 #include "Concepts/GetTypeHashable.h"
 #include "Concepts/StaticClassProvider.h"
 #include "Concepts/StaticStructProvider.h"
+#include "Concepts/StructSerializableWithDefaults.h"
 #include "Containers/Array.h"
 #include "Containers/EnumAsByte.h"
 #include "Containers/Map.h"
@@ -1099,12 +1100,25 @@ public:
 		{
 			return GetCapabilities().HasStructuredSerializer;
 		}
-		/** 
-		 * Serialize this structure 
-		 * @return true if the package is new enough to support this, if false, it will fall back to ordinary script struct serialization
+
+		UE_DEPRECATED(5.6, "Use the overload of Serialize with defaults.")
+		virtual bool Serialize(FArchive& Ar, void* Data)
+		{
+			return Serialize(Ar, Data, nullptr, nullptr);
+		}
+
+		UE_DEPRECATED(5.6, "Use the overload of Serialize with defaults.")
+		virtual bool Serialize(FStructuredArchive::FSlot Slot, void* Data)
+		{
+			return Serialize(Slot, Data, nullptr, nullptr);
+		}
+
+		/**
+		 * Serialize this structure.
+		 * @return true if serialization was handled, false to fall back to standard script struct serialization.
 		 */
-		virtual bool Serialize(FArchive& Ar, void *Data) = 0;
-		virtual bool Serialize(FStructuredArchive::FSlot Slot, void *Data) = 0;
+		virtual bool Serialize(FArchive& Ar, void* Data, UStruct* DefaultsStruct, const void* Defaults) = 0;
+		virtual bool Serialize(FStructuredArchive::FSlot Slot, void* Data, UStruct* DefaultsStruct, const void* Defaults) = 0;
 
 		/** return true if this class implements a post serialize call **/
 		bool HasPostSerialize() const
@@ -1422,24 +1436,38 @@ public:
 			// could also not attempt to call the destructor for types where this is not possible, but I didn't do that here
 			DestructItem((CPPSTRUCT*)Dest);
 		}
-		virtual bool Serialize(FArchive& Ar, void *Data) override
+		virtual bool Serialize(FArchive& Ar, void* Data, UStruct* DefaultsStruct, const void* Defaults) override
 		{
 			check(TTraits::WithSerializer); // don't call this if we have indicated it is not necessary
 			if constexpr (TStructOpsTypeTraits<CPPSTRUCT>::WithSerializer)
 			{
-				return ((CPPSTRUCT*)Data)->Serialize(Ar);
+				if constexpr (TModels_V<CStructSerializableWithDefaults, CPPSTRUCT>)
+				{
+					return ((CPPSTRUCT*)Data)->Serialize(Ar, DefaultsStruct, Defaults);
+				}
+				else
+				{
+					return ((CPPSTRUCT*)Data)->Serialize(Ar);
+				}
 			}
 			else
 			{
 				return false;
 			}
 		}
-		virtual bool Serialize(FStructuredArchive::FSlot Slot, void *Data) override
+		virtual bool Serialize(FStructuredArchive::FSlot Slot, void* Data, UStruct* DefaultsStruct, const void* Defaults) override
 		{
 			check(TTraits::WithStructuredSerializer); // don't call this if we have indicated it is not necessary
 			if constexpr (TStructOpsTypeTraits<CPPSTRUCT>::WithStructuredSerializer)
 			{
-				return ((CPPSTRUCT*)Data)->Serialize(Slot);
+				if constexpr (TModels_V<CStructStructuredSerializableWithDefaults, CPPSTRUCT>)
+				{
+					return ((CPPSTRUCT*)Data)->Serialize(Slot, DefaultsStruct, Defaults);
+				}
+				else
+				{
+					return ((CPPSTRUCT*)Data)->Serialize(Slot);
+				}
 			}
 			else
 			{
