[Optim] Use RW lock for UObjectAnnotation. Improves rewind debugger performance

UObjectTrace::GetObjectId() Average 1.305ms -> 0.108ms

diff --git a/Engine/Source/Runtime/CoreUObject/Private/UObject/OverridableManager.cpp b/Engine/Source/Runtime/CoreUObject/Private/UObject/OverridableManager.cpp
--- a/Engine/Source/Runtime/CoreUObject/Private/UObject/OverridableManager.cpp
+++ b/Engine/Source/Runtime/CoreUObject/Private/UObject/OverridableManager.cpp
@@ -395,7 +395,7 @@ FOverridableManager::FOverridableManager()
 void FOverridableManager::HandleObjectsReInstantiated(const TMap<UObject*, UObject*>& OldToNewInstanceMap)
 {
 #if WITH_EDITORONLY_DATA
-	UE::TScopeLock AnnotationMapLock(AnnotationMapCritical);
+	UE::TRWScopeLock AnnotationMapLock(AnnotationMapCritical, SLT_ReadOnly);
 	for (const auto& Pair : GetAnnotationMap())
 	{
 		if (FOverriddenPropertySet* OverriddenProperties = Pair.Value.OverriddenProperties.Get())
@@ -413,7 +413,7 @@ void FOverridableManager::AddReferencedObjects(FReferenceCollector& Collector)
 	// that the replace references archives go through this object
 	if (!GIsGarbageCollecting)
 	{
-		UE::TScopeLock AnnotationMapLock(AnnotationMapCritical);
+		UE::TRWScopeLock AnnotationMapLock(AnnotationMapCritical, SLT_ReadOnly);
 		for (const auto& Pair : GetAnnotationMap())
 		{
 			if (FOverriddenPropertySet* OverriddenProperties = Pair.Value.OverriddenProperties.Get())
@@ -428,7 +428,7 @@ void FOverridableManager::AddReferencedObjects(FReferenceCollector& Collector)
 void FOverridableManager::HandleDeadObjectReferences(const TSet<UClass*>& DeadClasses, const TSet<UObject*>& ActiveInstances, const TSet<UObject*>& TemplateInstances)
 {
 #if WITH_EDITORONLY_DATA
-	UE::TScopeLock AnnotationMapLock(AnnotationMapCritical);
+	UE::TRWScopeLock AnnotationMapLock(AnnotationMapCritical, SLT_ReadOnly);
 	for (const auto& Pair : GetAnnotationMap())
 	{
 		if (FOverriddenPropertySet* OverriddenProperties = Pair.Value.OverriddenProperties.Get())
diff --git a/Engine/Source/Runtime/CoreUObject/Public/UObject/UObjectAnnotation.h b/Engine/Source/Runtime/CoreUObject/Public/UObject/UObjectAnnotation.h
--- a/Engine/Source/Runtime/CoreUObject/Public/UObject/UObjectAnnotation.h
+++ b/Engine/Source/Runtime/CoreUObject/Public/UObject/UObjectAnnotation.h
@@ -43,7 +43,7 @@ public:
 #if !(UE_BUILD_SHIPPING || UE_BUILD_TEST)
 		if (!bAutoRemove)
 		{
-			UE::TScopeLock AnnotationMapLock(AnnotationMapCritical);
+			UE::TRWScopeLock AnnotationMapLock(AnnotationMapCritical, SLT_ReadOnly);
 			// in this case we are only verifying that the external assurances of removal are met
 			check(!AnnotationMap.Find(Object));
 		}
@@ -99,7 +99,7 @@ private:
 		}
 		else
 		{
-			UE::TScopeLock AnnotationMapLock(AnnotationMapCritical);
+			UE::TRWScopeLock AnnotationMapLock(AnnotationMapCritical, SLT_Write);
 			const bool bWasEmpty = (AnnotationMap.Num() == 0);
 
 			AnnotationMap.Add(Object, LocalAnnotation);
@@ -149,8 +149,9 @@ public:
 		// Avoid holding the lock while we call GUObjectArray.RemoveUObjectDeleteListener as it could deadlock
 		TAnnotation Result;
 		{
-			UE::TScopeLock AnnotationMapLock(AnnotationMapCritical);
 			SetAnnotationCacheKeyAndValue(Object, TAnnotation());
+
+			UE::TRWScopeLock AnnotationMapLock(AnnotationMapCritical, SLT_Write);
 			bHadElements = (AnnotationMap.Num() > 0);
 			AnnotationMap.RemoveAndCopyValue(Object, Result);
 			bIsNowEmpty = (AnnotationMap.Num() == 0);
@@ -180,8 +181,9 @@ public:
 		bool bIsNowEmpty = false;
 		// Avoid holding the lock while we call GUObjectArray.RemoveUObjectDeleteListener as it could deadlock
 		{
-			UE::TScopeLock AnnotationMapLock(AnnotationMapCritical);
 			SetAnnotationCacheKeyAndValue(Object, TAnnotation());
+
+			UE::TRWScopeLock AnnotationMapLock(AnnotationMapCritical, SLT_Write);
 			bHadElements = (AnnotationMap.Num() > 0);
 			AnnotationMap.Remove(Object);
 			bIsNowEmpty = (AnnotationMap.Num() == 0);
@@ -207,9 +209,9 @@ public:
 
 		// Avoid holding the lock while we call GUObjectArray.RemoveUObjectDeleteListener as it could deadlock
 		{
-			UE::TScopeLock AnnotationMapLock(AnnotationMapCritical);
 			SetAnnotationCacheKeyAndValue(nullptr, TAnnotation());
 			bHadElements = (AnnotationMap.Num() > 0);
+			UE::TRWScopeLock AnnotationMapLock(AnnotationMapCritical, SLT_Write);
 			AnnotationMap.Empty();
 		}
 		if (bHadElements)
@@ -233,24 +235,31 @@ public:
 	{
 		check(Object);
 
-		UE::TScopeLock AnnotationMapLock(AnnotationMapCritical);
-
+		
 		// If we are within a transaction, we don't touch the annotation cache, as it can be altered from the open.
 		// Look up the annotation directly from the map.
 		if (AutoRTFM::IsTransactional())
 		{
+			UE::TRWScopeLock AnnotationMapLock(AnnotationMapCritical, SLT_ReadOnly);
 			TAnnotation* Entry = AnnotationMap.Find(Object);
 			return Entry ? *Entry : TAnnotation();
 		}
 
-		// Update the cache if it isn't a match.
-		if (Object != AnnotationCacheKey)
-		{			
-			TAnnotation* Entry = AnnotationMap.Find(Object);
-			SetAnnotationCacheKeyAndValue(Object, Entry ? *Entry : TAnnotation());
+		{
+			// Check cache
+			UE::TRWScopeLock AnnotationCacheLock(AnnotationMapCritical, SLT_ReadOnly);
+			if (Object == AnnotationCacheKey)
+			{
+				return AnnotationCacheValue;
+			}
 		}
+		
+		// Update the cache if it isn't a match.
+		UE::TRWScopeLock AnnotationMapLock(AnnotationMapCritical, SLT_ReadOnly);
+		TAnnotation* Entry = AnnotationMap.Find(Object);
+		SetAnnotationCacheKeyAndValue(Object, Entry ? *Entry : TAnnotation());
 
-		return AnnotationCacheValue;
+		return Entry ? *Entry : TAnnotation();
 	}
 
 	/**
@@ -267,12 +276,13 @@ public:
 	 */
 	void Reserve(int32 ExpectedNumElements)
 	{
-		UE::TScopeLock AnnotationMapLock(AnnotationMapCritical);
+		UE::TRWScopeLock AnnotationMapLock(AnnotationMapCritical, SLT_Write);
 		AnnotationMap.Empty(ExpectedNumElements);
 	}
 
 	virtual SIZE_T GetAllocatedSize() const override
 	{
+		UE::TRWScopeLock AnnotationMapLock(AnnotationMapCritical, SLT_ReadOnly);
 		return AnnotationMap.GetAllocatedSize();
 	}
 
@@ -290,6 +300,7 @@ private:
 		}
 		else
 		{
+			UE::TRWScopeLock AnnotationCacheLock(AnnotationCacheCritical, SLT_Write);
 			AnnotationCacheKey = Key;
 			AnnotationCacheValue = Value;
 		}
@@ -338,7 +349,8 @@ protected:
 	/**
 	 * Guards the annotation map, cache key and cache value.
 	 */
-	FTransactionallySafeCriticalSection AnnotationMapCritical;
+	mutable FTransactionallySafeRWLock AnnotationMapCritical;
+	mutable FTransactionallySafeRWLock AnnotationCacheCritical;
 
 	/**
 	 * Map from live objects to an annotation. Guarded by AnnotationMapCritical.
@@ -346,13 +358,13 @@ protected:
 	TMap<const UObjectBase *, TAnnotation> AnnotationMap;
 
 	/**
-	 * Key for a one-item cache of the last lookup into AnnotationMap. Guarded by AnnotationMapCritical.
+	 * Key for a one-item cache of the last lookup into AnnotationMap. Guarded by AnnotationCacheCritical.
 	 * Annotation are often called back-to-back so this is a performance optimization for that.
 	 */
 	const UObjectBase* AnnotationCacheKey = nullptr;
 
 	/**
-	 * Value for a one-item cache of the last lookup into AnnotationMap. Guarded by AnnotationMapCritical.
+	 * Value for a one-item cache of the last lookup into AnnotationMap. Guarded by AnnotationCacheCritical.
 	 */
 	TAnnotation AnnotationCacheValue;
 
