Async Loading:
- Converted progress delegate to a TSharedPtr instead of TUniquePtr to allow passing the same progress delegate to multiple load requests without having to do memory allocations.
- One use case is when chaining load requests from the progress callback from the async loading thread.
- Did not convert the completion progress callback as it will be replaced by the FullyLoaded progress callback in the future.
- Updated the LoadAssetAsync API to allow passing a progress delegate.

diff --git a/Engine/Source/Runtime/CoreUObject/Private/Serialization/AsyncLoading2.cpp b/Engine/Source/Runtime/CoreUObject/Private/Serialization/AsyncLoading2.cpp
--- a/Engine/Source/Runtime/CoreUObject/Private/Serialization/AsyncLoading2.cpp
+++ b/Engine/Source/Runtime/CoreUObject/Private/Serialization/AsyncLoading2.cpp
@@ -675,7 +675,7 @@ struct FPackageRequest
 	FName CustomName;
 	FPackagePath PackagePath;
 	TUniquePtr<FLoadPackageAsyncDelegate> PackageLoadedDelegate;
-	TUniquePtr<FLoadPackageAsyncProgressDelegate> PackageProgressDelegate;
+	TSharedPtr<FLoadPackageAsyncProgressDelegate> PackageProgressDelegate;
 	FPackageRequest* Next = nullptr;
 	FPackageReferencer PackageReferencer;
 
@@ -684,7 +684,7 @@ struct FPackageRequest
 		return &InstancingContext;
 	}
 
-	static FPackageRequest Create(int32 RequestId, EPackageFlags PackageFlags, uint32 LoadFlags, int32 PIEInstanceID, int32 Priority, const FLinkerInstancingContext* InstancingContext, const FPackagePath& PackagePath, FName CustomName, TUniquePtr<FLoadPackageAsyncDelegate> PackageLoadedDelegate, TUniquePtr<FLoadPackageAsyncProgressDelegate> PackageProgressDelegate, FPackageReferencer PackageReferencer)
+	static FPackageRequest Create(int32 RequestId, EPackageFlags PackageFlags, uint32 LoadFlags, int32 PIEInstanceID, int32 Priority, const FLinkerInstancingContext* InstancingContext, const FPackagePath& PackagePath, FName CustomName, TUniquePtr<FLoadPackageAsyncDelegate> PackageLoadedDelegate, TSharedPtr<FLoadPackageAsyncProgressDelegate> PackageProgressDelegate, FPackageReferencer PackageReferencer)
 	{
 		return FPackageRequest
 		{
@@ -699,7 +699,7 @@ struct FPackageRequest
 			CustomName,
 			PackagePath,
 			MoveTemp(PackageLoadedDelegate),
-			MoveTemp(PackageProgressDelegate),
+			PackageProgressDelegate,
 			nullptr,
 			PackageReferencer
 		};
@@ -3558,7 +3558,7 @@ struct FAsyncPackage2
 	double GetLoadStartTime() const;
 
 	void AddCompletionCallback(TUniquePtr<FLoadPackageAsyncDelegate>&& Callback);
-	void AddProgressCallback(TUniquePtr<FLoadPackageAsyncProgressDelegate>&& Callback);
+	void AddProgressCallback(TSharedPtr<FLoadPackageAsyncProgressDelegate> Callback);
 
 	FORCEINLINE UPackage* GetLinkerRoot() const
 	{
@@ -3816,7 +3816,7 @@ private:
 	/** Callbacks called when we finished loading this package */
 	TArray<TUniquePtr<FLoadPackageAsyncDelegate>, TInlineAllocator<2>> CompletionCallbacks;
 	/** Callbacks called for the different loading phase of this package */
-	TArray<TUniquePtr<FLoadPackageAsyncProgressDelegate>, TInlineAllocator<2>> ProgressCallbacks;
+	TArray<TSharedPtr<FLoadPackageAsyncProgressDelegate>, TInlineAllocator<2>> ProgressCallbacks;
 
 public:
 
@@ -4038,7 +4038,7 @@ private:
 		UPackage* UPackage = nullptr;
 		FAsyncPackage2* AsyncPackage = nullptr;
 		TArray<TUniquePtr<FLoadPackageAsyncDelegate>, TInlineAllocator<2>> CompletionCallbacks;
-		TArray<TUniquePtr<FLoadPackageAsyncProgressDelegate>, TInlineAllocator<2>> ProgressCallbacks;
+		TArray<TSharedPtr<FLoadPackageAsyncProgressDelegate>, TInlineAllocator<2>> ProgressCallbacks;
 		TArray<int32, TInlineAllocator<2>> RequestIDs;
 
 		static FCompletedPackageRequest FromUnreslovedPackage(
@@ -4097,7 +4097,7 @@ private:
 
 				const EAsyncLoadingProgress ProgressType = Result == EAsyncLoadingResult::Succeeded ? EAsyncLoadingProgress::FullyLoaded : EAsyncLoadingProgress::Failed;
 
-				for (TUniquePtr<FLoadPackageAsyncProgressDelegate>& ProgressCallback : ProgressCallbacks)
+				for (TSharedPtr<FLoadPackageAsyncProgressDelegate>& ProgressCallback : ProgressCallbacks)
 				{
 					FLoadPackageAsyncProgressParams Params
 					{
@@ -4107,11 +4107,6 @@ private:
 					};
 					
 					ProgressCallback->Invoke(Params);
-
-					if (Params.bReleaseDelegate)
-					{
-						verify(ProgressCallback.Release());
-					}
 				}
 				ProgressCallbacks.Empty();
 			}
@@ -4306,8 +4301,8 @@ public:
 	void UpdatePackagePriority(FAsyncLoadingThreadState2& ThreadState, FAsyncPackage2* Package);
 	void UpdatePackagePriorityRecursive(FAsyncLoadingThreadState2& ThreadState, FAsyncPackage2* Package, int32 NewPriority);
 
-	FAsyncPackage2* FindOrInsertPackage(FAsyncLoadingThreadState2& ThreadState, FAsyncPackageDesc2& InDesc, bool& bInserted, FAsyncPackage2* ImportedByPackage, TUniquePtr<FLoadPackageAsyncDelegate>&& PackageLoadedDelegate = TUniquePtr<FLoadPackageAsyncDelegate>(), TUniquePtr<FLoadPackageAsyncProgressDelegate>&& PackageProgressDelegate = TUniquePtr<FLoadPackageAsyncProgressDelegate>());
-	void QueueUnresolvedPackage(FAsyncLoadingThreadState2& ThreadState, EPackageStoreEntryStatus PackageStatus, FAsyncPackageDesc2& PackageDesc, TUniquePtr<FLoadPackageAsyncDelegate>&& LoadPackageAsyncDelegate, TUniquePtr<FLoadPackageAsyncProgressDelegate>&& LoadPackageAsyncProgressDelegate);
+	FAsyncPackage2* FindOrInsertPackage(FAsyncLoadingThreadState2& ThreadState, FAsyncPackageDesc2& InDesc, bool& bInserted, FAsyncPackage2* ImportedByPackage, TUniquePtr<FLoadPackageAsyncDelegate>&& PackageLoadedDelegate = TUniquePtr<FLoadPackageAsyncDelegate>(), TSharedPtr<FLoadPackageAsyncProgressDelegate> PackageProgressDelegate = nullptr);
+	void QueueUnresolvedPackage(FAsyncLoadingThreadState2& ThreadState, EPackageStoreEntryStatus PackageStatus, FAsyncPackageDesc2& PackageDesc, TUniquePtr<FLoadPackageAsyncDelegate>&& LoadPackageAsyncDelegate, TSharedPtr<FLoadPackageAsyncProgressDelegate> LoadPackageAsyncProgressDelegate);
 
 	/**
 	* [ASYNC* THREAD] Loads all packages
@@ -4348,7 +4343,7 @@ public:
 		return true;
 	}
 
-	int32 LoadPackageInternal(const FPackagePath& InPackagePath, FName InCustomName, TUniquePtr<FLoadPackageAsyncDelegate>&& InCompletionDelegate, TUniquePtr<FLoadPackageAsyncProgressDelegate>&& InProgressDelegate, EPackageFlags InPackageFlags, int32 InPIEInstanceID, int32 InPackagePriority, const FLinkerInstancingContext* InInstancingContext, uint32 InLoadFlags);
+	int32 LoadPackageInternal(const FPackagePath& InPackagePath, FName InCustomName, TUniquePtr<FLoadPackageAsyncDelegate>&& InCompletionDelegate, TSharedPtr<FLoadPackageAsyncProgressDelegate> InProgressDelegate, EPackageFlags InPackageFlags, int32 InPIEInstanceID, int32 InPackagePriority, const FLinkerInstancingContext* InInstancingContext, uint32 InLoadFlags);
 
 	virtual int32 LoadPackage(
 		const FPackagePath& InPackagePath,
@@ -5234,7 +5229,7 @@ void FAsyncLoadingThread2::MergePostLoadGroups(FAsyncLoadingThreadState2& Thread
 	delete Source;
 }
 
-FAsyncPackage2* FAsyncLoadingThread2::FindOrInsertPackage(FAsyncLoadingThreadState2& ThreadState, FAsyncPackageDesc2& Desc, bool& bInserted, FAsyncPackage2* ImportedByPackage, TUniquePtr<FLoadPackageAsyncDelegate>&& PackageLoadedDelegate, TUniquePtr<FLoadPackageAsyncProgressDelegate>&& PackageProgressDelegate)
+FAsyncPackage2* FAsyncLoadingThread2::FindOrInsertPackage(FAsyncLoadingThreadState2& ThreadState, FAsyncPackageDesc2& Desc, bool& bInserted, FAsyncPackage2* ImportedByPackage, TUniquePtr<FLoadPackageAsyncDelegate>&& PackageLoadedDelegate, TSharedPtr<FLoadPackageAsyncProgressDelegate> PackageProgressDelegate)
 {
 	TRACE_CPUPROFILER_EVENT_SCOPE(FindOrInsertPackage);
 	check(ThreadState.bCanAccessAsyncLoadingThreadData);
@@ -5271,7 +5266,7 @@ FAsyncPackage2* FAsyncLoadingThread2::FindOrInsertPackage(FAsyncLoadingThreadSta
 		}
 		if (PackageProgressDelegate.IsValid())
 		{
-			Package->AddProgressCallback(MoveTemp(PackageProgressDelegate));
+			Package->AddProgressCallback(PackageProgressDelegate);
 		}
 	}
 
@@ -7408,7 +7403,7 @@ void FAsyncPackage2::CallProgressCallbacks(EAsyncLoadingProgress ProgressType)
 			.ProgressType = ProgressType
 		};
 
-		for (TUniquePtr<FLoadPackageAsyncProgressDelegate>& ProgressCallback : ProgressCallbacks)
+		for (TSharedPtr<FLoadPackageAsyncProgressDelegate>& ProgressCallback : ProgressCallbacks)
 		{
 			ProgressCallback->Invoke(Params);
 		}
@@ -9380,7 +9375,7 @@ EAsyncPackageState::Type FAsyncLoadingThread2::ProcessLoadedPackagesFromGameThre
 				{
 					// In the editor we need to find any assets and packages and add them to list for later callback
 					EditorCompletedUPackages.Add(Package->LinkerRoot);
-					UE_MT_SCOPED_READ_ACCESS(Package->ConstructedObjectsAccessDetector)
+					UE_MT_SCOPED_READ_ACCESS(Package->ConstructedObjectsAccessDetector);
 					for (UObject* Object : Package->ConstructedObjects)
 					{
 						if (Object->IsAsset())
@@ -10489,7 +10484,7 @@ void FAsyncPackage2::MoveConstructedObjectsToPhase2()
 {
 	TRACE_CPUPROFILER_EVENT_SCOPE(MoveConstructedObjectsToPhase2);
 	{
-		UE_MT_SCOPED_READ_ACCESS(ConstructedObjectsAccessDetector)
+		UE_MT_SCOPED_READ_ACCESS(ConstructedObjectsAccessDetector);
 		for (UObject* Object : ConstructedObjects)
 		{
 			checkObject(Object, !Object->HasAnyFlags(RF_NeedLoad | RF_NeedInitialization));
@@ -10545,7 +10540,7 @@ void FAsyncPackage2::MoveConstructedObjectsToPhase2()
 void FAsyncPackage2::ClearConstructedObjects()
 {
 	TRACE_CPUPROFILER_EVENT_SCOPE(ClearConstructedObjects);
-	UE_MT_SCOPED_WRITE_ACCESS(ConstructedObjectsAccessDetector)
+	UE_MT_SCOPED_WRITE_ACCESS(ConstructedObjectsAccessDetector);
 
 	for (UObject* Object : ConstructedObjects)
 	{
@@ -10857,7 +10852,7 @@ void FAsyncPackage2::AddCompletionCallback(TUniquePtr<FLoadPackageAsyncDelegate>
 	CompletionCallbacks.Emplace(MoveTemp(Callback));
 }
 
-void FAsyncPackage2::AddProgressCallback(TUniquePtr<FLoadPackageAsyncProgressDelegate>&& Callback)
+void FAsyncPackage2::AddProgressCallback(TSharedPtr<FLoadPackageAsyncProgressDelegate> Callback)
 {
 	ProgressCallbacks.Emplace(MoveTemp(Callback));
 }
@@ -10866,7 +10861,7 @@ UE_TRACE_EVENT_BEGIN(CUSTOM_LOADTIMER_LOG, LoadAsyncPackageInternal, NoSync)
 	UE_TRACE_EVENT_FIELD(UE::Trace::WideString, PackageName)
 UE_TRACE_EVENT_END()
 
-int32 FAsyncLoadingThread2::LoadPackageInternal(const FPackagePath& InPackagePath, FName InCustomName, TUniquePtr<FLoadPackageAsyncDelegate>&& InCompletionDelegate, TUniquePtr<FLoadPackageAsyncProgressDelegate>&& InProgressDelegate, EPackageFlags InPackageFlags, int32 InPIEInstanceID, int32 InPackagePriority, const FLinkerInstancingContext* InInstancingContext, uint32 InLoadFlags)
+int32 FAsyncLoadingThread2::LoadPackageInternal(const FPackagePath& InPackagePath, FName InCustomName, TUniquePtr<FLoadPackageAsyncDelegate>&& InCompletionDelegate, TSharedPtr<FLoadPackageAsyncProgressDelegate> InProgressDelegate, EPackageFlags InPackageFlags, int32 InPIEInstanceID, int32 InPackagePriority, const FLinkerInstancingContext* InInstancingContext, uint32 InLoadFlags)
 {
 	TRACE_CPUPROFILER_EVENT_SCOPE(LoadPackage);
 
@@ -10925,7 +10920,7 @@ PRAGMA_ENABLE_DEPRECATION_WARNINGS
 
 int32 FAsyncLoadingThread2::LoadPackage(const FPackagePath& InPackagePath, FLoadPackageAsyncOptionalParams InParams)
 {
-	return LoadPackageInternal(InPackagePath, InParams.CustomPackageName, MoveTemp(InParams.CompletionDelegate), MoveTemp(InParams.ProgressDelegate), InParams.PackageFlags, InParams.PIEInstanceID, InParams.PackagePriority, InParams.InstancingContext, InParams.LoadFlags);
+	return LoadPackageInternal(InPackagePath, InParams.CustomPackageName, MoveTemp(InParams.CompletionDelegate), InParams.ProgressDelegate, InParams.PackageFlags, InParams.PIEInstanceID, InParams.PackagePriority, InParams.InstancingContext, InParams.LoadFlags);
 }
 
 int32 FAsyncLoadingThread2::LoadPackage(const FPackagePath& InPackagePath, FName InCustomName, FLoadPackageAsyncDelegate InCompletionDelegate, EPackageFlags InPackageFlags, int32 InPIEInstanceID, int32 InPackagePriority, const FLinkerInstancingContext* InInstancingContext, uint32 InLoadFlags)
@@ -10935,10 +10930,10 @@ int32 FAsyncLoadingThread2::LoadPackage(const FPackagePath& InPackagePath, FName
 		? MakeUnique<FLoadPackageAsyncDelegate>(MoveTemp(InCompletionDelegate))
 		: TUniquePtr<FLoadPackageAsyncDelegate>();
 
-	return LoadPackageInternal(InPackagePath, InCustomName, MoveTemp(CompletionDelegate), TUniquePtr<FLoadPackageAsyncProgressDelegate>(), InPackageFlags, InPIEInstanceID, InPackagePriority, InInstancingContext, InLoadFlags);
+	return LoadPackageInternal(InPackagePath, InCustomName, MoveTemp(CompletionDelegate), nullptr, InPackageFlags, InPIEInstanceID, InPackagePriority, InInstancingContext, InLoadFlags);
 }
 
-void FAsyncLoadingThread2::QueueUnresolvedPackage(FAsyncLoadingThreadState2& ThreadState, EPackageStoreEntryStatus PackageStatus, FAsyncPackageDesc2& PackageDesc, TUniquePtr<FLoadPackageAsyncDelegate>&& PackageLoadedDelegate, TUniquePtr<FLoadPackageAsyncProgressDelegate>&& PackageProgressDelegate)
+void FAsyncLoadingThread2::QueueUnresolvedPackage(FAsyncLoadingThreadState2& ThreadState, EPackageStoreEntryStatus PackageStatus, FAsyncPackageDesc2& PackageDesc, TUniquePtr<FLoadPackageAsyncDelegate>&& PackageLoadedDelegate, TSharedPtr<FLoadPackageAsyncProgressDelegate> PackageProgressDelegate)
 {
 	const FName FailedPackageName = PackageDesc.UPackageName;
 
@@ -10982,11 +10977,6 @@ void FAsyncLoadingThread2::QueueUnresolvedPackage(FAsyncLoadingThreadState2& Thr
 		};
 
 		PackageProgressDelegate->Invoke(Params);
-
-		if (Params.bReleaseDelegate)
-		{
-			verify(PackageProgressDelegate.Release());
-		}
 	}
 
 	if (PackageLoadedDelegate.IsValid())
diff --git a/Engine/Source/Runtime/CoreUObject/Private/UObject/UObjectGlobals.cpp b/Engine/Source/Runtime/CoreUObject/Private/UObject/UObjectGlobals.cpp
--- a/Engine/Source/Runtime/CoreUObject/Private/UObject/UObjectGlobals.cpp
+++ b/Engine/Source/Runtime/CoreUObject/Private/UObject/UObjectGlobals.cpp
@@ -1525,6 +1525,7 @@ int32 LoadAssetAsync(FTopLevelAssetPath InAssetPath, FLoadAssetAsyncDelegate InC
 	uint32 LoadFlags = InOptionalParams.LoadFlags;
 	
 	FLoadPackageAsyncOptionalParams PackageParams{
+		.ProgressDelegate = InOptionalParams.ProgressDelegate,
 		.PackagePriority = InOptionalParams.PackagePriority,
 		.InstancingContext = InOptionalParams.InstancingContext,
 		.LoadFlags = LoadFlags
diff --git a/Engine/Source/Runtime/CoreUObject/Public/UObject/UObjectGlobals.h b/Engine/Source/Runtime/CoreUObject/Public/UObject/UObjectGlobals.h
--- a/Engine/Source/Runtime/CoreUObject/Public/UObject/UObjectGlobals.h
+++ b/Engine/Source/Runtime/CoreUObject/Public/UObject/UObjectGlobals.h
@@ -713,8 +713,6 @@ struct FLoadPackageAsyncProgressParams
 	const UPackage* LoadedPackage { nullptr };
 	/* Progress of async loading. */
 	const EAsyncLoadingProgress ProgressType { EAsyncLoadingProgress::Failed };
-	/* Whether the current delegate should be released or destroyed when not required anymore. */
-	bool bReleaseDelegate { false };
 };
 
 /**
@@ -763,7 +761,7 @@ struct FLoadPackageAsyncOptionalParams
 	/** Non Thread-safe delegate to be invoked from game-thread on completion. **/
 	TUniquePtr<FLoadPackageAsyncDelegate> CompletionDelegate;
 	/** Thread-safe delegate to be invoked at different state of progress for the given package. **/
-	TUniquePtr<FLoadPackageAsyncProgressDelegate> ProgressDelegate;
+	TSharedPtr<FLoadPackageAsyncProgressDelegate> ProgressDelegate;
 	/** Package flags used to construct loaded package in memory. **/
 	EPackageFlags PackageFlags { PKG_None };
 	/** Play in Editor instance ID. **/
@@ -851,6 +849,8 @@ DECLARE_DELEGATE_ThreeParams(FLoadAssetAsyncDelegate, const FTopLevelAssetPath&
  */
 struct FLoadAssetAsyncOptionalParams
 {
+	/** Thread-safe delegate to be invoked at different state of progress for the given package. **/
+	TSharedPtr<FLoadPackageAsyncProgressDelegate> ProgressDelegate;
 	/** Loading priority. **/
 	int32 PackagePriority { 0 };
 	/** Additional context to map object names to their instanced counterpart when loading an instanced package. **/
diff --git a/Engine/Source/Runtime/Engine/Private/Tests/Loading/UnattendedAsyncLoadingTests.cpp b/Engine/Source/Runtime/Engine/Private/Tests/Loading/UnattendedAsyncLoadingTests.cpp
--- a/Engine/Source/Runtime/Engine/Private/Tests/Loading/UnattendedAsyncLoadingTests.cpp
+++ b/Engine/Source/Runtime/Engine/Private/Tests/Loading/UnattendedAsyncLoadingTests.cpp
@@ -112,7 +112,7 @@ namespace UnattendedLoadTestImpl
 				UE_LOG(LogTemp, Display, TEXT("Loading %s (Remaining: %d)"), *PackageName.ToString(), PackagesToLoad.Num());
 				LoadPackageAsync(PackageName.ToString(),
 					{
-						.ProgressDelegate = MakeUnique<FLoadPackageAsyncProgressDelegate>(
+						.ProgressDelegate = MakeShared<FLoadPackageAsyncProgressDelegate>(
 							FLoadPackageAsyncProgressDelegate::DelegateType::CreateLambda(
 								[this, ExecutionCheck = FExecutionCheck()](FLoadPackageAsyncProgressParams& Params) mutable
 								{
