- Optimization when getting a class's config filename to use the existing one if it has already been calculated

diff --git a/Engine/Source/Runtime/CoreUObject/Private/UObject/Class.cpp b/Engine/Source/Runtime/CoreUObject/Private/UObject/Class.cpp
--- a/Engine/Source/Runtime/CoreUObject/Private/UObject/Class.cpp
+++ b/Engine/Source/Runtime/CoreUObject/Private/UObject/Class.cpp
@@ -7036,6 +7036,13 @@ const FString UClass::GetConfigName() const
 	}
 	else
 	{
+		// if the branch was already loaded, we have the IniPath right there
+		FConfigBranch* Branch = GConfig->FindBranch(*ClassConfigName.ToString(), FString());
+		if (Branch != nullptr)
+		{
+			return Branch->IniPath;
+		}
+
 		// generate the class ini name, and make sure it's up to date
 		FString ConfigGameName;
 		FConfigContext::ReadIntoGConfig().Load(*ClassConfigName.ToString(), ConfigGameName);
