Changes to garbage collector behavior related to borrowed objects:
- borrowed objects will always have garbage reference elimination disabled regardless if they're native or not and regardless of reference elimination being enabled or not
- borrowed objects will not be garbage collected until they are returned to the owning server

Refactored remote object related EInternalObject flags to be only compiled with remote object support enabled. Also, because we now require one additional flag (Borrowed) we ran out of available flag bits when compiling with packed FUObjectItem so making these two features incompatible for now.

diff --git a/Engine/Source/Programs/Shared/EpicGames.Core/UnrealEngineTypes.cs b/Engine/Source/Programs/Shared/EpicGames.Core/UnrealEngineTypes.cs
--- a/Engine/Source/Programs/Shared/EpicGames.Core/UnrealEngineTypes.cs
+++ b/Engine/Source/Programs/Shared/EpicGames.Core/UnrealEngineTypes.cs
@@ -683,16 +683,6 @@ namespace EpicGames.Core
 		/// </summary>
 		ReachabilityFlag2 = 1 << 16,
 
-		/// <summary>
-		/// Object is no longer owned by this process
-		/// </summary>
-		Remote = 1 << 17,
-
-		/// <summary>
-		/// Object referenced by remote process
-		/// </summary>
-		RemoteReference = 1 << 18,
-
 		/// <summary>
 		/// Object is ready to be imported by another package during loading
 		/// </summary>
diff --git a/Engine/Source/Runtime/CoreUObject/Private/UObject/RemoteObject.cpp b/Engine/Source/Runtime/CoreUObject/Private/UObject/RemoteObject.cpp
--- a/Engine/Source/Runtime/CoreUObject/Private/UObject/RemoteObject.cpp
+++ b/Engine/Source/Runtime/CoreUObject/Private/UObject/RemoteObject.cpp
@@ -197,6 +197,11 @@ void InitRemoteObjects()
 	{
 		UE::RemoteExecutor::FetchNextDeferredRPCDelegate.BindStatic([]() { return TOptional<TTuple<FName, FRemoteWorkPriority, bool, TFunction<void(void)>>>(); });
 	}
+
+#if UE_WITH_REMOTE_OBJECT_HANDLE
+	// We rely on garbage elimination begin disabled because we don't allow it inside of borrowed objects
+	checkf(UObjectBaseUtility::IsGarbageEliminationEnabled() == false, TEXT("Remote object support requires garbage elimination to be disabled"));
+#endif
 }
 
 void RegisterRemoteObjectId(FRemoteObjectId ObjectId, FRemoteServerId ResidentServerId)
@@ -207,7 +212,9 @@ void RegisterRemoteObjectId(FRemoteObjectId ObjectId, FRemoteServerId ResidentSe
 void RegisterSharedObject(UObject* Object)
 {
 	FUObjectItem* ObjectItem = GUObjectArray.ObjectToObjectItem(Object);
+#if UE_WITH_REMOTE_OBJECT_HANDLE
 	ObjectItem->SetFlags(EInternalObjectFlags::RemoteReference);
+#endif
 }
 
 void MarkAsRemote(UObject* Object, FRemoteServerId DestinationServerId)
@@ -216,18 +223,56 @@ void MarkAsRemote(UObject* Object, FRemoteServerId DestinationServerId)
 	checkf(!Object->IsTemplate(), TEXT("Attempted to Migrate Template Object '%s' which is considered an asset and never allowed to migrate"), *GetNameSafe(Object));
 
 	FUObjectItem* ObjectItem = GUObjectArray.ObjectToObjectItem(Object);
+#if UE_WITH_REMOTE_OBJECT_HANDLE
 	ObjectItem->SetFlags(EInternalObjectFlags::Remote);
-	ObjectItem->ClearFlags(EInternalObjectFlags_RootFlags | EInternalObjectFlags::RemoteReference);
+	ObjectItem->ClearFlags(EInternalObjectFlags_RootFlags | EInternalObjectFlags::RemoteReference | EInternalObjectFlags::Borrowed);
+#endif
 	ObjectMaps->FindOrAddRemoteObjectStub(Object, DestinationServerId);
 }
 
+void MarkAsRemoteReference(UObject* Object)
+{
+#if UE_WITH_REMOTE_OBJECT_HANDLE
+	FUObjectItem* ObjectItem = GUObjectArray.ObjectToObjectItem(Object);
+	ObjectItem->SetFlags(EInternalObjectFlags::RemoteReference);
+#endif
+}
+
+bool IsRemoteReference(const UObject* Object)
+{
+#if UE_WITH_REMOTE_OBJECT_HANDLE
+	return Object->HasAnyInternalFlags(EInternalObjectFlags::RemoteReference);
+#else
+	return false;
+#endif
+}
+
+void MarkAsBorrowed(UObject* Object)
+{
+#if UE_WITH_REMOTE_OBJECT_HANDLE
+	FUObjectItem* ObjectItem = GUObjectArray.ObjectToObjectItem(Object);
+	ObjectItem->SetFlags(EInternalObjectFlags::Borrowed);
+#endif
+}
+
+bool IsBorrowed(const UObject* Object)
+{
+#if UE_WITH_REMOTE_OBJECT_HANDLE
+	return Object->HasAnyInternalFlags(EInternalObjectFlags::Borrowed);
+#else
+	return false;
+#endif
+}
+
 void MarkAsLocal(UObject* Object)
 {
 	ensureMsgf(!Object->HasAnyFlags(EObjectFlags::RF_ClassDefaultObject | EObjectFlags::RF_ArchetypeObject), TEXT("We're about to set an ArchetypeObject %s as remote reference"), *GetNameSafe(Object));
 
+#if UE_WITH_REMOTE_OBJECT_HANDLE
 	FUObjectItem* ObjectItem = GUObjectArray.ObjectToObjectItem(Object);
 	ObjectItem->ClearFlags(EInternalObjectFlags::Remote);
 	ObjectItem->SetFlags(EInternalObjectFlags::RemoteReference);
+#endif
 	ObjectMaps->FindOrAddRemoteObjectStub(Object, GetGlobalServerId());
 }
 
@@ -328,9 +373,13 @@ bool IsRemote(FRemoteObjectId ObjectId)
 
 bool IsRemote(const UObject* Object)
 {
+#if UE_WITH_REMOTE_OBJECT_HANDLE
 	const int32 InternalIndex = GUObjectArray.ObjectToIndex(Object);
 	const bool bIsRemote = InternalIndex >= 0 && GUObjectArray.IndexToObject(InternalIndex)->HasAnyFlags(EInternalObjectFlags::Remote);
 	return bIsRemote;
+#else
+	return false;
+#endif
 }
 
 bool IsOwned(const UObject* Object)
diff --git a/Engine/Source/Runtime/CoreUObject/Private/UObject/RemoteObjectPrivate.h b/Engine/Source/Runtime/CoreUObject/Private/UObject/RemoteObjectPrivate.h
--- a/Engine/Source/Runtime/CoreUObject/Private/UObject/RemoteObjectPrivate.h
+++ b/Engine/Source/Runtime/CoreUObject/Private/UObject/RemoteObjectPrivate.h
@@ -24,11 +24,31 @@ namespace UE::RemoteObject::Private
 	*/
 	void MarkAsRemote(UObject* Object, FRemoteServerId DestinationServerId);
 
+	/**
+	* Marks object as referenced by a remote object
+	*/
+	void MarkAsRemoteReference(UObject* Object);
+
+	/**
+	* @return true if object is marked as remote reference
+	*/
+	bool IsRemoteReference(const UObject* Object);
+
 	/**
 	* Marks object memory as local
 	*/
 	void MarkAsLocal(UObject* Object);
 
+	/**
+	* Marks object as borrowed
+	*/
+	void MarkAsBorrowed(UObject* Object);
+
+	/**
+	* Returns true if the specified object is marked as borrowed
+	*/
+	bool IsBorrowed(const UObject* Object);
+
 	/**
 	* Registers a stub for a remote object that is known to be resident on a specific server
 	*/
diff --git a/Engine/Source/Runtime/CoreUObject/Private/UObject/RemoteObjectSerialization.cpp b/Engine/Source/Runtime/CoreUObject/Private/UObject/RemoteObjectSerialization.cpp
--- a/Engine/Source/Runtime/CoreUObject/Private/UObject/RemoteObjectSerialization.cpp
+++ b/Engine/Source/Runtime/CoreUObject/Private/UObject/RemoteObjectSerialization.cpp
@@ -450,6 +450,7 @@ protected:
 	void WriteObjectPtr(const FObjectPtr& ObjPtr)
 	{
 		using namespace UE::CoreUObject::Private;
+		using namespace UE::RemoteObject::Private;
 
 		FRemoteObjectReferenceInfo Info = GetReferenceInfo(ObjPtr);
 		WriteObjectReference(Info);
@@ -460,7 +461,7 @@ protected:
 
 			// Anything can be marked as a remote reference, even assets in which case we rely on this flag to be set so that GC calls
 			// StoreObjectToDatabase for any remotely referenced asset (and only for remotely referenced assets) that's about to be GC'd
-			if (!Object->HasAnyInternalFlags(EInternalObjectFlags::RemoteReference))
+			if (!IsRemoteReference(Object))
 			{
 				if (ReferencedObjectsSet)
 				{
@@ -733,6 +734,8 @@ class FArchiveRemoteObjectReader : public TArchiveRemoteObjectBase<FMemoryReader
 
 	ERemoteReferenceType ReadObjectReference(FObjectPtr& Value)
 	{
+		using namespace UE::RemoteObject::Private;
+
 		ERemoteReferenceType Type;		
 		*this << Type;
 
@@ -764,7 +767,7 @@ class FArchiveRemoteObjectReader : public TArchiveRemoteObjectBase<FMemoryReader
 					UObject* Obj = ResolvedPathNameObjects[PathNameIndex];
 					if (Obj)
 					{
-						Obj->SetInternalFlags(EInternalObjectFlags::RemoteReference);
+						MarkAsRemoteReference(Obj);
 						Value = FObjectPtr(Obj);
 						bNeedsResolvingWithId = false;
 					}
diff --git a/Engine/Source/Runtime/CoreUObject/Private/UObject/RemoteObjectTransfer.cpp b/Engine/Source/Runtime/CoreUObject/Private/UObject/RemoteObjectTransfer.cpp
--- a/Engine/Source/Runtime/CoreUObject/Private/UObject/RemoteObjectTransfer.cpp
+++ b/Engine/Source/Runtime/CoreUObject/Private/UObject/RemoteObjectTransfer.cpp
@@ -217,7 +217,7 @@ public:
 
 					for (UObject* RO : ReferencedObjects)
 					{
-						RO->SetInternalFlags(EInternalObjectFlags::RemoteReference);
+						MarkAsRemoteReference(RO);
 					}
 
 					for (UObject* SO : SentObjects)
@@ -584,7 +584,7 @@ public:
 			FRemoteObjectId ReferencedObjectId = UE::CoreUObject::Private::FObjectHandleUtils::GetRemoteId(ReferencedObject);
 			UE_LOG(LogRemoteObject, VeryVerbose, TEXT("FRemoteObjectTransferQueue: CommitMultiServerCommit setting RemoteReference %s"), *ReferencedObjectId.ToString());
 
-			ReferencedObject->SetInternalFlags(EInternalObjectFlags::RemoteReference);
+			MarkAsRemoteReference(ReferencedObject);
 		}
 
 		UE_AUTORTFM_OPEN
@@ -800,11 +800,14 @@ FRemoteObjectTransferQueue GTransferQueue;
 UObject* MigrateObjectFromObjectData(FRemoteObjectData& ObjectData, const FUObjectMigrationContext& MigrationContext)
 {
 	using namespace UE::RemoteObject::Serialization;
+	using namespace UE::RemoteObject::Private;
+
 	check(!AutoRTFM::IsClosed());
 
 	// Notify that object data has been migrated from a remote server
 	OnObjectDataReceivedDelegate.Broadcast(ObjectData, MigrationContext);
 
+	FRemoteServerId LocalServerId = UE::RemoteObject::GetGlobalServerId();
 	TArray<FRemoteObjectId> ReceivedObjectRemoteIds;
 	TArray<UObject*> ReceivedObjects;
 	int32 RequestedObjectIndex = DeserializeObjectData(ObjectData, &MigrationContext, ReceivedObjectRemoteIds, ReceivedObjects);
@@ -838,8 +841,6 @@ UObject* MigrateObjectFromObjectData(FRemoteObjectData& ObjectData, const FUObje
 		const bool bSameNumberOfObjects = (1 + AllSubObjects.Num() == ReceivedObjects.Num());
 		if (!bSameNumberOfObjects)
 		{
-			FRemoteServerId LocalServerId = UE::RemoteObject::GetGlobalServerId();
-
 			UE_LOG(LogRemoteObject, Warning, TEXT("While migrating '%s' [%s -> %s] we received %d Objects whereas %d Objects got instantiated after deserialization.  Differences:"), *GetNameSafe(RootObject), *MigrationContext.RemoteServerId.ToString(), *LocalServerId.ToString(), ReceivedObjects.Num(), 1 + AllSubObjects.Num());
 			TArray<UObject*> ReceivedButNotInstanced;
 			for (UObject* Obj : ReceivedObjects)
@@ -869,6 +870,13 @@ UObject* MigrateObjectFromObjectData(FRemoteObjectData& ObjectData, const FUObje
 	}
 
 #if UE_WITH_REMOTE_OBJECT_HANDLE
+	if (LocalServerId != MigrationContext.OwnerServerId)
+	{
+		for (UObject* ReceivedObject : ReceivedObjects)
+		{
+			MarkAsBorrowed(ReceivedObject);
+		}
+	}
 	for (UObject* ReceivedObject : ReceivedObjects)
 	{
 		ReceivedObject->PostMigrate(MigrationContext);
@@ -880,6 +888,15 @@ UObject* MigrateObjectFromObjectData(FRemoteObjectData& ObjectData, const FUObje
 	{
 		UE::RemoteObject::Handle::ChangeOwnerServerId(ReceivedObject, MigrationContext.OwnerServerId);
 		UE::RemoteObject::Handle::ChangePhysicsServerId(ReceivedObject, MigrationContext.PhysicsServerId);
+
+#if UE_WITH_REMOTE_OBJECT_HANDLE
+		ensureAlwaysMsgf(ReceivedObject->HasAnyInternalFlags(EInternalObjectFlags::Borrowed) == !UE::RemoteObject::Handle::IsOwned(ReceivedObject),
+			TEXT("Object %s (%s) is %s owned but is %s borrowed"),
+			*ReceivedObject->GetPathName(),
+			*FRemoteObjectId(ReceivedObject).ToString(),
+			UE::RemoteObject::Handle::IsOwned(ReceivedObject) ? TEXT("") : TEXT("not"),
+			IsBorrowed(ReceivedObject) ? TEXT("") : TEXT("not"));
+#endif
 	}
 
 	// Notify that objects have been migrated from a remote server
@@ -1055,7 +1072,7 @@ void TransferObjectOwnershipToRemoteServer(UObject* Object, FRemoteServerId Dest
 
 	for (UObject* RO : ReferencedObjects)
 	{
-		RO->SetInternalFlags(EInternalObjectFlags::RemoteReference);
+		MarkAsRemoteReference(RO);
 	}
 
 	for (UObject* SO : SentObjects)
@@ -1197,6 +1214,7 @@ namespace UE::RemoteObject::Transfer::Private
 UObject* GetOutermostUnreachableRemoteReference(UObject* Object)
 {
 	UObject* OutermostUnreachableRemoteReference = Object;
+#if UE_WITH_REMOTE_OBJECT_HANDLE
 	for (UObject* Outer = Object->GetOuter(); Outer; Outer = Outer->GetOuter())
 	{
 		FUObjectItem* OuterItem = GUObjectArray.ObjectToObjectItem(Outer);
@@ -1209,6 +1227,7 @@ UObject* GetOutermostUnreachableRemoteReference(UObject* Object)
 			break;
 		}
 	}
+#endif
 	return OutermostUnreachableRemoteReference;
 }
 
@@ -1280,7 +1299,7 @@ void StoreObjectToDatabase(UObject* Object)
 
 		for (UObject* ReferencedObject : ReferencedObjects)
 		{
-			ReferencedObject->SetInternalFlags(EInternalObjectFlags::RemoteReference);
+			MarkAsRemoteReference(ReferencedObject);
 		}
 
 		// MarkAsRemote clears EInternalObjectFlags::RemoteReference flag which will prevent re-enetring this loop with the same outermost object but 
@@ -1291,10 +1310,10 @@ void StoreObjectToDatabase(UObject* Object)
 		// in which case we need to repeat this loop until Object is marked as remote (and is no longer marked as RemoteReference).
 		
 		// Assert if we're about to enter an infinite loop.
-		checkf(!Object->HasAnyInternalFlags(EInternalObjectFlags::RemoteReference) || OutermostUnreachableRemoteReference != Object,
+		checkf(!IsRemoteReference(Object) || OutermostUnreachableRemoteReference != Object,
 			TEXT("PeerId:%s failed to store %s to database"), *GetGlobalServerId().ToString(), *Object->GetPathName());
 
-	} while (Object->HasAnyInternalFlags(EInternalObjectFlags::RemoteReference));
+	} while (IsRemoteReference(Object));
 }
 
 } // namespace UE::RemoteObject::Transfer::Private
diff --git a/Engine/Source/Runtime/CoreUObject/Private/UObject/UObjectGlobals.cpp b/Engine/Source/Runtime/CoreUObject/Private/UObject/UObjectGlobals.cpp
--- a/Engine/Source/Runtime/CoreUObject/Private/UObject/UObjectGlobals.cpp
+++ b/Engine/Source/Runtime/CoreUObject/Private/UObject/UObjectGlobals.cpp
@@ -3134,7 +3134,11 @@ UObject* StaticDuplicateObjectEx( FObjectDuplicationParameters& Parameters )
 	
 	UE_CLOG(FPlatformProperties::RequiresCookedData() && Parameters.SourceObject->HasAnyInternalFlags(EInternalObjectFlags_AsyncLoading), LogUObjectGlobals, Warning, TEXT("Duplicating object '%s' that's still being async loaded"), *Parameters.SourceObject->GetFullName());
 	// Make sure we're not duplicating the AsyncLoading, Async or LoaderImport internal flags, they will prevent the object from being gcd.
-	Parameters.InternalFlagMask &= ~(EInternalObjectFlags::Async | EInternalObjectFlags::LoaderImport | EInternalObjectFlags_AsyncLoading | EInternalObjectFlags::RemoteReference);
+	Parameters.InternalFlagMask &= ~(EInternalObjectFlags::Async | EInternalObjectFlags::LoaderImport | EInternalObjectFlags_AsyncLoading 
+#if UE_WITH_REMOTE_OBJECT_HANDLE
+		| EInternalObjectFlags::RemoteReference
+#endif
+		);
 
 	// We can't modify the loader from a transaction, so check for async loading and reset the 
 	// loaders in the open.
diff --git a/Engine/Source/Runtime/CoreUObject/Public/UObject/FastReferenceCollector.h b/Engine/Source/Runtime/CoreUObject/Public/UObject/FastReferenceCollector.h
--- a/Engine/Source/Runtime/CoreUObject/Public/UObject/FastReferenceCollector.h
+++ b/Engine/Source/Runtime/CoreUObject/Public/UObject/FastReferenceCollector.h
@@ -659,12 +659,27 @@ FORCEINLINE void CallSlowARO(DispatcherType&, uint32 SlowAROIdx, uint8* Instance
 	LogIllegalTypeFatal(EMemberType::SlowARO, MemberIdx, Instance);
 }
 
+template <typename ObjectType>
+FORCEINLINE EOrigin GetSchemaOrigin(FSchemaView Schema, ObjectType* Instance)
+{
+	return Schema.GetOrigin();
+}
+#if UE_WITH_REMOTE_OBJECT_HANDLE
+template <>
+FORCEINLINE EOrigin GetSchemaOrigin(FSchemaView Schema, UObject* Instance)
+{
+	FUObjectItem* Item = GUObjectArray.ObjectToObjectItem(Instance);
+	// Consider borrowed objects to always be non-BP to prevent garbage elimination from mutating these objects
+	return !Item->HasAnyFlags(EInternalObjectFlags::Borrowed) ? Schema.GetOrigin() : EOrigin::Other;
+}
+#endif // UE_WITH_REMOTE_OBJECT_HANDLE
+
 template<class DispatcherType, typename ObjectType>
 FORCEINLINE_DEBUGGABLE void VisitMembers(DispatcherType& Dispatcher, FSchemaView Schema, ObjectType* Instance)
 {
 	check(!Schema.IsEmpty());
 
-	const EOrigin Origin = Schema.GetOrigin();
+	const EOrigin Origin = GetSchemaOrigin(Schema, Instance);
 	uint64* InstanceCursor = (uint64*)Instance;	// Advanced via Jump to reach far members
 	uint32 DebugIdx = 0;
 	for (const FMemberWord* WordIt = Schema.GetWords(); true; ++WordIt)
diff --git a/Engine/Source/Runtime/CoreUObject/Public/UObject/ObjectMacros.h b/Engine/Source/Runtime/CoreUObject/Public/UObject/ObjectMacros.h
--- a/Engine/Source/Runtime/CoreUObject/Public/UObject/ObjectMacros.h
+++ b/Engine/Source/Runtime/CoreUObject/Public/UObject/ObjectMacros.h
@@ -620,15 +620,27 @@ enum class EInternalObjectFlags : int32
 {
 	None = 0,
 
-	// NOTE that the minimum flag bit index is currently 14. See EInternalObjectFlags_MinFlagBitIndex and UE_ENABLE_FUOBJECT_ITEM_PACKING in UObjectArray.h
+#if UE_WITH_REMOTE_OBJECT_HANDLE
+	Remote = 1 << 0, //< Object is no longer owned by this process
+	RemoteReference = 1 << 1, //< Object referenced by remote process
+	Borrowed = 1 << 2, //< Object that is not owned by this process 
+#else
+	Remote = 0,
+	RemoteReference = 0,
+	Borrowed = 0,
+#endif
+
+
+	// NOTE that the minimum flag bit index is currently 14 (unless compiling with UE_WITH_REMOTE_OBJECT_HANDLE). 
+	// See EInternalObjectFlags_MinFlagBitIndex and UE_ENABLE_FUOBJECT_ITEM_PACKING in UObjectArray.h
 
 	ReachabilityFlag0 = 1 << 14, ///< One of the flags used by Garbage Collector to determine UObject's reachability state
 	ReachabilityFlag1 = 1 << 15, ///< One of the flags used by Garbage Collector to determine UObject's reachability state
 	ReachabilityFlag2 = 1 << 16, ///< One of the flags used by Garbage Collector to determine UObject's reachability state
 
 	AutoRTFMConstructionAborted = 1 << 17, //< Object was constructed in an AutoRTFM transaction that has been aborted
-	Remote = 1 << 18, //< Object is no longer owned by this process
-	RemoteReference = 1 << 19, //< Object referenced by remote process
+	// Unused = 1 << 18,
+	// Unused = 1 << 19,
 	LoaderImport = 1 << 20, ///< Object is ready to be imported by another package during loading
 	Garbage = 1 << 21, ///< Garbage from logical point of view and should not be referenced. This flag is mirrored in EObjectFlags as RF_Garbage for performance
 	AsyncLoadingPhase1 = 1 << 22, ///< Object is being asynchronously loaded.
@@ -650,8 +662,8 @@ ENUM_CLASS_FLAGS(EInternalObjectFlags);
 
 //~ Make sure these macros are up to date!
 #define EInternalObjectFlags_MinFlagBitIndex 14
-#define EInternalObjectFlags_GarbageCollectionKeepFlags (EInternalObjectFlags::Native | EInternalObjectFlags::Async | EInternalObjectFlags::AsyncLoadingPhase1 | EInternalObjectFlags::AsyncLoadingPhase2 | EInternalObjectFlags::LoaderImport | EInternalObjectFlags::RefCounted)
-#define EInternalObjectFlags_AllFlags (EInternalObjectFlags::ReachabilityFlag0 | EInternalObjectFlags::ReachabilityFlag1 | EInternalObjectFlags::ReachabilityFlag2 | EInternalObjectFlags::AutoRTFMConstructionAborted | EInternalObjectFlags::Remote | EInternalObjectFlags::RemoteReference | EInternalObjectFlags::LoaderImport | EInternalObjectFlags::Garbage | EInternalObjectFlags::ReachableInCluster | EInternalObjectFlags::ClusterRoot | EInternalObjectFlags::Native | EInternalObjectFlags::RefCounted | EInternalObjectFlags::Async | EInternalObjectFlags::AsyncLoadingPhase1 | EInternalObjectFlags::AsyncLoadingPhase2 | EInternalObjectFlags::RootSet | EInternalObjectFlags::PendingConstruction | (EInternalObjectFlags::Unreachable))
+#define EInternalObjectFlags_GarbageCollectionKeepFlags (EInternalObjectFlags::Borrowed | EInternalObjectFlags::Native | EInternalObjectFlags::Async | EInternalObjectFlags::AsyncLoadingPhase1 | EInternalObjectFlags::AsyncLoadingPhase2 | EInternalObjectFlags::LoaderImport | EInternalObjectFlags::RefCounted)
+#define EInternalObjectFlags_AllFlags (EInternalObjectFlags::Remote | EInternalObjectFlags::RemoteReference | EInternalObjectFlags::Borrowed | EInternalObjectFlags::ReachabilityFlag0 | EInternalObjectFlags::ReachabilityFlag1 | EInternalObjectFlags::ReachabilityFlag2 | EInternalObjectFlags::AutoRTFMConstructionAborted | EInternalObjectFlags::LoaderImport | EInternalObjectFlags::Garbage | EInternalObjectFlags::ReachableInCluster | EInternalObjectFlags::ClusterRoot | EInternalObjectFlags::Native | EInternalObjectFlags::RefCounted | EInternalObjectFlags::Async | EInternalObjectFlags::AsyncLoadingPhase1 | EInternalObjectFlags::AsyncLoadingPhase2 | EInternalObjectFlags::RootSet | EInternalObjectFlags::PendingConstruction | EInternalObjectFlags::Unreachable)
 #define EInternalObjectFlags_RootFlags (EInternalObjectFlags::RootSet | EInternalObjectFlags_GarbageCollectionKeepFlags)
 #define EInternalObjectFlags_ReachabilityFlags (EInternalObjectFlags::ReachabilityFlag0 | EInternalObjectFlags::ReachabilityFlag1 | EInternalObjectFlags::ReachabilityFlag2 | EInternalObjectFlags::Unreachable)
 #define EInternalObjectFlags_AsyncLoading (EInternalObjectFlags::AsyncLoadingPhase1 | EInternalObjectFlags::AsyncLoadingPhase2)
diff --git a/Engine/Source/Runtime/CoreUObject/Public/UObject/UObjectArray.h b/Engine/Source/Runtime/CoreUObject/Public/UObject/UObjectArray.h
--- a/Engine/Source/Runtime/CoreUObject/Public/UObject/UObjectArray.h
+++ b/Engine/Source/Runtime/CoreUObject/Public/UObject/UObjectArray.h
@@ -23,11 +23,11 @@ namespace Verse
 
 #ifndef UE_PACK_FUOBJECT_ITEM
 #	define UE_ENABLE_FUOBJECT_ITEM_PACKING 0
-#elif defined(UE_PACK_FUOBJECT_ITEM) && defined(PLATFORM_CPU_ARM_FAMILY) && defined(FORCE_ANSI_ALLOCATOR)
+#elif defined(UE_PACK_FUOBJECT_ITEM) && defined(PLATFORM_CPU_ARM_FAMILY) && defined(FORCE_ANSI_ALLOCATOR) && !UE_WITH_REMOTE_OBJECT_HANDLE
 //	disable packing on ARM64 is we use Ansi allocator as it might use MTE or HWASan which uses top byte of a pointer and this feature will discard that
 #	define UE_ENABLE_FUOBJECT_ITEM_PACKING 0
 #elif
-#	define UE_ENABLE_FUOBJECT_ITEM_PACKING UE_PACK_FUOBJECT_ITEM
+#	define UE_ENABLE_FUOBJECT_ITEM_PACKING UE_PACK_FUOBJECT_ITEM && !UE_WITH_REMOTE_OBJECT_HANDLE
 #endif
 
 /**
diff --git a/Engine/Source/Runtime/CoreUObject/Public/UObject/UObjectIterator.h b/Engine/Source/Runtime/CoreUObject/Public/UObject/UObjectIterator.h
--- a/Engine/Source/Runtime/CoreUObject/Public/UObject/UObjectIterator.h
+++ b/Engine/Source/Runtime/CoreUObject/Public/UObject/UObjectIterator.h
@@ -58,7 +58,11 @@ enum class EObjectIteratorThreadSafetyOptions : uint8
 
 inline EInternalObjectFlags GetObjectIteratorDefaultInternalExclusionFlags(EInternalObjectFlags InternalExclusionFlags)
 {
-	return InternalExclusionFlags | EInternalObjectFlags::Unreachable | EInternalObjectFlags::AutoRTFMConstructionAborted | EInternalObjectFlags::PendingConstruction | EInternalObjectFlags::Remote | UE::GetAsyncLoadingInternalFlagsExclusion();
+	return InternalExclusionFlags | EInternalObjectFlags::Unreachable | EInternalObjectFlags::AutoRTFMConstructionAborted | EInternalObjectFlags::PendingConstruction | UE::GetAsyncLoadingInternalFlagsExclusion()
+#if UE_WITH_REMOTE_OBJECT_HANDLE
+		| EInternalObjectFlags::Remote
+#endif
+		;
 }
 
 /**
diff --git a/Engine/Source/Runtime/Engine/Private/UnrealEngine.cpp b/Engine/Source/Runtime/Engine/Private/UnrealEngine.cpp
--- a/Engine/Source/Runtime/Engine/Private/UnrealEngine.cpp
+++ b/Engine/Source/Runtime/Engine/Private/UnrealEngine.cpp
@@ -9405,6 +9405,9 @@ bool UEngine::HandleObjCommand( const TCHAR* Cmd, FOutputDevice& Ar )
 			const bool bOnlyListDefaultObjects = FParse::Param(Cmd, TEXT("DEFAULTSONLY"));
 			const bool bShowDetailedObjectInfo = FParse::Param(Cmd, TEXT("NODETAILEDINFO")) == false && bTrackDetailedObjectInfo;
 			const bool bOnlyShowDisregardGC = FParse::Param(Cmd, TEXT("GCDISREGARDONLY"));
+#if UE_WITH_REMOTE_OBJECT_HANDLE
+			const bool bOnlyListBorrowedObjects = FParse::Param(Cmd, TEXT("BORROWED"));
+#endif
 
 			for( FThreadSafeObjectIterator It; It; ++It )
 			{
@@ -9457,6 +9460,14 @@ bool UEngine::HandleObjCommand( const TCHAR* Cmd, FOutputDevice& Ar )
 					continue;
 				}
 
+#if UE_WITH_REMOTE_OBJECT_HANDLE
+				if (bOnlyListBorrowedObjects && !UE::RemoteObject::Handle::IsOwned(*It))
+				{
+					ensureAlwaysMsgf(It->HasAnyInternalFlags(EInternalObjectFlags::Borrowed), TEXT("Object %s is not owned by this server but it's not borrowed either"), *It->GetPathName());
+					continue;
+				}
+#endif
+
 				if ( CheckType && !It->IsA(CheckType) )
 				{
 					continue;
