Fix memory leak when running BP functions added to a UHT native verse class.

#rnx

diff --git a/Engine/Source/Runtime/CoreUObject/Private/UObject/Class.cpp b/Engine/Source/Runtime/CoreUObject/Private/UObject/Class.cpp
--- a/Engine/Source/Runtime/CoreUObject/Private/UObject/Class.cpp
+++ b/Engine/Source/Runtime/CoreUObject/Private/UObject/Class.cpp
@@ -1067,8 +1067,11 @@ void UStruct::Link(FArchive& Ar, bool bRelinkExistingProperties)
 			RefLinkBuilder.AppendNoTerminate(*Property);
 		}
 
+		// We are checking for Script content to cover the case where a UFunction using script VM has been added
+		// to a native class.  This is the case for UHT generated Verse classes.  The Verse compiler will add the
+		// Verse script implementations to the UHT generated class.
 		const UClass* OwnerClass = Property->GetOwnerClass();
-		bool bOwnedByNativeClass = OwnerClass && OwnerClass->HasAnyClassFlags(CLASS_Native | CLASS_Intrinsic);
+		bool bOwnedByNativeClass = Script.IsEmpty() && OwnerClass && OwnerClass->HasAnyClassFlags(CLASS_Native | CLASS_Intrinsic);
 		bool bShouldHandleFinishDestroy = Property->ContainsFinishDestroy(EncounteredStructProps);
 
 		if ((!Property->HasAnyPropertyFlags(CPF_IsPlainOldData | CPF_NoDestructor)
