Removed UE_OBJECT_PTR_NONCONFORMANCE_SUPPORT, which no longer seems to be a problem.

diff --git a/Engine/Source/Runtime/CoreUObject/Public/UObject/ObjectPtr.h b/Engine/Source/Runtime/CoreUObject/Public/UObject/ObjectPtr.h
--- a/Engine/Source/Runtime/CoreUObject/Public/UObject/ObjectPtr.h
+++ b/Engine/Source/Runtime/CoreUObject/Public/UObject/ObjectPtr.h
@@ -40,16 +40,7 @@
 #define UE_TRANSITIONAL_OBJECT_PTR_TEMPLATE2_ARG_BOTH(Templ, Type1, Type2) auto 
 #define UE_TRANSITIONAL_OBJECT_PTR_TEMPLATE2_ARG_BOTH_SUFFIXED(Templ, Type1, Type2, Suffix) auto 
 
-#if PLATFORM_MICROSOFT && defined(_MSC_EXTENSIONS)
-	/**
-	 * Non-conformance mode in MSVC has issues where the presence of a conversion operator to bool (even an explicit one)
-	 * leads to operator ambiguity for operator== and operator!= with a NULL (not nullptr).  This macro controls the presence
-	 * of code to deal with this issue.
-	 */
-	#define UE_OBJECT_PTR_NONCONFORMANCE_SUPPORT 1
-#else
-	#define UE_OBJECT_PTR_NONCONFORMANCE_SUPPORT 0
-#endif
+#define UE_OBJECT_PTR_NONCONFORMANCE_SUPPORT 0 UE_DEPRECATED_MACRO(5.6, "UE_OBJECT_PTR_NONCONFORMANCE_SUPPORT has been deprecated and shouldn't be used")
 
 class UClass;
 
@@ -341,20 +332,6 @@ namespace ObjectPtr_Private
 		return Other;
 	}
 
-#if UE_OBJECT_PTR_NONCONFORMANCE_SUPPORT
-	/**
-	 * Force acceptance of explicitly TYPE_OF_NULL to avoid ambiguity issues triggered by the presence of a type conversion operator to bool.
-	 * This has the negative consequence of making it possible to coerce an arbitrary integer for the purpose of use in the comparison operators.
-	 */
-	template <typename T>
-	UE_OBJPTR_DEPRECATED(5.0, "Coercing a NULL for operations with a TObjectPtr is deprecated when running in a non-standards conforming compiler mode.")
-	constexpr const T* CoerceToPointer(TYPE_OF_NULL Other)
-	{
-		checkfSlow(Other == 0, TEXT("TObjectPtr cannot be compared to a non-zero NULL type value."));
-		return nullptr;
-	}
-#endif
-
 	/** Coerce to pointer through implicit conversion to CommonPointerType where CommonPointerType is deduced, and must be a C++ pointer, not a wrapper type. */
 	template <
 		typename T,
diff --git a/Engine/Source/Runtime/CoreUObject/Tests/ObjectPtrTest.cpp b/Engine/Source/Runtime/CoreUObject/Tests/ObjectPtrTest.cpp
--- a/Engine/Source/Runtime/CoreUObject/Tests/ObjectPtrTest.cpp
+++ b/Engine/Source/Runtime/CoreUObject/Tests/ObjectPtrTest.cpp
@@ -134,10 +134,9 @@ static_assert(!TModels_V<CEqualityComparableWith, FMutableObjectPtr, char*>, "Mu
 static_assert(TModels_V<CEqualityComparableWith, FConstObjectPtr, TYPE_OF_NULLPTR>, "Must be able to compare equality and inequality bidirectionally between TObjectPtr<const UObject> and nullptr");
 static_assert(TModels_V<CEqualityComparableWith, FMutableObjectPtr, TYPE_OF_NULLPTR>, "Must be able to compare equality and inequality bidirectionally between TObjectPtr<UObject> and nullptr");
 
-#if !UE_OBJECT_PTR_NONCONFORMANCE_SUPPORT // Specialized NULL support causes these checks to fail.
-static_assert(!TModels_V<CEqualityComparableWith, FConstObjectPtr, long>, "Should not be able to compare equality and inequality bidirectionally between TObjectPtr<const UObject> and long");
-static_assert(!TModels_V<CEqualityComparableWith, FMutableObjectPtr, long>, "Should not be able to compare equality and inequality bidirectionally between TObjectPtr<UObject> and long");
-#endif // #if !UE_OBJECT_PTR_NONCONFORMANCE_SUPPORT
+// Ensure allowing comparison against NULL doesn't mean allowing comparison against int
+static_assert(!TModels_V<CEqualityComparableWith, FConstObjectPtr, int>, "Should not be able to compare equality and inequality bidirectionally between TObjectPtr<const UObject> and long");
+static_assert(!TModels_V<CEqualityComparableWith, FMutableObjectPtr, int>, "Should not be able to compare equality and inequality bidirectionally between TObjectPtr<UObject> and long");
 
 
 #if WITH_LOW_LEVEL_TESTS
