Package MetaData
- Fix UObject::Rename to properly use FMetaDataUtilities::FMoveMetadataHelperContext when used in contexts that don't define WITH_EDITOR (replaced with WITH_METADATA), which is the case for FoundationTests.
- Fix FMetaDataUtilities::FMoveMetadataHelperContext to properly migrate metadata for renamed objects, supporting changing object name, outer and package.
- Implemented all known cases in new FoundationTests CoreUObject::UPackage::FMetaData.

diff --git a/Engine/Source/Runtime/CoreUObject/Private/UObject/MetaData.cpp b/Engine/Source/Runtime/CoreUObject/Private/UObject/MetaData.cpp
--- a/Engine/Source/Runtime/CoreUObject/Private/UObject/MetaData.cpp
+++ b/Engine/Source/Runtime/CoreUObject/Private/UObject/MetaData.cpp
@@ -18,7 +18,7 @@ DEFINE_LOG_CATEGORY_STATIC(LogMetaData, Log, All);
 //////////////////////////////////////////////////////////////////////////
 // FMetaDataUtilities
 
-#if WITH_EDITOR
+#if WITH_METADATA
 
 FAutoConsoleCommand FMetaDataUtilities::DumpAllConsoleCommand = FAutoConsoleCommand(
 	TEXT( "Metadata.Dump" ),
@@ -67,7 +67,7 @@ FMetaDataUtilities::FMoveMetadataHelperContext::FMoveMetadataHelperContext(UObje
 	if (GIsEditor)
 	{
 		check(SourceObject);
-		UPackage* Package = SourceObject->GetOutermost();
+		UPackage* Package = SourceObject->GetPackage();
 		check(Package);
 		OldPackage = Package;
 		OldObjectPath = FSoftObjectPath::ConstructFromObject(SourceObject);
@@ -79,37 +79,43 @@ FMetaDataUtilities::FMoveMetadataHelperContext::FMoveMetadataHelperContext(UObje
 FMetaDataUtilities::FMoveMetadataHelperContext::~FMoveMetadataHelperContext()
 {
 	// We only want to actually move things if we're in the editor
-	if (GIsEditor && (OldObject->GetOutermost() != OldPackage))
+	if (GIsEditor)
 	{
-		FMetaData& NewMetaData = OldObject->GetPackage()->GetMetaData();
-		FMetaData& OldMetaData = OldPackage->GetMetaData();
-		TMap< FName, FString>* OldObjectMetaData = OldMetaData.ObjectMetaDataMap.Find(OldObjectPath);
-		if (OldObjectMetaData != NULL)
-		{
-			NewMetaData.SetObjectValues(OldObject, *OldObjectMetaData);
-		}
-
-		if (bShouldSearchChildren)
+		FSoftObjectPath NewObjectPath = FSoftObjectPath::ConstructFromObject(OldObject);
+		if (NewObjectPath != OldObjectPath)
 		{
-			TArray<UObject*> Children;
-			GetObjectsWithOuter(OldObject, Children, true);
+			FMetaData& NewMetaData = OldObject->GetPackage()->GetMetaData();
+			FMetaData& OldMetaData = OldPackage->GetMetaData();
+			
+			TMap<FName, FString> OldObjectMetaData;
+			if (OldMetaData.ObjectMetaDataMap.RemoveAndCopyValue(OldObjectPath, OldObjectMetaData))
+			{
+				NewMetaData.SetObjectValues(OldObject, OldObjectMetaData);
+			}
 
-			for ( auto ChildIt = Children.CreateConstIterator(); ChildIt; ++ChildIt )
+			if (bShouldSearchChildren)
 			{
-				UObject* Child = *ChildIt;
-				FSoftObjectPath ChildPath = FSoftObjectPath::ConstructFromObject(Child);
-				ChildPath.SetPath(OldObjectPath.GetAssetPath(), ChildPath.GetSubPathUtf8String());
-				TMap<FName, FString>* OldObjectMetaDataPtr = OldMetaData.ObjectMetaDataMap.Find(ChildPath);
-				if ( OldObjectMetaDataPtr )
+				TArray<UObject*> Children;
+				GetObjectsWithOuter(OldObject, Children, true);
+
+				for ( auto ChildIt = Children.CreateConstIterator(); ChildIt; ++ChildIt )
 				{
-					NewMetaData.SetObjectValues(Child, *OldObjectMetaDataPtr);
+					UObject* Child = *ChildIt;
+					FSoftObjectPath ChildPath = FSoftObjectPath::ConstructFromObject(Child);
+					ChildPath.SetPath(OldObjectPath.GetAssetPath(), ChildPath.GetSubPathUtf8String());
+
+					TMap<FName, FString> ChildMetaData;
+					if (OldMetaData.ObjectMetaDataMap.RemoveAndCopyValue(ChildPath, ChildMetaData))
+					{
+						NewMetaData.SetObjectValues(Child, ChildMetaData);
+					}
 				}
 			}
 		}
 	}
 }
 
-#endif //WITH_EDITOR
+#endif //WITH_METADATA
 
 //////////////////////////////////////////////////////////////////////////
 // UDEPRECATED_MetaData implementation.
@@ -460,7 +466,7 @@ void FMetaData::RemoveValue(const UObject* Object, FName Key)
 TMap<FName, FString>* FMetaData::GetMapForObject(const UObject* Object)
 {
 	check(Object);
-	UPackage* Package = Object->GetOutermost();
+	UPackage* Package = Object->GetPackage();
 	check(Package);
 	FMetaData& Metadata = Package->GetMetaData();
 
@@ -481,7 +487,7 @@ void FMetaData::CopyMetadata(UObject* SourceObject, UObject* DestObject)
 	}
 
 	// Then get the metadata for the destination
-	UPackage* DestPackage = DestObject->GetOutermost();
+	UPackage* DestPackage = DestObject->GetPackage();
 	check(DestPackage);
 	FMetaData& DestMetadata = DestPackage->GetMetaData();
 
@@ -519,7 +525,7 @@ void FMetaData::RemoveMetaDataOutsidePackage(UPackage* MetaDataPackage)
 		FWeakObjectPtr ObjPtr = ObjPath.ResolveObject();
 
 		UObject* ObjectToRemove = ObjPtr.Get();
-		if ((ObjectToRemove != NULL) && (ObjectToRemove->GetOutermost() != GetTransientPackage()))
+		if ((ObjectToRemove != NULL) && (ObjectToRemove->GetPackage() != GetTransientPackage()))
 		{
 			UE_LOG(LogMetaData, Log, TEXT("Removing '%s' ref from Metadata '%s'"), *ObjectToRemove->GetPathName(), *MetaDataPackage->GetPathName());
 		}
diff --git a/Engine/Source/Runtime/CoreUObject/Private/UObject/Obj.cpp b/Engine/Source/Runtime/CoreUObject/Private/UObject/Obj.cpp
--- a/Engine/Source/Runtime/CoreUObject/Private/UObject/Obj.cpp
+++ b/Engine/Source/Runtime/CoreUObject/Private/UObject/Obj.cpp
@@ -237,11 +237,11 @@ bool UObject::Rename(const TCHAR* InName, UObject* NewOuter, ERenameFlags Flags)
 {
 	TRACE_CPUPROFILER_EVENT_SCOPE(UObject::Rename);
 
-#if WITH_EDITOR
+#if WITH_METADATA
 	// This guarantees that if this UObject is actually renamed and changes packages
 	// the metadata will be moved with it.
 	FMetaDataUtilities::FMoveMetadataHelperContext MoveMetaData(this, true);
-#endif //WITH_EDITOR
+#endif // WITH_METADATA
 
 	if (NewOuter && NewOuter != GetOuter())
 	{
diff --git a/Engine/Source/Runtime/CoreUObject/Public/UObject/MetaData.h b/Engine/Source/Runtime/CoreUObject/Public/UObject/MetaData.h
--- a/Engine/Source/Runtime/CoreUObject/Public/UObject/MetaData.h
+++ b/Engine/Source/Runtime/CoreUObject/Public/UObject/MetaData.h
@@ -210,9 +210,7 @@ private:
 	// Redirect map from deprecated keys to current key names
 	static COREUOBJECT_API TMap<FName, FName> KeyRedirectMap;
 };
-#endif
 
-#if WITH_EDITOR
 struct FMetaDataUtilities
 {
 private:
@@ -264,4 +262,4 @@ private:
 	FMetaDataUtilities() {}
 };
 
-#endif //WITH_EDITOR
+#endif //WITH_METADATA
diff --git a/Engine/Source/Runtime/CoreUObject/Tests/PackageMetaDataTest.cpp b/Engine/Source/Runtime/CoreUObject/Tests/PackageMetaDataTest.cpp
new file mode 100644
--- /dev/null
+++ b/Engine/Source/Runtime/CoreUObject/Tests/PackageMetaDataTest.cpp
@@ -0,0 +1,95 @@
+// Copyright Epic Games, Inc. All Rights Reserved.
+
+#if WITH_LOW_LEVEL_TESTS && WITH_METADATA
+
+#include "TestHarness.h"
+#include "UObject/Package.h"
+#include "UObject/MetaData.h"
+#include "UObject/ObjectRedirector.h"
+
+namespace UE::CoreObject::Private::Tests
+{
+	TEST_CASE("CoreUObject::UPackage::FMetaData", "[CoreUObject][Package][MetaData]")
+	{
+#if WITH_EDITORONLY_DATA
+		// FMetaDataUtilities::FMoveMetadataHelperContext used to handle objects renaming with regards to metadata will only work if GIsEditor
+		// is set, so we must set it here.
+		TGuardValue<bool> GIsEditorGuard(GIsEditor, true);
+#endif // WITH_EDITORONLY_DATA
+
+		const TCHAR* ObjectValueKey = TEXT("ObjectValueKey");
+		const TCHAR* ObjectValueValue = TEXT("ObjectValueValue");
+
+		UPackage* NewPackage = CreatePackage(TEXT("/Temp/TestPackage"));
+		TEST_NOT_NULL(TEXT("Should be able to create a new package"), NewPackage);
+
+		UObject* NewObject = ::NewObject<UObjectRedirector>(NewPackage, MakeUniversallyUniqueObjectName(NewPackage, NAME_Object));
+		TEST_NOT_NULL(TEXT("Should be able to create a new object"), NewObject);
+
+		UObject* NewSubObject = ::NewObject<UObjectRedirector>(NewObject, MakeUniversallyUniqueObjectName(NewObject, NAME_Object));
+		TEST_NOT_NULL(TEXT("Should be able to create a new sub object"), NewSubObject);
+
+		FMetaData& NewPackageMetaData = NewPackage->GetMetaData();
+		TEST_FALSE(TEXT("New object shouldn't have metadata values"), NewPackageMetaData.HasObjectValues(NewObject));
+		TEST_FALSE(TEXT("New sub object shouldn't have metadata values"), NewPackageMetaData.HasObjectValues(NewSubObject));
+
+		NewPackageMetaData.SetValue(NewObject, ObjectValueKey, ObjectValueValue);
+		TEST_TRUE(TEXT("New object should have metadata values"), NewPackageMetaData.HasObjectValues(NewObject));
+
+		NewPackageMetaData.SetValue(NewSubObject, ObjectValueKey, ObjectValueValue);
+		TEST_TRUE(TEXT("New sub object should have metadata values"), NewPackageMetaData.HasObjectValues(NewObject));
+
+		const FString* FoundObjectValue = NewPackageMetaData.FindValue(NewObject, ObjectValueKey);
+		TEST_NOT_NULL(TEXT("New object should have a valid metadata value"), FoundObjectValue);
+		TEST_TRUE(TEXT("New object should have valid metadata value"), FoundObjectValue && *FoundObjectValue == ObjectValueValue);
+
+		FoundObjectValue = NewPackageMetaData.FindValue(NewSubObject, ObjectValueKey);
+		TEST_NOT_NULL(TEXT("New sub object should have a valid metadata value"), FoundObjectValue);
+		TEST_TRUE(TEXT("New sub object should have valid metadata value"), FoundObjectValue && *FoundObjectValue == ObjectValueValue);
+
+		const FName NewObjectName = MakeUniversallyUniqueObjectName(NewPackage, NAME_Object);
+		const FString NewObjectNameString = NewObjectName.ToString();
+		TEST_TRUE(TEXT("Should be able to rename new object"), NewObject->Rename(*NewObjectNameString, nullptr));
+
+		FoundObjectValue = NewPackageMetaData.FindValue(NewObject, ObjectValueKey);
+		TEST_NOT_NULL(TEXT("New object should have a valid metadata value after rename"), FoundObjectValue);
+		TEST_TRUE(TEXT("New object should have valid metadata value after rename"), FoundObjectValue && *FoundObjectValue == ObjectValueValue);
+
+		FoundObjectValue = NewPackageMetaData.FindValue(NewSubObject, ObjectValueKey);
+		TEST_NOT_NULL(TEXT("New sub object should have a valid metadata value after parent rename"), FoundObjectValue);
+		TEST_TRUE(TEXT("New sub object should have valid metadata value after parent rename"), FoundObjectValue && *FoundObjectValue == ObjectValueValue);
+
+		UPackage* NewPackage2 = CreatePackage(TEXT("/Temp/TestPackage1"));
+		TEST_NOT_NULL(TEXT("Should be able to create a new package"), NewPackage2);
+		TEST_TRUE(TEXT("Should be able to rename new object into another package"), NewObject->Rename(nullptr, NewPackage2));
+
+		FoundObjectValue = NewPackageMetaData.FindValue(NewObject, ObjectValueKey);
+		TEST_NULL(TEXT("New object should have an invalid metadata value after rename in original package metadata"), FoundObjectValue);
+
+		FoundObjectValue = NewPackageMetaData.FindValue(NewSubObject, ObjectValueKey);
+		TEST_NULL(TEXT("New sub object should have a valid metadata value after parent rename in original package metadata"), FoundObjectValue);
+
+		FMetaData& NewPackage2MetaData = NewPackage2->GetMetaData();
+		FoundObjectValue = NewPackage2MetaData.FindValue(NewObject, ObjectValueKey);
+		TEST_NOT_NULL(TEXT("New object should have a valid metadata value in renamed package"), FoundObjectValue);
+		TEST_TRUE(TEXT("New object should have valid metadata value in renamed pacjage"), FoundObjectValue && *FoundObjectValue == ObjectValueValue);
+
+		FoundObjectValue = NewPackage2MetaData.FindValue(NewSubObject, ObjectValueKey);
+		TEST_NOT_NULL(TEXT("New sub object should have a valid metadata value in renamed package"), FoundObjectValue);
+		TEST_TRUE(TEXT("New sub object should have valid metadata value in renamed package"), FoundObjectValue && *FoundObjectValue == ObjectValueValue);
+
+		TEST_TRUE(TEXT("Should be able to rename new sub object"), NewSubObject->Rename(*NewObjectNameString, nullptr));
+
+		FoundObjectValue = NewPackage2MetaData.FindValue(NewSubObject, ObjectValueKey);
+		TEST_NOT_NULL(TEXT("New sub object should have a valid metadata value after rename"), FoundObjectValue);
+		TEST_TRUE(TEXT("New sub object should have valid metadata value after renamepackage"), FoundObjectValue && *FoundObjectValue == ObjectValueValue);
+
+		TEST_TRUE(TEXT("Should be able to rename new sub object"), NewSubObject->Rename(nullptr, NewPackage2));
+
+		FoundObjectValue = NewPackage2MetaData.FindValue(NewSubObject, ObjectValueKey);
+		TEST_NOT_NULL(TEXT("New object should have a valid metadata value after removing from parent"), FoundObjectValue);
+		TEST_TRUE(TEXT("New object should have valid metadata value after removing from parent"), FoundObjectValue && *FoundObjectValue == ObjectValueValue);
+	}
+}
+
+#endif // WITH_LOW_LEVEL_TESTS && WITH_METADATA
\ No newline at end of file
