[ZenLoader]
- Add support to call progress callbacks for package already in-flight

diff --git a/Engine/Source/Runtime/CoreUObject/Private/Serialization/AsyncLoading2.cpp b/Engine/Source/Runtime/CoreUObject/Private/Serialization/AsyncLoading2.cpp
--- a/Engine/Source/Runtime/CoreUObject/Private/Serialization/AsyncLoading2.cpp
+++ b/Engine/Source/Runtime/CoreUObject/Private/Serialization/AsyncLoading2.cpp
@@ -3869,6 +3869,7 @@ public:
 
 	/** [EDL] End Event driven loader specific stuff */
 
+	void CallProgressCallback(FLoadPackageAsyncProgressDelegate& ProgressCallback, EAsyncLoadingProgress ProgressType);
 	void CallProgressCallbacks(EAsyncLoadingProgress ProgressType);
 private:
 	void InitializeExportArchive(FExportArchive& Ar, bool bIsOptionalSegment);
@@ -5266,6 +5267,23 @@ FAsyncPackage2* FAsyncLoadingThread2::FindOrInsertPackage(FAsyncLoadingThreadSta
 		}
 		if (PackageProgressDelegate.IsValid())
 		{
+			// To avoid missing progress callbacks for package already in-flight,
+			// call any state that the package might already have reached.
+			if (Package->AsyncPackageLoadingState >= EAsyncPackageLoadingState2::WaitingForIo)
+			{
+				Package->CallProgressCallback(*PackageProgressDelegate, EAsyncLoadingProgress::Started);
+			}
+
+			if (Package->AsyncPackageLoadingState >= EAsyncPackageLoadingState2::ProcessPackageSummary)
+			{
+				Package->CallProgressCallback(*PackageProgressDelegate, EAsyncLoadingProgress::Read);
+			}
+
+			if (Package->AsyncPackageLoadingState >= EAsyncPackageLoadingState2::ExportsDone)
+			{
+				Package->CallProgressCallback(*PackageProgressDelegate, EAsyncLoadingProgress::Serialized);
+			}
+
 			Package->AddProgressCallback(PackageProgressDelegate);
 		}
 	}
@@ -7390,6 +7408,18 @@ EEventLoadNodeExecutionResult FAsyncPackage2::Event_DependenciesReady(FAsyncLoad
 	return EEventLoadNodeExecutionResult::Complete;
 }
 
+void FAsyncPackage2::CallProgressCallback(FLoadPackageAsyncProgressDelegate& ProgressCallback, EAsyncLoadingProgress ProgressType)
+{
+	FLoadPackageAsyncProgressParams Params
+	{
+		.PackageName = Desc.UPackageName,
+		.LoadedPackage = GetLinkerRoot(),
+		.ProgressType = ProgressType
+	};
+
+	ProgressCallback.Invoke(Params);
+}
+
 void FAsyncPackage2::CallProgressCallbacks(EAsyncLoadingProgress ProgressType)
 {
 	if (ProgressCallbacks.Num() != 0)
