AssetRegistry/RedirectCollector/AssetManager: Remove the now-redundant reporting of ObjectRedirectors discovered by the AssetRegistry at end of scan; we have already added reporting of each ObjectRedirector as soon as they are discovered. The reporting at end-of-scan was handling collapsing of chained redirectors into a single resultant redirect; change the handling of that to be internal to RedirectCollector, so that we can make the chained redirects available as soon as we find them, and so we robustly handle deletions or modifications of redirects in the chain.
#rnx

diff --git a/Engine/Source/Runtime/AssetRegistry/Private/AssetRegistry.cpp b/Engine/Source/Runtime/AssetRegistry/Private/AssetRegistry.cpp
--- a/Engine/Source/Runtime/AssetRegistry/Private/AssetRegistry.cpp
+++ b/Engine/Source/Runtime/AssetRegistry/Private/AssetRegistry.cpp
@@ -4393,48 +4393,51 @@ void UAssetRegistryImpl::AssetCreated(UObject* NewAsset)
 void UAssetRegistryImpl::AssetDeleted(UObject* DeletedAsset)
 {
 	checkf(GIsEditor, TEXT("Updating the AssetRegistry is only available in editor"));
-	if (ensure(DeletedAsset) && DeletedAsset->IsAsset())
+	if (!ensure(DeletedAsset) || !DeletedAsset->IsAsset())
 	{
-		UPackage* DeletedObjectPackage = DeletedAsset->GetOutermost();
-		bool bIsEmptyPackage = DeletedObjectPackage != nullptr && UPackage::IsEmptyPackage(DeletedObjectPackage, DeletedAsset);
+		return;
+	}
 
-		bool bShouldSkipAsset;
-		{
-			LLM_SCOPE(ELLMTag::AssetRegistry);
-			UE::AssetRegistry::FInterfaceWriteScopeLock InterfaceScopeLock(InterfaceLock);
+	UPackage* DeletedObjectPackage = DeletedAsset->GetOutermost();
+	bool bIsEmptyPackage = DeletedObjectPackage != nullptr && UPackage::IsEmptyPackage(DeletedObjectPackage, DeletedAsset);
 
-			// Deleting the last asset in a package causes the package to be garbage collected.
-			// If the UPackage object is GCed, it will be considered 'Unloaded' which will cause it to
-			// be fully loaded from disk when save is invoked.
-			// We want to keep the package around so we can save it empty or delete the file.
-			if (bIsEmptyPackage)
-			{
-				GuardedData.AddEmptyPackage(DeletedObjectPackage->GetFName());
-			}
-			bShouldSkipAsset = GuardedData.ShouldSkipAsset(DeletedAsset);
-		}
+	bool bShouldSkipAsset;
+	{
+		LLM_SCOPE(ELLMTag::AssetRegistry);
+		UE::AssetRegistry::FInterfaceWriteScopeLock InterfaceScopeLock(InterfaceLock);
 
-#if WITH_EDITOR
-		if (!IsGathering() && FAssetData::IsRedirector(DeletedAsset))
+		// Deleting the last asset in a package causes the package to be garbage collected.
+		// If the UPackage object is GCed, it will be considered 'Unloaded' which will cause it to
+		// be fully loaded from disk when save is invoked.
+		// We want to keep the package around so we can save it empty or delete the file.
+		if (bIsEmptyPackage)
 		{
-			// Need to remove from GRedirectCollector
-			GRedirectCollector.RemoveAssetPathRedirection(FSoftObjectPath::ConstructFromObject(DeletedAsset));
+			GuardedData.AddEmptyPackage(DeletedObjectPackage->GetFName());
 		}
+		bShouldSkipAsset = GuardedData.ShouldSkipAsset(DeletedAsset);
+	}
+
+#if WITH_EDITOR
+	// Update Redirectors
+	if (FAssetData::IsRedirector(DeletedAsset))
+	{
+		// Need to remove from GRedirectCollector
+		GRedirectCollector.RemoveAssetPathRedirection(FSoftObjectPath::ConstructFromObject(DeletedAsset));
+	}
 #endif
 
-		if (!bShouldSkipAsset)
-		{
-			FAssetData AssetDataDeleted = FAssetData(DeletedAsset, FAssetData::ECreationFlags::AllowBlueprintClass,
-				EAssetRegistryTagsCaller::AssetRegistryQuery);
+	if (!bShouldSkipAsset)
+	{
+		FAssetData AssetDataDeleted = FAssetData(DeletedAsset, FAssetData::ECreationFlags::AllowBlueprintClass,
+			EAssetRegistryTagsCaller::AssetRegistryQuery);
 
-			checkf(IsInGameThread(), TEXT("AssetDeleted is not yet implemented as callable from other threads"));
-			// Let subscribers know that the asset was removed from the registry
-			AssetRemovedEvent.Broadcast(AssetDataDeleted);
-			OnAssetsRemoved().Broadcast({AssetDataDeleted});
+		checkf(IsInGameThread(), TEXT("AssetDeleted is not yet implemented as callable from other threads"));
+		// Let subscribers know that the asset was removed from the registry
+		AssetRemovedEvent.Broadcast(AssetDataDeleted);
+		OnAssetsRemoved().Broadcast({AssetDataDeleted});
 
-			// Notify listeners that an in-memory asset was just deleted
-			InMemoryAssetDeletedEvent.Broadcast(DeletedAsset);
-		}
+		// Notify listeners that an in-memory asset was just deleted
+		InMemoryAssetDeletedEvent.Broadcast(DeletedAsset);
 	}
 }
 
@@ -5419,11 +5422,6 @@ Impl::EGatherStatus FAssetRegistryImpl::TickGatherer(Impl::FTickContext& TickCon
 
 void FAssetRegistryImpl::OnInitialSearchCompleted(Impl::FEventContext& EventContext)
 {
-#if WITH_EDITOR
-	// update redirectors
-	UpdateRedirectCollector();
-#endif
-
 	// Handle any deferred loading operations
 	SetPerformanceMode(Impl::EPerformanceMode::MostlyStatic);
 
@@ -5445,11 +5443,6 @@ void FAssetRegistryImpl::OnInitialSearchCompleted(Impl::FEventContext& EventCont
 
 void FAssetRegistryImpl::OnAdditionalMountSearchCompleted(Impl::FEventContext& EventContext)
 {
-#if WITH_EDITOR
-	// update redirectors
-	UpdateRedirectCollector();
-#endif
-
 	// Handle any deferred loading operations
 	SetPerformanceMode(Impl::EPerformanceMode::MostlyStatic);
 
@@ -7236,19 +7229,16 @@ bool FAssetRegistryImpl::RemoveAssetPath(Impl::FEventContext& EventContext, FNam
 
 void FAssetRegistryImpl::AddAssetData(Impl::FEventContext& EventContext, FAssetData* AssetData)
 {
-	// Make sure to consider redirections!
 #if WITH_EDITOR
+	// Update Redirectors
+	if (AssetData->IsRedirector())
 	{
-		if (AssetData->IsRedirector())
+		FString RedirectDestinationString;
+		AssetData->GetTagValue(UE::AssetRegistry::Impl::DestinationObjectFName, RedirectDestinationString);
+		FSoftObjectPath RedirectDestination = RedirectDestinationString;
+		if (!RedirectDestination.IsNull())
 		{
-			const FName DestinationObjectFName(TEXT("DestinationObject"));
-			FString RedirectDestinationString;
-			AssetData->GetTagValue(DestinationObjectFName, RedirectDestinationString);
-			FSoftObjectPath RedirectDestination = RedirectDestinationString;
-			if (!RedirectDestination.IsNull())
-			{
-				GRedirectCollector.AddAssetPathRedirection(AssetData->GetSoftObjectPath(), RedirectDestination);
-			}
+			GRedirectCollector.AddAssetPathRedirection(AssetData->GetSoftObjectPath(), RedirectDestination);
 		}
 	}
 #endif
@@ -7353,6 +7343,37 @@ void FAssetRegistryImpl::UpdateAssetData(Impl::FEventContext& EventContext, FAss
 		}
 	}
 
+#if WITH_EDITOR
+	// Update Redirectors
+	if (AssetData->IsRedirector() || NewAssetData.IsRedirector())
+	{
+		FSoftObjectPath OldSource = AssetData->GetSoftObjectPath();
+		FSoftObjectPath NewSource = NewAssetData.GetSoftObjectPath();
+		FSoftObjectPath OldTarget;
+		FSoftObjectPath NewTarget;
+		if (AssetData->IsRedirector())
+		{
+			FString TargetString;
+			AssetData->GetTagValue(UE::AssetRegistry::Impl::DestinationObjectFName, TargetString);
+			OldTarget = TargetString;
+		}
+		if (NewAssetData.IsRedirector())
+		{
+			FString TargetString;
+			NewAssetData.GetTagValue(UE::AssetRegistry::Impl::DestinationObjectFName, TargetString);
+			NewTarget = TargetString;
+		}
+		if (OldSource != NewSource && OldTarget.IsValid())
+		{
+			GRedirectCollector.RemoveAssetPathRedirection(OldSource);
+		}
+		if (NewTarget.IsValid())
+		{
+			GRedirectCollector.AddAssetPathRedirection(NewSource, NewTarget);
+		}
+	}
+#endif
+
 	bool bModified;
 	State.UpdateAssetData(AssetData, MoveTemp(NewAssetData), &bModified);
 	
@@ -7378,45 +7399,56 @@ void FAssetRegistryImpl::AddNonOverlappingTags(Impl::FEventContext& EventContext
 
 bool FAssetRegistryImpl::RemoveAssetData(Impl::FEventContext& EventContext, FAssetData* AssetData)
 {
+	if (!ensure(AssetData))
+	{
+		return false;
+	}
 	bool bRemoved = false;
 
-	if (ensure(AssetData))
+
+	if (!ShouldSkipAsset(AssetData->AssetClassPath, AssetData->PackageFlags))
 	{
-		if (!ShouldSkipAsset(AssetData->AssetClassPath, AssetData->PackageFlags))
-		{
-			EventContext.AssetEvents.Emplace(*AssetData, Impl::FEventContext::EEvent::Removed);
-		}
+		EventContext.AssetEvents.Emplace(*AssetData, Impl::FEventContext::EEvent::Removed);
+	}
 
-		// Make sure to consider redirections!
 #if WITH_EDITOR
-		if (AssetData->IsRedirector())
-		{
-			GRedirectCollector.RemoveAssetPathRedirection(AssetData->GetSoftObjectPath());
-		}
+	// Save a copy of the AssetData's SoftObjectPath if we need to remove it; we don't want to remove the
+	// redirection unless the AssetRegistryState confirms the AssetData existed and was removed, but after
+	// it is removed our AssetData pointer might become a dangling pointer.
+	FSoftObjectPath RedirectSource = AssetData->IsRedirector() ? AssetData->GetSoftObjectPath() : FSoftObjectPath();
 #endif
 
-		// Remove from the class map if removing a blueprint
-		if (ClassGeneratorNames.Contains(AssetData->AssetClassPath))
+	// Remove from the class map if removing a blueprint
+	if (ClassGeneratorNames.Contains(AssetData->AssetClassPath))
+	{
+		const FString OldGeneratedClass = AssetData->GetTagValueRef<FString>(FBlueprintTags::GeneratedClassPath);
+		if (!OldGeneratedClass.IsEmpty() && OldGeneratedClass != TEXTVIEW("None"))
 		{
-			const FString OldGeneratedClass = AssetData->GetTagValueRef<FString>(FBlueprintTags::GeneratedClassPath);
-			if (!OldGeneratedClass.IsEmpty() && OldGeneratedClass != TEXTVIEW("None"))
+			const FTopLevelAssetPath OldGeneratedClassPathName(FPackageName::ExportTextPathToObjectPath(OldGeneratedClass));
+			if (ensureAlwaysMsgf(!OldGeneratedClassPathName.IsNull(),
+				TEXT("Short class name used: OldGeneratedClass=%s"), *OldGeneratedClass))
 			{
-				const FTopLevelAssetPath OldGeneratedClassPathName(FPackageName::ExportTextPathToObjectPath(OldGeneratedClass));
-				if (ensureAlwaysMsgf(!OldGeneratedClassPathName.IsNull(),
-					TEXT("Short class name used: OldGeneratedClass=%s"), *OldGeneratedClass))
-				{
-					CachedBPInheritanceMap.Remove(OldGeneratedClassPathName);
+				CachedBPInheritanceMap.Remove(OldGeneratedClassPathName);
 
-					// Invalidate caching because CachedBPInheritanceMap got modified
-					TempCachedInheritanceBuffer.bDirty = true;
-				}
+				// Invalidate caching because CachedBPInheritanceMap got modified
+				TempCachedInheritanceBuffer.bDirty = true;
 			}
 		}
-
-		bool bRemovedDependencyData;
-		State.RemoveAssetData(AssetData, true /* bRemoveDependencyData */, bRemoved, bRemovedDependencyData);
 	}
 
+	bool bRemovedDependencyData;
+	State.RemoveAssetData(AssetData, true /* bRemoveDependencyData */, bRemoved, bRemovedDependencyData);
+
+#if WITH_EDITOR
+	if (bRemoved)
+	{
+		// Update Redirectors
+		if (RedirectSource.IsValid())
+		{
+			GRedirectCollector.RemoveAssetPathRedirection(RedirectSource);
+		}
+	}
+#endif
 	return bRemoved;
 }
 
@@ -8091,41 +8123,6 @@ void FAssetRegistryImpl::PushProcessLoadedAssetsBatch(Impl::FEventContext& Event
 	}
 }
 
-
-void FAssetRegistryImpl::UpdateRedirectCollector()
-{
-	// Look for all redirectors in the AssetRegistry
-	TArray<FSoftObjectPath> Paths;
-	State.EnumerateAssetsByClassPathName(UE::AssetRegistry::GetClassPathObjectRedirector(),
-		[this, &Paths](const FAssetData* AssetData)
-		{
-			Paths.Add(AssetData->GetSoftObjectPath());
-			return true; // Keep iterating
-		});
-
-	TArray<FSoftObjectPath> RedirectedPaths;
-	RedirectedPaths.SetNumZeroed(Paths.Num());
-	ParallelFor(Paths.Num(), [&RedirectedPaths, &Paths, this](int32 Index)
-	{
-		const FSoftObjectPath& Source = Paths[Index];
-		FSoftObjectPath Destination = GetRedirectedObjectPath(Source, nullptr, nullptr, /*bNeedsScanning*/ false);
-		if (Destination != Source)
-		{
-			RedirectedPaths[Index] = Destination;
-		}
-	});
-
-	for (int32_t AssetIndex = 0; AssetIndex < Paths.Num(); AssetIndex++)
-	{
-		FSoftObjectPath& Destination = RedirectedPaths[AssetIndex];
-		if (RedirectedPaths[AssetIndex].IsNull())
-		{
-			continue;
-		}
-		GRedirectCollector.AddAssetPathRedirection(Paths[AssetIndex], Destination);
-	}
-}
-
 }
 
 #endif // WITH_EDITOR
diff --git a/Engine/Source/Runtime/AssetRegistry/Private/AssetRegistryImpl.h b/Engine/Source/Runtime/AssetRegistry/Private/AssetRegistryImpl.h
--- a/Engine/Source/Runtime/AssetRegistry/Private/AssetRegistryImpl.h
+++ b/Engine/Source/Runtime/AssetRegistry/Private/AssetRegistryImpl.h
@@ -532,9 +532,6 @@ private:
 	 */
 	bool TryPostLoadAssetRegistryTags(FAssetData* AssetData);
 
-	/** Update Redirect collector with redirects loaded from asset registry */
-	void UpdateRedirectCollector();
-
 	/** If new gatherers were registered, this call will refresh the mapping */
 	void RebuildAssetDependencyGathererMapIfNeeded();
 
diff --git a/Engine/Source/Runtime/CoreUObject/Private/Misc/RedirectCollector.cpp b/Engine/Source/Runtime/CoreUObject/Private/Misc/RedirectCollector.cpp
--- a/Engine/Source/Runtime/CoreUObject/Private/Misc/RedirectCollector.cpp
+++ b/Engine/Source/Runtime/CoreUObject/Private/Misc/RedirectCollector.cpp
@@ -42,21 +42,19 @@ FAutoConsoleCommand CVarRedirectCollectorDumpAllAssetRedirects(
 
 void RedirectCollectorDumpAllAssetRedirects()
 {
-	UE::TDynamicUniqueLock<FCriticalSection> ScopeLock(GRedirectCollector.AcquireLock());
-	const TMap<FSoftObjectPath, FSoftObjectPath>& Redirects = GRedirectCollector.GetObjectPathRedirectionMapUnderLock(ScopeLock);
 	FString FullyQualifiedFileName = FPaths::ProfilingDir() + FString::Printf(TEXT("AllRedirects (%s).csv"), *FDateTime::Now().ToString(TEXT("%Y%m%d_%H%M%S")));
 	TUniquePtr<FArchive> OutputFile(IFileManager::Get().CreateFileWriter(*FullyQualifiedFileName));
 	if (OutputFile.IsValid())
 	{
 		TAnsiStringBuilder<4096> StringBuilder;
 
-		for (const TPair<FSoftObjectPath, FSoftObjectPath>& Pair : Redirects)
-		{
-			StringBuilder	<< WriteToAnsiString<256>(Pair.Key.ToString()).ToString() 
-							<< TEXT(",") 
-							<< WriteToAnsiString<256>(Pair.Value.ToString()).ToString()
-							<< LINE_TERMINATOR;
-		}
+		GRedirectCollector.EnumerateRedirectsUnderLock([&StringBuilder](FRedirectCollector::FRedirectionData& Data)
+			{
+				StringBuilder << WriteToAnsiString<256>(Data.GetSource().ToString()).ToString()
+					<< TEXT(",")
+					<< WriteToAnsiString<256>(Data.GetFirstTarget().ToString()).ToString()
+					<< LINE_TERMINATOR;
+			});
 		OutputFile->Serialize(StringBuilder.GetData(), StringBuilder.Len());
 	}
 }
@@ -267,24 +265,125 @@ void FRedirectCollector::GetAllSourcePathsForTargetPath(const FSoftObjectPath& T
 
 void FRedirectCollector::AddObjectPathRedirectionInternal(const FSoftObjectPath& Source, const FSoftObjectPath& Destination)
 {
-	if (FSoftObjectPath* ExistingDestination = ObjectPathRedirectionMap.Find(Source))
+	FSimpleOrChainedRedirect& ExistingRedirect = ObjectPathRedirectionMap.FindOrAdd(Source);
+	if (!ExistingRedirect.GetFirstTarget().IsNull())
 	{
-		// We are replacing a redirect not adding one. 
-		// That means we need to remove an old reverse lookup and add a new one
-		if (ObjectPathSourcesArray* ReverseLookupArray = ObjectPathRedirectionReverseMap.Find(*ExistingDestination))
+		if (ExistingRedirect.GetFirstTarget() == Destination)
 		{
-			ReverseLookupArray->Remove(Source);
+			return;
 		}
 
-		*ExistingDestination = Destination;
+		// We are replacing a redirect not adding one. 
+		// That means we need to remove all old sources that had us in their chain to their final destination and
+		// then add them to their new final destination.
+
+		// Bootstrap the replacement destination in as a SimpleRedirect with First == Final. We will replace it if
+		// necessary with a ChainedRedirect in the loop below, because it will be one of the
+		// SourcesThatWentToOldDestination.
+		// We need its GetFirstTarget == Destination present in the ObjectPathRedirectionMap for the calls to
+		// TraverseToFinalTarget to work for all of the redirectors that chain into it.
+		FSoftObjectPath OldFinalTarget = ExistingRedirect.GetFinalTarget();
+		ExistingRedirect = FSimpleOrChainedRedirect(Destination);
+
+		// Get all redirects that that had OldFinalTarget as their FinalTarget, and clear OldFinalTarget from
+		// the reverse map; we will reconstruct it if necessary.
+		ObjectPathSourcesArray SourcesThatWentToOldTarget;
+		ObjectPathRedirectionReverseMap.RemoveAndCopyValue(OldFinalTarget, SourcesThatWentToOldTarget);
+
+		// For all redirects (including the one from our Source argument) that had OldFinalTarget as their
+		// FinalTarget, calculate their new FinalTarget using the graph of FirstTargets, and set the data
+		// for their FinalTarget and for their entry in the reverse map.
+		for (const FSoftObjectPath& SourceThatWentToOldTarget : SourcesThatWentToOldTarget)
+		{
+			if (FSimpleOrChainedRedirect* OldRedirectionTarget
+				= ObjectPathRedirectionMap.Find(SourceThatWentToOldTarget))
+			{
+				FSoftObjectPath FirstTarget = OldRedirectionTarget->GetFirstTarget();
+				const FSoftObjectPath& FinalTarget = TraverseToFinalTarget(&FirstTarget);
+
+				ObjectPathRedirectionReverseMap.FindOrAdd(FinalTarget).AddUnique(SourceThatWentToOldTarget);
+
+				*OldRedirectionTarget = FSimpleOrChainedRedirect::ConstructSimpleOrChained(
+					FirstTarget, FinalTarget);
+			}
+		}
 	}
 	else
 	{
-		ObjectPathRedirectionMap.Add(Source, Destination);
+		// Add FirstTarget data for the new redirect before calling TraverseToFinalTarget. An empty FirstTarget value
+		// is invalid, and we might encounter it during TraverseToFinalTarget if there is a cycle.
+		ExistingRedirect = FSimpleOrChainedRedirect(Destination);
+		const FSoftObjectPath& FinalTarget = TraverseToFinalTarget(&Destination);
+		if (FinalTarget != Destination)
+		{
+			ExistingRedirect = FSimpleOrChainedRedirect(Destination, FinalTarget);
+		}
+
+		// Add the redirect's source to the reverse map for its FinalTarget, and if it was chained into by any
+		// existing redirects, remove them from the reverse lookup for Source, change their FinalTarget to FinalTarget,
+		// and add them to the reverse lookup for FinalTarget.
+		ObjectPathSourcesArray OldReverseLookupArray;
+		ObjectPathRedirectionReverseMap.RemoveAndCopyValue(Source, OldReverseLookupArray);
+		for (const FSoftObjectPath& ChainedPath : OldReverseLookupArray)
+		{
+			FSimpleOrChainedRedirect* ChainedRedirect = ObjectPathRedirectionMap.Find(ChainedPath);
+			check(ChainedRedirect); // It should be in ObjectPathRedirectionMap because it was in the reverse map
+			*ChainedRedirect = FSimpleOrChainedRedirect(ChainedRedirect->GetFirstTarget(), FinalTarget);
+		}
+
+		ObjectPathSourcesArray& FinalTargetArray = ObjectPathRedirectionReverseMap.FindOrAdd(FinalTarget);
+		FinalTargetArray.Add(Source);
+		FinalTargetArray.Append(MoveTemp(OldReverseLookupArray));
 	}
+}
 
-	ObjectPathSourcesArray& ReverseLookupArray = ObjectPathRedirectionReverseMap.FindOrAdd(Destination);
-	ReverseLookupArray.AddUnique(Source);
+const FSoftObjectPath& FRedirectCollector::TraverseToFinalTarget(const FSoftObjectPath* FirstTarget) const
+{
+	// Called within the CriticalSection
+
+	// Our input argument FirstTarget needs to be a pointer so that we are guaranteed it points to an LValue
+	// and we can return a reference to it. We need to support being called with a const FSoftObjectPath, but
+	// if we accepted const FSoftObjectPath& then we might get passed an X value and would not be allowed to
+	// return a pointer to it.
+	const FSimpleOrChainedRedirect* NextRedirect = ObjectPathRedirectionMap.Find(*FirstTarget);
+	if (!NextRedirect)
+	{
+		// This is the most common case; handle it as cheaply as possible
+		return *FirstTarget;
+	}
+
+	// Handle cycles in the graph of redirections
+	TSet<FSoftObjectPath> SeenPaths;
+	SeenPaths.Add(*FirstTarget);
+
+	const FSimpleOrChainedRedirect* CurrentRedirect;
+	do
+	{
+		CurrentRedirect = NextRedirect;
+		const FSoftObjectPath& CurrentTarget = CurrentRedirect->GetFirstTarget();
+		bool bAlreadySeen = false;
+		SeenPaths.Add(CurrentTarget, &bAlreadySeen);
+		if (bAlreadySeen)
+		{
+			// A cycle; return the first path we encountered in the cycle, which is CurrentTarget
+
+			ensureMsgf(!SeenPaths.Contains(CurrentTarget),
+				TEXT("Found circular redirect from %s to itself! Setting FinalDestination of %s to %s."),
+				*CurrentTarget.ToString(), *FirstTarget->ToString(), *CurrentTarget.ToString());
+			UE_LOG(LogRedirectors, Error, TEXT("Logging redirection chain: "));
+			for (const FSoftObjectPath& Entry : SeenPaths)
+			{
+				UE_LOG(LogRedirectors, Error, TEXT(" %s"), *Entry.ToString());
+			}
+			UE_LOG(LogRedirectors, Error, TEXT(" %s"), *CurrentTarget.ToString());
+
+			return CurrentTarget;
+		}
+
+		NextRedirect = ObjectPathRedirectionMap.Find(CurrentTarget);
+	} while (NextRedirect);
+
+	return CurrentRedirect->GetFirstTarget();
 }
 
 bool FRedirectCollector::ShouldTrackPackageReferenceTypes()
@@ -323,21 +422,41 @@ void FRedirectCollector::AddAssetPathRedirection(const FSoftObjectPath& Original
 
 bool FRedirectCollector::TryRemoveObjectPathRedirectionInternal(const FSoftObjectPath& Source)
 {
-	FSoftObjectPath Destination;
-	bool bRemoved = ObjectPathRedirectionMap.RemoveAndCopyValue(Source, Destination);
-	if (bRemoved)
+	FSimpleOrChainedRedirect OldRedirect;
+	if (!ObjectPathRedirectionMap.RemoveAndCopyValue(Source, OldRedirect))
+	{
+		return false;
+	}
+
+	FCoreRedirects::RecordRemovedObjectRedirector(Source, OldRedirect.GetFirstTarget());
+
+	// Get all redirects that that had OldFinalTarget as their FinalTarget, and clear OldFinalTarget from
+	// the reverse map; we will reconstruct it if necessary.
+	ObjectPathSourcesArray SourcesThatWentToOldTarget;
+	ObjectPathRedirectionReverseMap.RemoveAndCopyValue(OldRedirect.GetFinalTarget(), SourcesThatWentToOldTarget);
+
+	// For all redirects (except the one from our Source argument) that had OldFinalTarget as their
+	// FinalTarget, calculate their new FinalTarget using the graph of FirstTargets, and set the data
+	// for their FinalTarget and for their entry in the reverse map.
+	for (const FSoftObjectPath& SourceThatWentToOldTarget : SourcesThatWentToOldTarget)
 	{
-		FCoreRedirects::RecordRemovedObjectRedirector(Source, Destination);
-		if (ObjectPathSourcesArray* ReverseArray = ObjectPathRedirectionReverseMap.Find(Destination))
+		if (SourceThatWentToOldTarget == Source)
 		{
-			ReverseArray->RemoveSingleSwap(Source);
-			if (ReverseArray->Num() == 0)
-			{
-				ObjectPathRedirectionReverseMap.Remove(Destination);
-			}
+			continue;
+		}
+
+		if (FSimpleOrChainedRedirect* OldRedirectionTarget
+			= ObjectPathRedirectionMap.Find(SourceThatWentToOldTarget))
+		{
+			FSoftObjectPath FirstTarget = OldRedirectionTarget->GetFirstTarget();
+			const FSoftObjectPath& FinalTarget = TraverseToFinalTarget(&FirstTarget);
+
+			ObjectPathRedirectionReverseMap.FindOrAdd(FinalTarget).AddUnique(SourceThatWentToOldTarget);
+
+			*OldRedirectionTarget = FSimpleOrChainedRedirect::ConstructSimpleOrChained(FirstTarget, FinalTarget);
 		}
 	}
-	return bRemoved;
+	return true;
 }
 
 void FRedirectCollector::RemoveAssetPathRedirection(const FSoftObjectPath& OriginalPath)
@@ -347,63 +466,64 @@ void FRedirectCollector::RemoveAssetPathRedirection(const FSoftObjectPath& Origi
 	TryRemoveObjectPathRedirectionInternal(OriginalPath);
 }
 
-FSoftObjectPath FRedirectCollector::GetAssetPathRedirection(const FSoftObjectPath& OriginalPath)
+FSoftObjectPath FRedirectCollector::GetAssetPathRedirection(const FSoftObjectPath& OriginalPath) const
 {
 	FScopeLock ScopeLock(&CriticalSection);
-	TArray<FSoftObjectPath, TInlineAllocator<2>> SeenPaths;
-
-	// We need to follow the redirect chain recursively
-	FSoftObjectPath CurrentPath = OriginalPath;
 
-	while (!CurrentPath.IsNull())
+	const FSimpleOrChainedRedirect* Redirection;
+	if (Redirection = ObjectPathRedirectionMap.Find(OriginalPath); Redirection)
 	{
-		if (!ensureMsgf(!SeenPaths.Contains(CurrentPath), TEXT("Found circular redirect from %s to %s! Returning None instead"), *OriginalPath.ToString(), *CurrentPath.ToString()))
-		{
-			UE_LOG(LogRedirectors, Error, TEXT("Logging redirection chain: "));
-			for (const FSoftObjectPath& Entry : SeenPaths)
-			{
-				UE_LOG(LogRedirectors, Error, TEXT(" %s"), *Entry.ToString());
-			}
-			return FSoftObjectPath();
-		}
-		SeenPaths.Add(CurrentPath);
+		return Redirection->GetFinalTarget();
+	}
+	if (Redirection = ObjectPathRedirectionMap.Find(OriginalPath.GetWithoutSubPath()); Redirection)
+	{
+		return FSoftObjectPath(Redirection->GetFinalTarget().GetAssetPath(), OriginalPath.GetSubPathString());
+	}
+	return FSoftObjectPath();
+}
 
-		if (FSoftObjectPath* NewPath = ObjectPathRedirectionMap.Find(CurrentPath))
-		{
-			CurrentPath = *NewPath;
-		}
-		else if ((NewPath = ObjectPathRedirectionMap.Find(CurrentPath.GetWithoutSubPath())))
-		{
-			CurrentPath = FSoftObjectPath(NewPath->GetAssetPath(), CurrentPath.GetSubPathString());
-		}
-		else
-		{
-			break;
-		}
+void FRedirectCollector::EnumerateRedirectsUnderLock(TFunctionRef<void(FRedirectionData&)> Callback) const
+{
+	FScopeLock ScopeLock(&CriticalSection);
+	
+	for (const TPair<FSoftObjectPath, FSimpleOrChainedRedirect>& Pair : ObjectPathRedirectionMap)
+	{
+		FRedirectionData RedirectionData{ Pair.Key, Pair.Value.GetFirstTarget(), Pair.Value.GetFinalTarget() };
+		Callback(RedirectionData);
 	}
+}
 
-	if (CurrentPath != OriginalPath)
+const TMap<FSoftObjectPath, FSoftObjectPath>& FRedirectCollector::GetObjectPathRedirectionMapUnderLock(
+	const UE::TDynamicUniqueLock<FCriticalSection>& Lock) const
+{
+	ensure(Lock.OwnsLock());
+	static TMap<FSoftObjectPath, FSoftObjectPath> Result;
+	Result.Empty();
+	for (const TPair<FSoftObjectPath, FSimpleOrChainedRedirect>& Pair : ObjectPathRedirectionMap)
 	{
-		return CurrentPath;
+		Result.Add(Pair.Key, Pair.Value.GetFirstTarget());
 	}
-	return FSoftObjectPath();
+	return Result;
 }
 
 FRedirectCollector GRedirectCollector;
 #if WITH_AUTOMATION_WORKER
 
-IMPLEMENT_SIMPLE_AUTOMATION_TEST(FRedirectCollectorReverseLookupTest, "System.Core.Misc.RedirectCollector.ReverseLookup", EAutomationTestFlags_ApplicationContextMask | EAutomationTestFlags::EngineFilter)
+IMPLEMENT_SIMPLE_AUTOMATION_TEST(FRedirectCollectorReverseLookupTest,
+	"System.Core.Misc.RedirectCollector.ReverseLookup",
+	EAutomationTestFlags_ApplicationContextMask | EAutomationTestFlags::EngineFilter)
 bool FRedirectCollectorReverseLookupTest::RunTest(const FString& Parameters)
 {
 	FScopeLock ScopeLock(&GRedirectCollector.CriticalSection);
 	bool bSuccess = true;
 
 	// Validate that every forward redirect has a corresponding reverse redirect
-	for (const TPair<FSoftObjectPath, FSoftObjectPath>& ForwardRedirect : GRedirectCollector.ObjectPathRedirectionMap)
+	for (const TPair<FSoftObjectPath, FRedirectCollector::FSimpleOrChainedRedirect>& ForwardRedirect
+		: GRedirectCollector.ObjectPathRedirectionMap)
 	{
 		bool bFoundReverseEntry = false;
 		if (const FRedirectCollector::ObjectPathSourcesArray* Sources =
-			GRedirectCollector.ObjectPathRedirectionReverseMap.Find(ForwardRedirect.Value))
+			GRedirectCollector.ObjectPathRedirectionReverseMap.Find(ForwardRedirect.Value.GetFinalTarget()))
 		{
 			if (Sources->Contains(ForwardRedirect.Key))
 			{
@@ -414,20 +534,23 @@ bool FRedirectCollectorReverseLookupTest::RunTest(const FString& Parameters)
 		if (!bFoundReverseEntry)
 		{
 			bSuccess = false;
-			AddError(FString::Printf(TEXT("Failed to find matching reverse lookup for redirect %s --> %s"),
-				*ForwardRedirect.Key.ToString(), *ForwardRedirect.Value.ToString()));
+			AddError(FString::Printf(TEXT("Failed to find matching reverse lookup for redirect %s --> (%s, %s)"),
+				*ForwardRedirect.Key.ToString(), *ForwardRedirect.Value.GetFirstTarget().ToString(),
+				*ForwardRedirect.Value.GetFinalTarget().ToString()));
 		}
 	}
 
 	// Validate that every reverse redirect has a corresponding forward redirect
-	for (const TPair<FSoftObjectPath, FRedirectCollector::ObjectPathSourcesArray>& ReverseRedirectList : GRedirectCollector.ObjectPathRedirectionReverseMap)
+	for (const TPair<FSoftObjectPath, FRedirectCollector::ObjectPathSourcesArray>& ReverseRedirectList
+		: GRedirectCollector.ObjectPathRedirectionReverseMap)
 	{
 		for (const FSoftObjectPath& Source : ReverseRedirectList.Value)
 		{
 			bool bFoundForwardEntry = false;
-			if (const FSoftObjectPath* Target = GRedirectCollector.ObjectPathRedirectionMap.Find(Source))
+			if (const FRedirectCollector::FSimpleOrChainedRedirect* Target
+				= GRedirectCollector.ObjectPathRedirectionMap.Find(Source))
 			{
-				if (*Target == ReverseRedirectList.Key)
+				if (Target->GetFinalTarget() == ReverseRedirectList.Key)
 				{
 					bFoundForwardEntry = true;
 				}
@@ -436,8 +559,8 @@ bool FRedirectCollectorReverseLookupTest::RunTest(const FString& Parameters)
 			if (!bFoundForwardEntry)
 			{
 				bSuccess = false;
-				AddError(FString::Printf(TEXT("Failed to find matching reverse lookup for redirect %s --> %s"),
-					*Source.ToString(), *ReverseRedirectList.Key.ToString()));
+				AddError(FString::Printf(TEXT("Failed to find matching forward lookup for reverse redirect %s <-- %s"),
+					*ReverseRedirectList.Key.ToString(), *Source.ToString()));
 			}
 		}
 	}
diff --git a/Engine/Source/Runtime/CoreUObject/Private/UObject/CoreRedirects.cpp b/Engine/Source/Runtime/CoreUObject/Private/UObject/CoreRedirects.cpp
--- a/Engine/Source/Runtime/CoreUObject/Private/UObject/CoreRedirects.cpp
+++ b/Engine/Source/Runtime/CoreUObject/Private/UObject/CoreRedirects.cpp
@@ -369,13 +369,13 @@ namespace UE::CoreRedirects::Private
 #if WITH_EDITOR
 		// ensure that it's safe to add any redirector in GRedirectCollector
 		{
-			TDynamicUniqueLock<FCriticalSection> ScopeLock(GRedirectCollector.AcquireLock());
-			const TMap<FSoftObjectPath, FSoftObjectPath>& RedirectMap = GRedirectCollector.GetObjectPathRedirectionMapUnderLock(ScopeLock);
 			TArray<const FCoreRedirect> RedirectList;
-			for (const TPair<FSoftObjectPath, FSoftObjectPath>& Redirect : RedirectMap)
-			{
-				RedirectList.Emplace(ECoreRedirectFlags::Type_Asset, Redirect.Key.ToString(), Redirect.Value.ToString());
-			}
+			GRedirectCollector.EnumerateRedirectsUnderLock([&RedirectList](FRedirectCollector::FRedirectionData& Data)
+				{
+					RedirectList.Emplace(ECoreRedirectFlags::Type_Asset,
+						Data.GetSource().ToString(),
+						Data.GetFirstTarget().ToString());
+				});
 
 			FCoreRedirects::AddRedirectList(RedirectList, TEXT("GRedirectCollector"));
 			if (!FCoreRedirects::ValidateAssetRedirects())
diff --git a/Engine/Source/Runtime/CoreUObject/Public/Misc/RedirectCollector.h b/Engine/Source/Runtime/CoreUObject/Public/Misc/RedirectCollector.h
--- a/Engine/Source/Runtime/CoreUObject/Public/Misc/RedirectCollector.h
+++ b/Engine/Source/Runtime/CoreUObject/Public/Misc/RedirectCollector.h
@@ -14,7 +14,10 @@
 #include "HAL/CriticalSection.h"
 #include "HAL/Platform.h"
 #include "Misc/ScopeLock.h"
+#include "Misc/TVariant.h"
+#include "Templates/Function.h"
 #include "Templates/TypeHash.h"
+#include "Templates/UniquePtr.h"
 #include "UObject/NameTypes.h"
 #include "UObject/SoftObjectPath.h"
 #include "UObject/UnrealNames.h"
@@ -114,7 +117,7 @@ public:
 
 	/** Returns a remapped asset path, if it returns null there is no relevant redirector */
 	/** Returns a remapped asset path, if there is no relevant redirector, the return value reports true from IsNull() */
-	COREUOBJECT_API FSoftObjectPath GetAssetPathRedirection(const FSoftObjectPath& OriginalPath);
+	COREUOBJECT_API FSoftObjectPath GetAssetPathRedirection(const FSoftObjectPath& OriginalPath) const;
 
 	/**
 	 * Do we have any references to resolve.
@@ -139,19 +142,39 @@ public:
 	COREUOBJECT_API void OnStartupPackageLoadComplete();
 
 	/** Access to the collected list of redirects when already holding the lock. */
-	const TMap<FSoftObjectPath, FSoftObjectPath>& GetObjectPathRedirectionMapUnderLock(const UE::TDynamicUniqueLock<FCriticalSection>& Lock) const
+	UE_DEPRECATED(5.6, "Use EnumerateRedirectsUnderLock instead.")
+	COREUOBJECT_API const TMap<FSoftObjectPath, FSoftObjectPath>&
+	GetObjectPathRedirectionMapUnderLock(const UE::TDynamicUniqueLock<FCriticalSection>& Lock) const;
+
+	/**
+	 * Data about a redirector that has been reported to the RedirectCollector, both the input Source and FirstTarget
+	 * data, and the derived FinalTarget data.
+	 */
+	struct FRedirectionData
 	{
-		ensure(Lock.OwnsLock());
-		return ObjectPathRedirectionMap;
-	}
-	
+	public:
+		const FSoftObjectPath& GetSource() const;
+		const FSoftObjectPath& GetFirstTarget() const;
+		const FSoftObjectPath& GetFinalTarget() const;
+
+	private:
+		const FSoftObjectPath& Source;
+		const FSoftObjectPath& FirstTarget;
+		const FSoftObjectPath& FinalTarget;
+		FRedirectionData(const FSoftObjectPath& InSource, const FSoftObjectPath& InFirstTarget,
+			const FSoftObjectPath& InFinalTarget);
+		friend FRedirectCollector;
+	};
+	/**
+	 * Access to the collected list of redirects through a callback function. The function is called inside the
+	 * RedirectCollector's lock and must not call any other functions on the RedirectCollector; doing so will deadlock.
+	 */
+	COREUOBJECT_API void EnumerateRedirectsUnderLock(TFunctionRef<void(FRedirectionData&)>) const;
+
 	/** Returns the set of paths, if any, that are redirected TO the provided path.*/
 	COREUOBJECT_API void GetAllSourcePathsForTargetPath(const FSoftObjectPath& TargetPath, TArray<FSoftObjectPath>& OutSourcePaths) const;
 
-	/** Used with GetObjectPathRedirectionMapUnderLock with code like:
-	 *  UE::TDynamicUniqueLock<FCriticalSection> ScopeLock(GRedirectCollector.AcquireLock());
-	 *  GRedirectCollector.GetObjectPathRedirectionMapUnderLock(ScopeLock);
-	 */
+	UE_DEPRECATED(5.6, "Use EnumerateRedirectsUnderLock instead.")
 	UE::TDynamicUniqueLock<FCriticalSection> AcquireLock() const
 	{
 		return UE::TDynamicUniqueLock<FCriticalSection>(CriticalSection);
@@ -165,20 +188,54 @@ private:
 	/** Handles removing forward and reverse map entries. Must be called while holding the critical section */
 	bool TryRemoveObjectPathRedirectionInternal(const FSoftObjectPath& Source);
 
+	/**
+	 * Searches the graph of FirstTargets starting at the input FirstTarget, to find the FinalTarget, the first one
+	 * that is not itself a registered redirector. FirstTarget must not be null, see comment in function. Must be
+	 * called while holding the critical section.
+	 */
+	const FSoftObjectPath& TraverseToFinalTarget(const FSoftObjectPath* FirstTarget) const;
+
 	/** A map of assets referenced by soft object paths, with the key being the package with the reference */
 	typedef TSet<FSoftObjectPathProperty> FSoftObjectPathPropertySet;
 	typedef TMap<FName, FSoftObjectPathPropertySet> FSoftObjectPathMap;
 
 	/** Return whether SoftObjectPathExclusions are currently being tracked, based on commandline and cook phase. */
-	COREUOBJECT_API bool ShouldTrackPackageReferenceTypes();
+	bool ShouldTrackPackageReferenceTypes();
 
 	/** The discovered references that should be followed during cook */
 	FSoftObjectPathMap SoftObjectPathMap;
 	/** The discovered references to packages and the collect type for whether they should be followed during cook. */
 	TMap<FName, TMap<FName, ESoftObjectPathCollectType>> PackageReferenceTypes;
 
+	/** Structure to hold the target data for a redirect when the chained final target != the input first target. */
+	struct FChainedRedirectionData
+	{
+		FSoftObjectPath FirstTarget;
+		FSoftObjectPath FinalTarget;
+	};
+	/**
+	 * Variant structure to save memory for storing redirects' target data: most redirects are non-chained, and
+	 * FirstTarget == FinalTarget; a few can end up being chained and we need to store extra data for them.
+	 */
+	struct FSimpleOrChainedRedirect
+	{
+	public:
+		FSimpleOrChainedRedirect();
+		FSimpleOrChainedRedirect(FSoftObjectPath InSimpleTarget);
+		FSimpleOrChainedRedirect(FSoftObjectPath InFirstTarget, FSoftObjectPath InFinalTarget);
+		static FSimpleOrChainedRedirect ConstructSimpleOrChained(FSoftObjectPath InFirstTarget,
+			FSoftObjectPath InFinalTarget);
+
+		const FSoftObjectPath& GetFirstTarget() const;
+		const FSoftObjectPath& GetFinalTarget() const;
+		bool IsSimpleRedirect() const;
+		bool IsChainedRedirect() const;
+
+	private:
+		TVariant<FSoftObjectPath, TUniquePtr<FChainedRedirectionData>> Data;
+	};
 	/** When saving, apply this remapping to all soft object paths */
-	TMap<FSoftObjectPath, FSoftObjectPath> ObjectPathRedirectionMap;
+	TMap<FSoftObjectPath, FSimpleOrChainedRedirect> ObjectPathRedirectionMap;
 
 	/** A reverse lookup map for use with GetAllSourcePathsForTargetPath */
 	typedef TArray<FSoftObjectPath, TInlineAllocator<1>> ObjectPathSourcesArray;
@@ -201,4 +258,82 @@ private:
 // global redirect collector callback structure
 COREUOBJECT_API extern FRedirectCollector GRedirectCollector;
 
+
+///////////////////////////////////////////////////////
+// Inline implementations
+///////////////////////////////////////////////////////
+
+
+inline FRedirectCollector::FSimpleOrChainedRedirect::FSimpleOrChainedRedirect()
+	: Data(TInPlaceType<FSoftObjectPath>(), FSoftObjectPath())
+{
+}
+
+inline FRedirectCollector::FSimpleOrChainedRedirect::FSimpleOrChainedRedirect(
+	FSoftObjectPath InSimpleTarget)
+	: Data(TInPlaceType<FSoftObjectPath>(), MoveTemp(InSimpleTarget))
+{
+}
+
+inline FRedirectCollector::FSimpleOrChainedRedirect::FSimpleOrChainedRedirect(
+	FSoftObjectPath InFirstTarget, FSoftObjectPath InFinalTarget)
+	: Data(TInPlaceType<TUniquePtr<FChainedRedirectionData>>(),
+		new FChainedRedirectionData{ MoveTemp(InFirstTarget), MoveTemp(InFinalTarget) })
+{
+}
+
+inline FRedirectCollector::FSimpleOrChainedRedirect
+FRedirectCollector::FSimpleOrChainedRedirect::ConstructSimpleOrChained(
+	FSoftObjectPath InFirstTarget, FSoftObjectPath InFinalTarget)
+{
+	return InFirstTarget == InFinalTarget
+		? FSimpleOrChainedRedirect(MoveTemp(InFirstTarget))
+		: FSimpleOrChainedRedirect(MoveTemp(InFirstTarget), MoveTemp(InFinalTarget));
+}
+
+inline const FSoftObjectPath& FRedirectCollector::FSimpleOrChainedRedirect::GetFirstTarget() const
+{
+	return IsSimpleRedirect()
+		? Data.Get<FSoftObjectPath>()
+		: Data.Get<TUniquePtr<FChainedRedirectionData>>()->FirstTarget;
+}
+
+inline const FSoftObjectPath& FRedirectCollector::FSimpleOrChainedRedirect::GetFinalTarget() const
+{
+	return IsSimpleRedirect()
+		? Data.Get<FSoftObjectPath>()
+		: Data.Get<TUniquePtr<FChainedRedirectionData>>()->FinalTarget;
+}
+
+inline bool FRedirectCollector::FSimpleOrChainedRedirect::IsSimpleRedirect() const
+{
+	return Data.IsType<FSoftObjectPath>();
+}
+
+inline bool FRedirectCollector::FSimpleOrChainedRedirect::IsChainedRedirect() const
+{
+	return Data.IsType<TUniquePtr<FChainedRedirectionData>>();
+}
+
+inline const FSoftObjectPath& FRedirectCollector::FRedirectionData::GetSource() const
+{
+	return Source;
+}
+
+inline const FSoftObjectPath& FRedirectCollector::FRedirectionData::GetFirstTarget() const
+{
+	return FirstTarget;
+}
+
+inline const FSoftObjectPath& FRedirectCollector::FRedirectionData::GetFinalTarget() const
+{
+	return FinalTarget;
+}
+
+inline FRedirectCollector::FRedirectionData::FRedirectionData(const FSoftObjectPath& InSource,
+	const FSoftObjectPath& InFirstTarget, const FSoftObjectPath& InFinalTarget)
+	: Source(InSource), FirstTarget(InFirstTarget), FinalTarget(InFinalTarget)
+{
+}
+
 #endif // WITH_EDITOR
