[UE] Decouple IDO placeholder types from UObject lifetime management. This is already handled by serialization of the internal UClass reference.

diff --git a/Engine/Source/Editor/UnrealEd/Private/Kismet2/KismetReinstanceUtilities.cpp b/Engine/Source/Editor/UnrealEd/Private/Kismet2/KismetReinstanceUtilities.cpp
--- a/Engine/Source/Editor/UnrealEd/Private/Kismet2/KismetReinstanceUtilities.cpp
+++ b/Engine/Source/Editor/UnrealEd/Private/Kismet2/KismetReinstanceUtilities.cpp
@@ -2971,12 +2971,6 @@ void FBlueprintCompileReinstancer::ReplaceInstancesOfClass_Inner(const TMap<UCla
 			
 			OldToNewInstanceMap.Add(OldToNew.Key, OldToNew.Value);
 			SourceObjects.Add(OldToNew.Key);
-
-			// Remove references to placeholder types before they get replaced below.
-			if (UE::FPropertyBagRepository::IsPropertyBagPlaceholderType(OldToNew.Key))
-			{
-				UE::FPropertyBagRepository::RemovePropertyBagPlaceholderType(OldToNew.Key);
-			}
 		}
 	}
 
diff --git a/Engine/Source/Runtime/CoreUObject/Internal/UObject/PropertyBagRepository.h b/Engine/Source/Runtime/CoreUObject/Internal/UObject/PropertyBagRepository.h
--- a/Engine/Source/Runtime/CoreUObject/Internal/UObject/PropertyBagRepository.h
+++ b/Engine/Source/Runtime/CoreUObject/Internal/UObject/PropertyBagRepository.h
@@ -5,11 +5,12 @@
 #include "CoreTypes.h"
 #include "Misc/TransactionallySafeCriticalSection.h"
 #include "Containers/Map.h"
-#include "UObject/GCObject.h"
 #include "Templates/FunctionFwd.h"
+#include "UObject/ObjectPtr.h"
 
 class UObject;
 class FMemoryArchive;
+struct FPropertyChangedChainEvent;
 
 namespace UE
 {
@@ -27,7 +28,7 @@ struct FPropertyBagRepositoryStats
 #endif
 
 // Singleton class tracking property bag association with objects
-class FPropertyBagRepository : public FGCObject
+class FPropertyBagRepository
 {
 	struct FPropertyBagAssociationData
 	{
@@ -122,11 +123,6 @@ public:
 	UE_INTERNAL void GatherStats(FPropertyBagRepositoryStats& Stats);
 	#endif
 
-	// FGCObject interface
-	virtual void AddReferencedObjects(FReferenceCollector& Collector) override;
-	virtual FString GetReferencerName() const override;
-	// End FGCObject interface
-
 	// query for whether or not the given struct/class is a placeholder type
 	static COREUOBJECT_API bool IsPropertyBagPlaceholderType(const UStruct* Type);
 	// query for whether or not the given object was created as a placeholder type
@@ -154,13 +150,6 @@ public:
 		return Cast<UClass>(CreatePropertyBagPlaceholderType(Outer, Class, Name, Flags, T::StaticClass()));
 	}
 
-	/**
-	 * Remove a placeholder type object from the internal registry.
-	 * 
-	 * @param PlaceholderType	Placeholder type object to remove.
-	 */
-	static COREUOBJECT_API void RemovePropertyBagPlaceholderType(UStruct* PlaceholderType);
-
 private:
 	void Lock() const { CriticalSection.Lock(); }
 	void Unlock() const { CriticalSection.Unlock(); }
diff --git a/Engine/Source/Runtime/CoreUObject/Private/UObject/PropertyBagRepository.cpp b/Engine/Source/Runtime/CoreUObject/Private/UObject/PropertyBagRepository.cpp
--- a/Engine/Source/Runtime/CoreUObject/Private/UObject/PropertyBagRepository.cpp
+++ b/Engine/Source/Runtime/CoreUObject/Private/UObject/PropertyBagRepository.cpp
@@ -47,18 +47,12 @@ public:
 		return Instance;
 	}
 
-	void AddReferencedObjects(FReferenceCollector& Collector)
-	{
-		ConsumePendingPlaceholderTypes();
-		Collector.AddStableReferenceSet(&PlaceholderTypes);
-	}
-
-	void Add(UStruct* Type)
+	void Add(const UStruct* Type)
 	{
 		PendingPlaceholderTypes.Enqueue(Type);
 	}
 
-	void Remove(UStruct* Type)
+	void Remove(const UStruct* Type)
 	{
 		PlaceholderTypes.Remove(Type);
 	}
@@ -85,7 +79,7 @@ protected:
 		{
 			FScopeLock ScopeLock(&CriticalSection);
 
-			TObjectPtr<UStruct> PendingType;
+			TObjectPtr<const UStruct> PendingType;
 			while(PendingPlaceholderTypes.Dequeue(PendingType))
 			{
 				PlaceholderTypes.Add(PendingType);
@@ -97,10 +91,10 @@ private:
 	FCriticalSection CriticalSection;
 
 	// List of types that have been registered.
-	TSet<TObjectPtr<UStruct>> PlaceholderTypes;
+	TSet<TObjectPtr<const UStruct>> PlaceholderTypes;
 
 	// Types that have been added but not yet registered. Utilizes a thread-safe queue so we can avoid race conditions during an async load.
-	TQueue<TObjectPtr<UStruct>> PendingPlaceholderTypes;
+	TQueue<TObjectPtr<const UStruct>> PendingPlaceholderTypes;
 };
 
 class FPropertyBagRepositoryLock
@@ -742,6 +736,13 @@ void FPropertyBagRepository::MarkAsFixedUp(const UObject* Object)
 
 bool FPropertyBagRepository::RemoveAssociationUnsafe(const UObject* Owner)
 {
+	const UStruct* OwnerAsTypeObject = Cast<const UStruct>(Owner);
+	if (OwnerAsTypeObject && IsPropertyBagPlaceholderType(OwnerAsTypeObject))
+	{
+		FPropertyBagPlaceholderTypeRegistry::Get().Remove(OwnerAsTypeObject);
+		return true;
+	}
+
 	FPropertyBagAssociationData OldData;
 	if (AssociatedData.RemoveAndCopyValue(Owner, OldData))
 	{
@@ -814,16 +815,6 @@ const UObject* FPropertyBagRepository::FindInstanceForDataObject(const UObject*
 	return Owner ? *Owner : nullptr;
 }
 
-void FPropertyBagRepository::AddReferencedObjects(FReferenceCollector& Collector)
-{
-	FPropertyBagPlaceholderTypeRegistry::Get().AddReferencedObjects(Collector);
-}
-
-FString FPropertyBagRepository::GetReferencerName() const
-{
-	return TEXT("FPropertyBagRepository");
-}
-
 UObject* FPropertyBagRepository::CreateInstanceDataObject(UObject* Owner, FArchive& Archive, int64 StartOffset, int64 EndOffset, bool bIsArchetype)
 {
 	// Limit the scope of the lock during to the find because this calls itself recursively to handle archetypes.
@@ -996,19 +987,13 @@ UStruct* FPropertyBagRepository::CreatePropertyBagPlaceholderType(UObject* Outer
 		PlaceholderTypeAsClass->AssembleReferenceTokenStream();
 	}
 
-	// Use the property bag repository for now to manage property bag placeholder types (e.g. object lifetime).
-	// Note: The object lifetime of instances of this type will rely on existing references that are serialized.
+	// Use the property bag repository for now to register property bag placeholder types for query purposes.
+	// Note: Object lifetimes of this type and its instances depend on existing references that are serialized.
 	FPropertyBagPlaceholderTypeRegistry::Get().Add(PlaceholderType);
 
 	return PlaceholderType;
 }
 
-void FPropertyBagRepository::RemovePropertyBagPlaceholderType(UStruct* PlaceholderType)
-{
-	ensure(IsPropertyBagPlaceholderType(PlaceholderType));
-	FPropertyBagPlaceholderTypeRegistry::Get().Remove(PlaceholderType);
-}
-
 #if STATS
 void FPropertyBagRepository::GatherStats(FPropertyBagRepositoryStats& Stats)
 {
