[CoreUObject]
* Changed TLinkerImportPlaceholder to not use template specialization in header

diff --git a/Engine/Source/Runtime/CoreUObject/Private/UObject/LinkerPlaceholderBase.cpp b/Engine/Source/Runtime/CoreUObject/Private/UObject/LinkerPlaceholderBase.cpp
--- a/Engine/Source/Runtime/CoreUObject/Private/UObject/LinkerPlaceholderBase.cpp
+++ b/Engine/Source/Runtime/CoreUObject/Private/UObject/LinkerPlaceholderBase.cpp
@@ -599,7 +599,7 @@ int32 FLinkerPlaceholderBase::ResolvePlaceholderPropertyValues(UObject* NewObjec
 
 //------------------------------------------------------------------------------
 template<>
-int32 TLinkerImportPlaceholder<UClass>::ResolvePropertyReferences(UClass* ReplacementClass)
+int32 TLinkerImportPlaceholder<UClass>::SpecializedResolvePropertyReferences(UClass* ReplacementClass)
 {
 	int32 ReplacementCount = 0;
 	UClass* PlaceholderClass = CastChecked<UClass>(GetPlaceholderAsUObject());
@@ -720,7 +720,7 @@ int32 TLinkerImportPlaceholder<UClass>::ResolvePropertyReferences(UClass* Replac
 
 //------------------------------------------------------------------------------
 template<>
-int32 TLinkerImportPlaceholder<UFunction>::ResolvePropertyReferences(UFunction* ReplacementFunc)
+int32 TLinkerImportPlaceholder<UFunction>::SpecializedResolvePropertyReferences(UFunction* ReplacementFunc)
 {
 	int32 ReplacementCount = 0;
 	UFunction* PlaceholderFunc = CastChecked<UFunction>(GetPlaceholderAsUObject());
diff --git a/Engine/Source/Runtime/CoreUObject/Private/UObject/LinkerPlaceholderBase.h b/Engine/Source/Runtime/CoreUObject/Private/UObject/LinkerPlaceholderBase.h
--- a/Engine/Source/Runtime/CoreUObject/Private/UObject/LinkerPlaceholderBase.h
+++ b/Engine/Source/Runtime/CoreUObject/Private/UObject/LinkerPlaceholderBase.h
@@ -341,6 +341,11 @@ private:
 	 */
 	int32 ResolvePropertyReferences(PlaceholderType* ReplacementObj);
 
+	/**
+	 * Specialization stored in source file
+	 */
+	int32 SpecializedResolvePropertyReferences(PlaceholderType* ReplacementObj);
+
 	/**
 	 * Iterates through all known ReferencingScriptExpressions and replaces 
 	 * references to this placeholder with the specified replacement object.
@@ -368,15 +373,3 @@ private:
 
 // Templatized implementation
 #include "UObject/LinkerPlaceholderBase.inl" // IWYU pragma: export
-
-/*******************************************************************************
- * TLinkerImportPlaceholder<> Specializations
- ******************************************************************************/
- 
-/**  */
-template<>
-int32 TLinkerImportPlaceholder<UClass>::ResolvePropertyReferences(UClass* ReplacementClass);
-
-/**  */
-template<>
-int32 TLinkerImportPlaceholder<UFunction>::ResolvePropertyReferences(UFunction* ReplacementFunc);
diff --git a/Engine/Source/Runtime/CoreUObject/Private/UObject/LinkerPlaceholderBase.inl b/Engine/Source/Runtime/CoreUObject/Private/UObject/LinkerPlaceholderBase.inl
--- a/Engine/Source/Runtime/CoreUObject/Private/UObject/LinkerPlaceholderBase.inl
+++ b/Engine/Source/Runtime/CoreUObject/Private/UObject/LinkerPlaceholderBase.inl
@@ -141,10 +141,17 @@ void TLinkerImportPlaceholder<PlaceholderType>::AddReferencingScriptExpr(Placeho
 template<class PlaceholderType>
 int32 TLinkerImportPlaceholder<PlaceholderType>::ResolvePropertyReferences(PlaceholderType* ReplacementObj)
 {
+	if constexpr (std::is_same_v<PlaceholderType, UClass> || std::is_same_v<PlaceholderType, UFunction>)
+	{
+		return SpecializedResolvePropertyReferences(ReplacementObj);
+	}
+	else
+	{
 #if !CHECK_PUREVIRTUALS
-	// requires template specialization (technically not "pure virtual"):
-	PURE_VIRTUAL(TLinkerImportPlaceholder<PlaceholderType>::ResolvePropertyReferences, return 0;)
+		// requires template specialization (technically not "pure virtual"):
+		PURE_VIRTUAL(TLinkerImportPlaceholder<PlaceholderType>::ResolvePropertyReferences, return 0;)
 #endif
+	}
 }
 
 //------------------------------------------------------------------------------
