Skipping property initialization tracking during SubObjectsShadowing.
When SubObjectsShadowing is enabled uninitialized properties can be incorrectly tracked as initialized.

#rnx

diff --git a/Engine/Source/Runtime/CoreUObject/Private/UObject/Class.cpp b/Engine/Source/Runtime/CoreUObject/Private/UObject/Class.cpp
--- a/Engine/Source/Runtime/CoreUObject/Private/UObject/Class.cpp
+++ b/Engine/Source/Runtime/CoreUObject/Private/UObject/Class.cpp
@@ -1456,13 +1456,16 @@ void UStruct::LoadTaggedPropertiesFromText(FStructuredArchive::FSlot Slot, uint8
 			#if WITH_EDITORONLY_DATA
 				if (bSerialized)
 				{
-					if (SerializeContext->bTrackInitializedProperties)
-					{
-						UE::FInitializedPropertyValueState(this, Data).Set(Property, ItemIndex);
-					}
-					if (SerializeContext->bTrackSerializedProperties)
+					if (Tag.OverrideOperation != EOverriddenPropertyOperation::SubObjectsShadowing)
 					{
-						UE::FSerializedPropertyValueState(this, Data).Set(Property, ItemIndex);
+						if (SerializeContext->bTrackInitializedProperties)
+						{
+							UE::FInitializedPropertyValueState(this, Data).Set(Property, ItemIndex);
+						}
+						if (SerializeContext->bTrackSerializedProperties)
+						{
+							UE::FSerializedPropertyValueState(this, Data).Set(Property, ItemIndex);
+						}
 					}
 				}
 			#endif
@@ -1936,13 +1939,16 @@ void UStruct::SerializeVersionedTaggedProperties(FStructuredArchive::FSlot Slot,
 					}
 
 				#if WITH_EDITORONLY_DATA
-					if (SerializeContext->bTrackInitializedProperties && Property)
-					{
-						UE::FInitializedPropertyValueState(this, Data).Set(Property, Tag.ArrayIndex);
-					}
-					if (SerializeContext->bTrackSerializedProperties && Property)
+					if (Tag.OverrideOperation != EOverriddenPropertyOperation::SubObjectsShadowing)
 					{
-						UE::FSerializedPropertyValueState(this, Data).Set(Property, Tag.ArrayIndex);
+						if (SerializeContext->bTrackInitializedProperties && Property)
+						{
+							UE::FInitializedPropertyValueState(this, Data).Set(Property, Tag.ArrayIndex);
+						}
+						if (SerializeContext->bTrackSerializedProperties && Property)
+						{
+							UE::FSerializedPropertyValueState(this, Data).Set(Property, Tag.ArrayIndex);
+						}
 					}
 				#endif
 				}
