Fixed name of type in static_assert.

diff --git a/Engine/Source/Runtime/CoreUObject/Tests/ObjectPtrTest.cpp b/Engine/Source/Runtime/CoreUObject/Tests/ObjectPtrTest.cpp
--- a/Engine/Source/Runtime/CoreUObject/Tests/ObjectPtrTest.cpp
+++ b/Engine/Source/Runtime/CoreUObject/Tests/ObjectPtrTest.cpp
@@ -135,8 +135,8 @@ static_assert(TModels_V<CEqualityComparableWith, FConstObjectPtr, TYPE_OF_NULLPT
 static_assert(TModels_V<CEqualityComparableWith, FMutableObjectPtr, TYPE_OF_NULLPTR>, "Must be able to compare equality and inequality bidirectionally between TObjectPtr<UObject> and nullptr");
 
 // Ensure allowing comparison against NULL doesn't mean allowing comparison against int
-static_assert(!TModels_V<CEqualityComparableWith, FConstObjectPtr, int>, "Should not be able to compare equality and inequality bidirectionally between TObjectPtr<const UObject> and long");
-static_assert(!TModels_V<CEqualityComparableWith, FMutableObjectPtr, int>, "Should not be able to compare equality and inequality bidirectionally between TObjectPtr<UObject> and long");
+static_assert(!TModels_V<CEqualityComparableWith, FConstObjectPtr, int>, "Should not be able to compare equality and inequality bidirectionally between TObjectPtr<const UObject> and int");
+static_assert(!TModels_V<CEqualityComparableWith, FMutableObjectPtr, int>, "Should not be able to compare equality and inequality bidirectionally between TObjectPtr<UObject> and int");
 
 
 #if WITH_LOW_LEVEL_TESTS
