[EditorPaths] Fix soft references to actors within Level Instances being broken on asset reloads

When assets are reloaded soft references to living objects are fixed-up. This however needs to take editor paths feature into account (if enabled), otherwise paths that were proper editor paths get changed to paths to temporary actor instances living in level instances, which effectively breaks said references.

diff --git a/Engine/Source/Runtime/CoreUObject/Private/Serialization/ArchiveUObject.cpp b/Engine/Source/Runtime/CoreUObject/Private/Serialization/ArchiveUObject.cpp
--- a/Engine/Source/Runtime/CoreUObject/Private/Serialization/ArchiveUObject.cpp
+++ b/Engine/Source/Runtime/CoreUObject/Private/Serialization/ArchiveUObject.cpp
@@ -8,6 +8,7 @@
 #include "UObject/UnrealType.h"
 #include "Serialization/ObjectAndNameAsStringProxyArchive.h"
 #include "Serialization/ArchiveReplaceObjectRef.h"
+#include "Misc/EditorPathHelper.h"
 #include "Misc/PackageName.h"
 
 DEFINE_LOG_CATEGORY_STATIC(LogArchiveUObject, Log, All);
@@ -115,7 +116,11 @@ FArchive& FArchiveUObject::SerializeSoftObjectPtr(FArchive& Ar, FSoftObjectPtr&
 
 		if (Ar.IsLoading() || (Object && Ar.IsModifyingWeakAndStrongReferences()))
 		{
+#if WITH_EDITOR
+			Value = FEditorPathHelper::GetEditorPath(Object);
+#else
 			Value = Object;
+#endif
 		}
 	}
 
