[UE] Remove hack to impose explicit instanced reference semantics on all loose IDO object properties.

Notes:
- This was added in as an early workaround in order to extend IDO support to container properties serialized using experimental override logic, which requires certain property flags to be set.
- Since we can now ascertain from the property tag whether the TPS data uses experimental override logic, I've removed the hack and replaced it with a property tag state check.
- This shouldn't have any real impact other than eliminating unnecessary proliferation to loose IDO object properties that don't require the extra semantics to be added for serialization.

diff --git a/Engine/Source/Runtime/CoreUObject/Private/UObject/InstanceDataObjectUtils.cpp b/Engine/Source/Runtime/CoreUObject/Private/UObject/InstanceDataObjectUtils.cpp
--- a/Engine/Source/Runtime/CoreUObject/Private/UObject/InstanceDataObjectUtils.cpp
+++ b/Engine/Source/Runtime/CoreUObject/Private/UObject/InstanceDataObjectUtils.cpp
@@ -948,13 +948,24 @@ namespace UE
 	}
 
 	// recursively gives a property the metadata and flags of a loose property
-	static void MarkPropertyAsLoose(FProperty* Property)
+	static void MarkPropertyAsLoose(FProperty* Property, EPropertyFlags PropertyFlags = CPF_None)
 	{
 		Property->SetMetaData(NAME_IsLooseMetadata, TEXT("True"));
-		Property->SetPropertyFlags(CPF_Edit | CPF_EditConst);
+		Property->SetPropertyFlags(CPF_Edit | CPF_EditConst | PropertyFlags);
+
 		if (const FArrayProperty* AsArrayProperty = CastField<FArrayProperty>(Property))
 		{
-			MarkPropertyAsLoose(AsArrayProperty->Inner);
+			// experimental override serialization of arrays requires certain flags be set on the inner property (it will assert otherwise)
+			if (PropertyFlags & CPF_ExperimentalOverridableLogic)
+			{
+				PropertyFlags &= ~CPF_ExperimentalOverridableLogic;
+				if (ensureMsgf(AsArrayProperty->Inner->IsA<FObjectProperty>(), TEXT("Expected array inner type to be an object property (%s: %s)"), *AsArrayProperty->GetPathName(), *AsArrayProperty->Inner->GetClass()->GetName()))
+				{
+					PropertyFlags |= CPF_InstancedReference | CPF_PersistentInstance;
+				}
+			}
+
+			MarkPropertyAsLoose(AsArrayProperty->Inner, PropertyFlags);
 		}
 		else if (const FSetProperty* AsSetProperty = CastField<FSetProperty>(Property))
 {
@@ -962,8 +973,18 @@ namespace UE
 		}
 		else if (const FMapProperty* AsMapProperty = CastField<FMapProperty>(Property))
 		{
-			MarkPropertyAsLoose(AsMapProperty->KeyProp);
-			MarkPropertyAsLoose(AsMapProperty->ValueProp);
+			// experimental override serialization of maps requires certain flags to be set on the key property (it will assert otherwise)
+			if (PropertyFlags & CPF_ExperimentalOverridableLogic)
+			{
+				PropertyFlags &= ~CPF_ExperimentalOverridableLogic;
+				if (ensureMsgf(AsMapProperty->KeyProp->IsA<FObjectProperty>(), TEXT("Expected map key type to be an object property (%s: %s)"), *AsMapProperty->GetPathName(), *AsMapProperty->KeyProp->GetClass()->GetName()))
+				{
+					PropertyFlags |= CPF_InstancedReference | CPF_PersistentInstance;
+				}
+			}
+
+			MarkPropertyAsLoose(AsMapProperty->KeyProp, PropertyFlags);
+			MarkPropertyAsLoose(AsMapProperty->ValueProp);	// override serialization doesn't require any flags on the value property
 		}
 		else if (const FOptionalProperty* AsOptionalProperty = CastField<FOptionalProperty>(Property))
 		{
@@ -978,18 +999,8 @@ namespace UE
 		}
 		else if (FObjectProperty* AsObjectProperty = CastField<FObjectProperty>(Property))
 		{
-			// Hack for now - the assumption is that IDOs are generated only for class types that impose this flag on all object properties.
-			// There is currently an implicit assumption in the serialization logic that all inner properties have this flag set for containers.
-			// Since this is a "loose" property, the underlying type will not explicitly tell us this, and there is no way to know from the tagged
-			// property data stream if this flag was set when it was last serialized for the instance in question. So for now we just always set it.
-			// 
-			// Note that we are not currently including other related flags such as CPF_InstancedReference, CPF_ContainsInstancedReference, etc.
-			// For the most part those have been relegated to object construction and loading paths. We are not instancing IDO types explicitly;
-			// they are instead serving as an editable data archetype for the actual instance, whose type may impose some post-initialization
-			// effects on that data as part of the construction/serialization path.
-			// 
-			// @todo - Remove if/when this flag is no longer required to signal whether this value is to be resolved via a subobject instancing graph.
-			AsObjectProperty->SetPropertyFlags(CPF_PersistentInstance| CPF_InstancedReference | CPF_TObjectPtr);
+			// TObjectPtr is required by UHT and thus for serializing its TPS data
+			AsObjectProperty->SetPropertyFlags(CPF_TObjectPtr);
 
 			// also assign the property class to UObject because loose properties can't infer their PropertyClass from TPS data so we'll assume it's as lenient as possible
 			AsObjectProperty->PropertyClass = UObject::StaticClass();
@@ -1097,8 +1108,16 @@ namespace UE
 							// skip loose types that have been explicitly excluded from IDOs
 							continue;
 						}
+						EPropertyFlags PropertyFlags = CPF_None;
+						if (const FPropertyTag* PropertyTag = It.GetNode().GetTag())
+						{
+							if (PropertyTag->bExperimentalOverridableLogic)
+							{
+								PropertyFlags |= CPF_ExperimentalOverridableLogic;
+							}
+						}
 						ConvertToInstanceDataObjectProperty(Property, Type, Outer, It.GetNode().GetSubTree(), EnumNames);
-						MarkPropertyAsLoose(Property);	// note: make sure not to mark until AFTER conversion, as this can mutate property flags on nested struct fields
+						MarkPropertyAsLoose(Property, PropertyFlags);	// note: make sure not to mark until AFTER conversion, as this can mutate property flags on nested struct fields
 						TrySetContainsLoosePropertyMetadata(Property);
 						LooseInstanceDataObjectProperties.Add(Property);
 						continue;
