[UE] Create placeholders for missing/renamed type imports to allow exported instance data to be serialized on editor load.

Notes:
- This feature is currently enabled for IDO-supported types only. At this time, the feature is not enabled for Blueprints based on IDO-supported types.
- Currently, references to placeholder exports will not be serialized by the package linker, and thus become transient after load. Persistence will follow as more dependent features are implemented.

diff --git a/Engine/Source/Editor/UnrealEd/Private/Kismet2/KismetReinstanceUtilities.cpp b/Engine/Source/Editor/UnrealEd/Private/Kismet2/KismetReinstanceUtilities.cpp
--- a/Engine/Source/Editor/UnrealEd/Private/Kismet2/KismetReinstanceUtilities.cpp
+++ b/Engine/Source/Editor/UnrealEd/Private/Kismet2/KismetReinstanceUtilities.cpp
@@ -3342,7 +3342,7 @@ void FBlueprintCompileReinstancer::PreCreateSubObjectsForReinstantiation_Inner(
 				{
 					// The old subobject class is a compiler artifact, with no reinstanced counterpart. Its source may have been renamed (or deleted) and thus it could not be recompiled.
 					// In this case, if property bag features are enabled, we attempt to redirect to a placeholder type so we can serialize to a property bag in order to preserve its data.
-					if (UE::FPropertyBagRepository::IsPropertyBagPlaceholderObjectSupportEnabled())
+					if (UE::FPropertyBagRepository::IsPropertyBagPlaceholderObjectFeatureEnabled(UE::FPropertyBagRepository::EPlaceholderObjectFeature::ReplaceMissingReinstancedTypes))
 					{
 						FTopLevelAssetPath OldSubObjectClassPath = OldSubObjectClass->GetReinstancedClassPathName();
 						if (OldSubObjectClassPath.IsValid())
diff --git a/Engine/Source/Runtime/CoreUObject/Internal/UObject/PropertyBagRepository.h b/Engine/Source/Runtime/CoreUObject/Internal/UObject/PropertyBagRepository.h
--- a/Engine/Source/Runtime/CoreUObject/Internal/UObject/PropertyBagRepository.h
+++ b/Engine/Source/Runtime/CoreUObject/Internal/UObject/PropertyBagRepository.h
@@ -123,12 +123,28 @@ public:
 	UE_INTERNAL void GatherStats(FPropertyBagRepositoryStats& Stats);
 	#endif
 
+	// Placeholder object feature flags; external callers can query a particular subfeature to see if it's enabled.
+	// Note: These have CVar counterparts defined in the .cpp file. Add an entry there to also create a toggle switch.
+	enum class EPlaceholderObjectFeature : uint8
+	{
+		/** Replace missing type imports with a placeholder type on load so that exports of the missing type can be serialized */
+		ReplaceMissingTypeImportsOnLoad,
+		/** Serialize references to placeholder exports on load so that they remain persistent and become visible to referencers */
+		SerializeExportReferencesOnLoad,
+		/** Replace missing types after reinstancing with a placeholder type so that data is not lost after rebuilding scripts */
+		ReplaceMissingReinstancedTypes,
+		/** Replace dead class types after a script compile with a placeholder type so that data is not lost after a script build failure */
+		ReplaceDeadClassInstanceTypes,
+	};
+
 	// query for whether or not the given struct/class is a placeholder type
 	static COREUOBJECT_API bool IsPropertyBagPlaceholderType(const UStruct* Type);
 	// query for whether or not the given object was created as a placeholder type
 	static COREUOBJECT_API bool IsPropertyBagPlaceholderObject(const UObject* Object);
 	// query for whether or not creating property bag placeholder objects should be allowed
 	static COREUOBJECT_API bool IsPropertyBagPlaceholderObjectSupportEnabled();
+	// query for whether or not a specific property bag placeholder object feature is enabled
+	static COREUOBJECT_API bool IsPropertyBagPlaceholderObjectFeatureEnabled(EPlaceholderObjectFeature Feature);
 
 	/**
 	 * Create a new placeholder type object to swap in for a missing class/struct. An object of
diff --git a/Engine/Source/Runtime/CoreUObject/Private/UObject/LinkerLoad.cpp b/Engine/Source/Runtime/CoreUObject/Private/UObject/LinkerLoad.cpp
--- a/Engine/Source/Runtime/CoreUObject/Private/UObject/LinkerLoad.cpp
+++ b/Engine/Source/Runtime/CoreUObject/Private/UObject/LinkerLoad.cpp
@@ -4303,7 +4303,7 @@ UClass* FLinkerLoad::GetExportLoadClass(int32 Index)
 #if WITH_EDITOR
 UClass* FLinkerLoad::TryCreatePlaceholderClassImport(int32 ImportIndex)
 {
-	const bool bAllowPlaceholderImportTypes = UE::FPropertyBagRepository::IsPropertyBagPlaceholderObjectSupportEnabled();
+	const bool bAllowPlaceholderImportTypes = UE::FPropertyBagRepository::IsPropertyBagPlaceholderObjectFeatureEnabled(UE::FPropertyBagRepository::EPlaceholderObjectFeature::ReplaceMissingTypeImportsOnLoad);
 	if (!bAllowPlaceholderImportTypes)
 	{
 		return nullptr;
@@ -6274,13 +6274,20 @@ FArchive& FLinkerLoad::operator<<( UObject*& Object )
 #if WITH_EDITOR
 	if (Object && UE::FPropertyBagRepository::IsPropertyBagPlaceholderObject(Object))
 	{
-		const FObjectProperty* ObjectProperty = CastField<FObjectProperty>(GetSerializedProperty());
-		if (!ObjectProperty || !Object->GetClass()->IsChildOf(ObjectProperty->PropertyClass))
+		if (UE::FPropertyBagRepository::IsPropertyBagPlaceholderObjectFeatureEnabled(UE::FPropertyBagRepository::EPlaceholderObjectFeature::SerializeExportReferencesOnLoad))
+		{
+			const FObjectProperty* ObjectProperty = CastField<FObjectProperty>(GetSerializedProperty());
+			if (!ObjectProperty || !Object->GetClass()->IsChildOf(ObjectProperty->PropertyClass))
+			{
+				// This is needed because the pointer's type is checked only at compile time, which may not match the property
+				// bag placeholder object's type at runtime, and so we can't allow it to be dereferenced as the wrong base type.
+				// Note: These currently won't be discovered for replacement at reinstancing time, so it will remain set to NULL.
+				UE_LOG(LogLinker, Warning, TEXT("Serializing reference to \"%s\" as NULL to ensure type safety."), *Object->GetPathName());
+				Object = nullptr;
+			}
+		}
+		else
 		{
-			// This is needed because the pointer's type is checked only at compile time, which may not match the property
-			// bag placeholder object's type at runtime, and so we can't allow it to be dereferenced as the wrong base type.
-			// Note: These currently won't be discovered for replacement at reinstancing time, so it will remain set to NULL.
-			UE_LOG(LogLinker, Warning, TEXT("Serializing reference to \"%s\" as NULL to ensure type safety."), *Object->GetPathName());
 			Object = nullptr;
 		}
 	}
@@ -6301,20 +6308,27 @@ FArchive& FLinkerLoad::operator<<(FObjectPtr& ObjectPtr)
 		// Note: References to placeholder objects cannot resolve to it if the underlying pointer type is unsafe.
 		if (ResolvedObject && UE::FPropertyBagRepository::IsPropertyBagPlaceholderObject(ResolvedObject))
 		{
-			const FObjectProperty* ObjectProperty = CastField<FObjectProperty>(GetSerializedProperty());
-			if (!ObjectProperty || !ResolvedObject->GetClass()->IsChildOf(ObjectProperty->PropertyClass))
+			if (UE::FPropertyBagRepository::IsPropertyBagPlaceholderObjectFeatureEnabled(UE::FPropertyBagRepository::EPlaceholderObjectFeature::SerializeExportReferencesOnLoad))
 			{
+				const FObjectProperty* ObjectProperty = CastField<FObjectProperty>(GetSerializedProperty());
+				if (!ObjectProperty || !ResolvedObject->GetClass()->IsChildOf(ObjectProperty->PropertyClass))
+				{
 #if UE_WITH_OBJECT_HANDLE_LATE_RESOLVE && UE_WITH_OBJECT_HANDLE_TYPE_SAFETY
-				// If type safety features are enabled, create a packed reference mapping for the placeholder-typed object.
-				// This resolves to NULL on access since the reference can't be cast to a pointer bound to its original type.
-				// However, the underlying value (when serialized) will always resolve to the placeholder object (e.g. for GC).
-				return FObjectPtr({ UE::CoreUObject::Private::MakePackedObjectRef(ResolvedObject).EncodedRef });
+					// If type safety features are enabled, create a packed reference mapping for the placeholder-typed object.
+					// This resolves to NULL on access since the reference can't be cast to a pointer bound to its original type.
+					// However, the underlying value (when serialized) will always resolve to the placeholder object (e.g. for GC).
+					return FObjectPtr({ UE::CoreUObject::Private::MakePackedObjectRef(ResolvedObject).EncodedRef });
 #else
-				// If type safety features are disabled, serialize it as an unsafe reference to a placeholder-typed object.
-				// Note: Similar to hard references above, this means we won't find it for replacement at reinstancing time.
-				UE_LOG(LogLinker, Warning, TEXT("Serializing reference to \"%s\" as NULL to ensure type safety. This will lead to data loss if the referencing object is saved."), *ResolvedObject->GetPathName());
-				ResolvedObject = nullptr;
+					// If type safety features are disabled, serialize it as an unsafe reference to a placeholder-typed object.
+					// Note: Similar to hard references above, this means we won't find it for replacement at reinstancing time.
+					UE_LOG(LogLinker, Warning, TEXT("Serializing reference to \"%s\" as NULL to ensure type safety. This will lead to data loss if the referencing object is saved."), *ResolvedObject->GetPathName());
+					ResolvedObject = nullptr;
 #endif
+				}
+			}
+			else
+			{
+				ResolvedObject = nullptr;
 			}
 		}
 #endif
@@ -6395,11 +6409,11 @@ void FLinkerLoad::BadNameIndexError(int32 NameIndex)
  */
 void FLinkerLoad::MarkScriptSerializationStart( const UObject* Obj )
 {
-	if (Obj && Obj->GetLinker() == this)
-	{
 #if WITH_EDITOR
-		bIsSerializingScriptProperties = true;
+	bIsSerializingScriptProperties = true;
 #endif
+	if (Obj && Obj->GetLinker() == this)
+	{
 		int32 Index = Obj->GetLinkerIndex();
 		if (ExportMap.IsValidIndex(Index))
 		{
@@ -6424,11 +6438,11 @@ void FLinkerLoad::MarkScriptSerializationStart( const UObject* Obj )
  */
 void FLinkerLoad::MarkScriptSerializationEnd( const UObject* Obj )
 {
-	if (Obj && Obj->GetLinker() == this)
-	{
 #if WITH_EDITOR
-		bIsSerializingScriptProperties = false;
+	bIsSerializingScriptProperties = false;
 #endif
+	if (Obj && Obj->GetLinker() == this)
+	{
 		int32 Index = Obj->GetLinkerIndex();
 		if (ExportMap.IsValidIndex(Index))
 		{
diff --git a/Engine/Source/Runtime/CoreUObject/Private/UObject/PropertyBagRepository.cpp b/Engine/Source/Runtime/CoreUObject/Private/UObject/PropertyBagRepository.cpp
--- a/Engine/Source/Runtime/CoreUObject/Private/UObject/PropertyBagRepository.cpp
+++ b/Engine/Source/Runtime/CoreUObject/Private/UObject/PropertyBagRepository.cpp
@@ -940,6 +940,40 @@ namespace Private
 		TEXT("If true, allows placeholder types to be created in place of missing types in order to redirect serialization into a property bag."),
 		ECVF_Default
 	);
+
+	using EPlaceholderObjectFeature = UE::FPropertyBagRepository::EPlaceholderObjectFeature;
+	constexpr uint32 PropertyBagPlaceholderObjectFeatureFlag(EPlaceholderObjectFeature Feature)
+	{
+		return 1 << static_cast<std::underlying_type_t<EPlaceholderObjectFeature>>(Feature);
+	}
+
+	static uint32 PropertyBagPlaceholderObjectEnabledFeaturesMask =
+		PropertyBagPlaceholderObjectFeatureFlag(EPlaceholderObjectFeature::ReplaceMissingTypeImportsOnLoad);
+
+	#define DEFINE_IDO_PLACEHOLDER_FEATURE_FLAG_CVAR(FEATURE) \
+		static bool bEnablePropertyBagPlaceholderObjectFeature_##FEATURE = !!(PropertyBagPlaceholderObjectEnabledFeaturesMask & PropertyBagPlaceholderObjectFeatureFlag(EPlaceholderObjectFeature::FEATURE)); \
+		static FAutoConsoleVariableRef CVarEnablePropertyBagPlaceholderObjectFeature_##FEATURE( \
+			TEXT("IDO.Placeholder.Feature."#FEATURE), \
+			bEnablePropertyBagPlaceholderObjectFeature_##FEATURE, \
+			TEXT("Enable/disable IDO placeholder feature: "#FEATURE), \
+			FConsoleVariableDelegate::CreateLambda([](IConsoleVariable*) \
+			{ \
+				if (bEnablePropertyBagPlaceholderObjectFeature_##FEATURE) \
+				{ \
+					PropertyBagPlaceholderObjectEnabledFeaturesMask |= PropertyBagPlaceholderObjectFeatureFlag(EPlaceholderObjectFeature::FEATURE); \
+				} \
+				else \
+				{ \
+					PropertyBagPlaceholderObjectEnabledFeaturesMask &= ~PropertyBagPlaceholderObjectFeatureFlag(EPlaceholderObjectFeature::FEATURE); \
+				} \
+			}), \
+			ECVF_Default \
+		)
+
+	DEFINE_IDO_PLACEHOLDER_FEATURE_FLAG_CVAR(ReplaceMissingTypeImportsOnLoad);
+	DEFINE_IDO_PLACEHOLDER_FEATURE_FLAG_CVAR(SerializeExportReferencesOnLoad);
+	DEFINE_IDO_PLACEHOLDER_FEATURE_FLAG_CVAR(ReplaceMissingReinstancedTypes);
+	DEFINE_IDO_PLACEHOLDER_FEATURE_FLAG_CVAR(ReplaceDeadClassInstanceTypes);
 #endif
 }
 
@@ -947,7 +981,7 @@ bool FPropertyBagRepository::IsPropertyBagPlaceholderObjectSupportEnabled()
 {
 #if WITH_EDITOR && UE_WITH_OBJECT_HANDLE_TYPE_SAFETY
 	static bool bIsInitialized = false;
-	static bool bForceDisabled = true;
+	static bool bForceDisabled = false;
 	if (!bIsInitialized)
 	{
 		Private::bEnablePropertyBagPlaceholderObjectSupport |= FParse::Param(FCommandLine::Get(), TEXT("WithPropertyBagPlaceholderObjects"));
@@ -965,6 +999,20 @@ bool FPropertyBagRepository::IsPropertyBagPlaceholderObjectSupportEnabled()
 #endif
 }
 
+bool FPropertyBagRepository::IsPropertyBagPlaceholderObjectFeatureEnabled(EPlaceholderObjectFeature Feature)
+{
+#if WITH_EDITOR
+	if (!IsPropertyBagPlaceholderObjectSupportEnabled())
+	{
+		return false;
+	}
+
+	return !!(Private::PropertyBagPlaceholderObjectEnabledFeaturesMask & Private::PropertyBagPlaceholderObjectFeatureFlag(Feature));
+#else
+	return false;
+#endif
+}
+
 UStruct* FPropertyBagRepository::CreatePropertyBagPlaceholderType(UObject* Outer, UClass* Class, FName Name, EObjectFlags Flags, UStruct* SuperStruct)
 {
 	// Generate and link a new type object using the given SuperStruct as its base.
