Create Instance Data Objects for newly created objects (e.g., spawned).

This feature can be disabled using this CVar:
IDO.EnableOnSpawnedObjects 0

#rnx

diff --git a/Engine/Source/Runtime/CoreUObject/Internal/UObject/InstanceDataObjectUtils.h b/Engine/Source/Runtime/CoreUObject/Internal/UObject/InstanceDataObjectUtils.h
--- a/Engine/Source/Runtime/CoreUObject/Internal/UObject/InstanceDataObjectUtils.h
+++ b/Engine/Source/Runtime/CoreUObject/Internal/UObject/InstanceDataObjectUtils.h
@@ -84,6 +84,11 @@ UE_API UObject* CreateInstanceDataObject(UObject* Owner, FArchive& Ar, int64 Sta
 /** Returns the IDO for the object if it has one, otherwise returns the object. */
 UE_API UObject* ResolveInstanceDataObject(UObject* Object);
 
+namespace Private
+{
+	UE_API bool TryCreateIDOForConstructedObject(UObject* Object);
+} // Private
+
 } // UE
 
 #undef UE_API
diff --git a/Engine/Source/Runtime/CoreUObject/Private/UObject/InstanceDataObjectUtils.cpp b/Engine/Source/Runtime/CoreUObject/Private/UObject/InstanceDataObjectUtils.cpp
--- a/Engine/Source/Runtime/CoreUObject/Private/UObject/InstanceDataObjectUtils.cpp
+++ b/Engine/Source/Runtime/CoreUObject/Private/UObject/InstanceDataObjectUtils.cpp
@@ -285,6 +285,13 @@ namespace UE
 		FConsoleVariableDelegate::CreateStatic(OnInstanceDataObjectSupportChanged)
 	);
 
+	bool bEnableIDOOnSpawnedObjects = true;
+	FAutoConsoleVariableRef EnableIDOOnSpawnedObjects(
+		TEXT("IDO.EnableOnSpawnedObjects"),
+		bEnableIDOOnSpawnedObjects,
+		TEXT("Allows an IDO to be created for every newly created (e.g., spawned) object.")
+	);
+
 	bool bEnableIDOUnknownProperties = true;
 	FAutoConsoleVariableRef EnableIDOUnknownProperties(
 		TEXT("IDO.Unknowns.EnableProperties"),
@@ -1511,6 +1518,55 @@ UObject* ResolveInstanceDataObject(UObject* Object)
 	return InstanceDataObject ? InstanceDataObject : Object;
 }
 
+namespace Private
+{
+
+bool TryCreateIDOForConstructedObject(UObject* Object)
+{
+	bool bDidCreateIDO = false;
+
+	if (bEnableIDOOnSpawnedObjects)
+	{
+		const UClass* ObjectClass = Object->GetClass();
+
+		// Ignore objects created very early in the initialization of CoreUObject.
+		if (ObjectClass && ObjectClass->GetClass())
+		{
+			// Do not create IDOs for the following types of objects:
+			// - Object that will be loaded: we handle those via FScopedObjectSerializeContext.
+			// - CDOs: created when the IDO class is created.
+			// - Archetypes: created on-demand.
+			// - Transient objects.
+			// - Objects with an obsolete type.
+			// - Types that do not support IDOs.
+			// - IDOs: we don't want to create IDOs for IDOs.
+			if (!Object->HasAnyFlags(RF_NeedLoad | RF_ClassDefaultObject | RF_ArchetypeObject | RF_Transient)
+				&& !Object->HasAnyInternalFlags(EInternalObjectFlags_AsyncLoading)
+				&& !ObjectClass->HasAnyClassFlags(CLASS_NewerVersionExists)
+				&& IsInstanceDataObjectSupportEnabledForClass(ObjectClass)
+				&& !Object->IsInPackage(GetTransientPackage()))
+			{
+				// Assume this object is getting loaded if an object is being serialized, so skip IDO creation.
+				const FUObjectSerializeContext* SerializeContext = FUObjectThreadContext::Get().GetSerializeContext();
+				if (!SerializeContext || !SerializeContext->SerializedObject)
+				{
+					bool bCreateInstanceDataObject = !IsInstanceDataObject(Object) && IsInstanceDataObjectSupportEnabled(Object);
+
+					if (bCreateInstanceDataObject)
+					{
+						UObject* IDO = CreateInstanceDataObject(Object);
+						bDidCreateIDO = (IDO != nullptr);
+					}
+				}
+			}
+		}
+	}
+
+	return bDidCreateIDO;
+}
+
+} // Private
+
 } // UE
 
 #endif // WITH_EDITORONLY_DATA
diff --git a/Engine/Source/Runtime/CoreUObject/Private/UObject/UObjectGlobals.cpp b/Engine/Source/Runtime/CoreUObject/Private/UObject/UObjectGlobals.cpp
--- a/Engine/Source/Runtime/CoreUObject/Private/UObject/UObjectGlobals.cpp
+++ b/Engine/Source/Runtime/CoreUObject/Private/UObject/UObjectGlobals.cpp
@@ -4980,6 +4980,10 @@ UObject* StaticConstructObject_Internal(const FStaticConstructObjectParameters&
 		Result->ClearGarbage();
 	}
 
+#if WITH_EDITORONLY_DATA
+	UE::Private::TryCreateIDOForConstructedObject(Result);
+#endif
+
 #if WITH_EDITOR
 	FCoreUObjectDelegates::OnObjectConstructed.Broadcast(Result);
 #endif
