Backing out my fix for property initialization tracking during SubObjectsShadowing.
It doesn't work because SubObjectsShadowing takes place during the serialization of IDOs.

[Backout] - CL42051502
#rnx
Original CL Desc
-----------------------------------------------------------------
Skipping property initialization tracking during SubObjectsShadowing.
When SubObjectsShadowing is enabled uninitialized properties can be incorrectly tracked as initialized.

#rnx

diff --git a/Engine/Source/Runtime/CoreUObject/Private/UObject/Class.cpp b/Engine/Source/Runtime/CoreUObject/Private/UObject/Class.cpp
--- a/Engine/Source/Runtime/CoreUObject/Private/UObject/Class.cpp
+++ b/Engine/Source/Runtime/CoreUObject/Private/UObject/Class.cpp
@@ -1458,16 +1458,13 @@ void UStruct::LoadTaggedPropertiesFromText(FStructuredArchive::FSlot Slot, uint8
 			#if WITH_EDITORONLY_DATA
 				if (bSerialized)
 				{
-					if (Tag.OverrideOperation != EOverriddenPropertyOperation::SubObjectsShadowing)
+					if (SerializeContext->bTrackInitializedProperties)
 					{
-						if (SerializeContext->bTrackInitializedProperties)
-						{
-							UE::FInitializedPropertyValueState(this, Data).Set(Property, ItemIndex);
-						}
-						if (SerializeContext->bTrackSerializedProperties)
-						{
-							UE::FSerializedPropertyValueState(this, Data).Set(Property, ItemIndex);
-						}
+						UE::FInitializedPropertyValueState(this, Data).Set(Property, ItemIndex);
+					}
+					if (SerializeContext->bTrackSerializedProperties)
+					{
+						UE::FSerializedPropertyValueState(this, Data).Set(Property, ItemIndex);
 					}
 				}
 			#endif
@@ -1940,16 +1937,13 @@ void UStruct::SerializeVersionedTaggedProperties(FStructuredArchive::FSlot Slot,
 					}
 
 				#if WITH_EDITORONLY_DATA
-					if (Tag.OverrideOperation != EOverriddenPropertyOperation::SubObjectsShadowing)
+					if (SerializeContext->bTrackInitializedProperties && Property)
 					{
-						if (SerializeContext->bTrackInitializedProperties && Property)
-						{
-							UE::FInitializedPropertyValueState(this, Data).Set(Property, Tag.ArrayIndex);
-						}
-						if (SerializeContext->bTrackSerializedProperties && Property)
-						{
-							UE::FSerializedPropertyValueState(this, Data).Set(Property, Tag.ArrayIndex);
-						}
+						UE::FInitializedPropertyValueState(this, Data).Set(Property, Tag.ArrayIndex);
+					}
+					if (SerializeContext->bTrackSerializedProperties && Property)
+					{
+						UE::FSerializedPropertyValueState(this, Data).Set(Property, Tag.ArrayIndex);
 					}
 				#endif
 				}
