[GCScopeLock]
- Avoid threads to contend with each other when taking an async GC lock.
- Improves WeakObjectPtr pinning 5x under heavy contention

diff --git a/Engine/Source/Runtime/CoreUObject/Private/UObject/GCScopeLock.h b/Engine/Source/Runtime/CoreUObject/Private/UObject/GCScopeLock.h
--- a/Engine/Source/Runtime/CoreUObject/Private/UObject/GCScopeLock.h
+++ b/Engine/Source/Runtime/CoreUObject/Private/UObject/GCScopeLock.h
@@ -11,6 +11,9 @@
 #include "HAL/PlatformProcess.h"
 #include "Misc/ScopeLock.h"
 #include "HAL/ThreadSafeBool.h"
+#include "Async/SharedMutex.h"
+#include "Async/SharedLock.h"
+#include "Async/UniqueLock.h"
 #include "HAL/Event.h"
 
 
@@ -29,8 +32,8 @@ class FGCCSyncObject
 	    This flag is not automatically enforced on the async threads, instead
 			threads have to manually implement support for it. */
 	TAtomic<int32> GCWantsToRunCounter {0};
-	/** Critical section for thread safe operations */
-	FCriticalSection Critical;
+	/** Shared mutex for thread safe operations */
+	UE::FSharedMutex SharedMutex;
 	/** Event used to block non-game threads when GC is running */
 	FEvent* GCUnlockedEvent;
 
@@ -59,7 +62,10 @@ public:
 					GCUnlockedEvent->Wait();
 				}
 				{
-					FScopeLock CriticalLock(&Critical);
+					// We only need a read lock here since AsyncCounter is already thread-safe
+					// with other threads that want to lock but we need to guard against
+					// the GC trying to take the write lock.
+					UE::TSharedLock SharedLock(SharedMutex);
 					if (GCCounter.GetValue() == 0)
 					{
 						AsyncCounter.Increment();
@@ -79,7 +85,7 @@ public:
 			bLocked = false;
 			if (GCCounter.GetValue() == 0)
 			{
-				FScopeLock CriticalLock(&Critical);
+				UE::TSharedLock SharedLock(SharedMutex);
 				if (GCCounter.GetValue() == 0)
 				{
 					AsyncCounter.Increment();
@@ -114,7 +120,8 @@ public:
 				return AsyncCounter.GetValue() == 0;
 			});
 			{
-				FScopeLock CriticalLock(&Critical);
+				// Guard against any reader locks owned by other threads.
+				UE::TUniqueLock ExclusiveLock(SharedMutex);
 				if (AsyncCounter.GetValue() == 0)
 				{
 					GCUnlockedEvent->Reset();
@@ -142,7 +149,12 @@ public:
 	bool TryGCLock()
 	{
 		bool bSuccess = false;
-		FScopeLock CriticalLock(&Critical);
+		if (AsyncCounter.GetValue() != 0)
+		{
+			return bSuccess;
+		}
+		// Guard against any reader locks owned by other threads.
+		UE::TUniqueLock ExclusiveLock(SharedMutex);
 		// If any other thread is currently locking we just exit
 		if (AsyncCounter.GetValue() == 0)
 		{
