No longer create a deterministic BulkDataId for editor bulkdata when duplicating/instancing a package.

- Using ::CreateUniqueGuid to create a new guid when editor bulkdata was being loaded as part of a duplicate or instancing operation was not safe as it uses the target package's name to modify the source guid to create deterministic output BUT if that package was ever later renamed it would be possible for a new duplicate/instance to use that name which would result in the same 'unique' guid being generated.
-- I am not 100% sure why I originally wrote the duplicate/instancing logic to be deterministic but if there was a good reason at the time it no longer applies. Instead we should just always generated a new guid to stay safe from collisions.
- Moved the check for package duplication and/or instancing to ShouldGenerateNewIdentifier so that all checks that might result in a new guid being generated are done in the same place.
- We now only call ::CreateUniqueGuid when loading from old style bulkdata, fixing up an old version error via ::CreateLegacyUniqueIdentifier or are trying to create a new guid due to a collision.

diff --git a/Engine/Source/Runtime/CoreUObject/Private/Serialization/EditorBulkData.cpp b/Engine/Source/Runtime/CoreUObject/Private/Serialization/EditorBulkData.cpp
--- a/Engine/Source/Runtime/CoreUObject/Private/Serialization/EditorBulkData.cpp
+++ b/Engine/Source/Runtime/CoreUObject/Private/Serialization/EditorBulkData.cpp
@@ -239,7 +239,7 @@ FString GetDebugNameFromOwner(UObject* Owner)
  * be added to themselves. The duplication will load the editor bulkdata to a transient package before
  * it is re-parented.
  */
-bool ShouldGenerateNewIdentifier(FLinkerLoad* LinkerLoad, UObject* Owner)
+bool ShouldGenerateNewIdentifier(FArchive& Ar, UObject* Owner)
 {
 	UPackage* Package = Owner != nullptr ? Owner->GetPackage() : nullptr;
 	if (Package != nullptr && Package->HasAnyFlags(RF_Transient))
@@ -247,6 +247,13 @@ bool ShouldGenerateNewIdentifier(FLinkerLoad* LinkerLoad, UObject* Owner)
 		return true;
 	}
 
+	const FLinkerLoad* LinkerLoad = Cast<FLinkerLoad>(Ar.GetLinker());
+
+	if (Ar.HasAllPortFlags(PPF_Duplicate) || (LinkerLoad && LinkerLoad->GetInstancingContext().IsInstanced()))
+	{
+		return true;
+	}
+
 	return false;
 }
 
@@ -1154,19 +1161,9 @@ void FEditorBulkData::Serialize(FArchive& Ar, UObject* Owner, bool bAllowRegiste
 		}
 		else if (Ar.IsLoading())
 		{
-			if (BulkDataId.IsValid())
+			if (BulkDataId.IsValid() && ShouldGenerateNewIdentifier(Ar, Owner))
 			{
-				FLinkerLoad* LinkerLoad = Cast<FLinkerLoad>(Ar.GetLinker());
-
-				if (ShouldGenerateNewIdentifier(LinkerLoad, Owner))
-				{
-					BulkDataId = FGuid::NewGuid();
-				}
-				else if (Ar.HasAllPortFlags(PPF_Duplicate) || (LinkerLoad && LinkerLoad->GetInstancingContext().IsInstanced()))
-				{
-					// When duplicating BulkDatas we need to create a new BulkDataId to respect the uniqueness contract
-					BulkDataId = CreateUniqueGuid(BulkDataId, Owner, TEXT("PPF_Duplicate serialization"));
-				}
+				BulkDataId = FGuid::NewGuid();
 			}
 
 			OffsetInFile = INDEX_NONE;
