CompactBinary: Delete FCbFieldIterator::MakeOwned() to avoid breaking the iterator

MakeRangeOwned() exists to make the range of the iterator owned.
MakeOwned() only makes the buffer for the current field owned and does not update the iterator end pointer.

diff --git a/Engine/Source/Runtime/Core/Public/Serialization/CompactBinary.h b/Engine/Source/Runtime/Core/Public/Serialization/CompactBinary.h
--- a/Engine/Source/Runtime/Core/Public/Serialization/CompactBinary.h
+++ b/Engine/Source/Runtime/Core/Public/Serialization/CompactBinary.h
@@ -1336,6 +1336,9 @@ public:
 		}
 	}
 
+	/** Call MakeRangeOwned() to make the range of the iterator owned. */
+	void MakeOwned() = delete;
+
 private:
 	using TCbFieldIterator::TCbFieldIterator;
 };
diff --git a/Engine/Source/Runtime/CoreUObject/Private/Cooker/CookDependency.cpp b/Engine/Source/Runtime/CoreUObject/Private/Cooker/CookDependency.cpp
--- a/Engine/Source/Runtime/CoreUObject/Private/Cooker/CookDependency.cpp
+++ b/Engine/Source/Runtime/CoreUObject/Private/Cooker/CookDependency.cpp
@@ -71,7 +71,7 @@ FCookDependency FCookDependency::Function(FName InFunctionName, FCbFieldIterator
 	FCookDependency Result(ECookDependency::Function);
 	Result.FunctionData.Name = InFunctionName;
 	Result.FunctionData.Args = MoveTemp(InArgs);
-	Result.FunctionData.Args.MakeOwned();
+	Result.FunctionData.Args.MakeRangeOwned();
 	return Result;
 }
 
@@ -719,13 +719,13 @@ bool FCookDependency::Load(FCbFieldView Value)
 		{
 			FCbFieldViewIterator FunctionArgumentsIterator = Field.CreateViewIterator();
 			*this = FCookDependency::Function(LocalFuncName, FCbFieldIterator::CloneRange(FunctionArgumentsIterator));
-			FunctionData.Args.MakeOwned();
+			FunctionData.Args.MakeRangeOwned();
 			++Field;
 		}
 		else // If the arguments are not in an array, then this is the old format and all the fields until the end are function parameters
 		{
 			*this = FCookDependency::Function(LocalFuncName, FCbFieldIterator::CloneRange(Field));
-			FunctionData.Args.MakeOwned();
+			FunctionData.Args.MakeRangeOwned();
 			Field.Reset(); // In the old format, all the fields are arguments of the function so there is nothing to read after.
 		}
 	}
