PR #13024: Add TInstancedStruct::NetSerialize function

diff --git a/Engine/Source/Runtime/CoreUObject/Public/StructUtils/InstancedStruct.h b/Engine/Source/Runtime/CoreUObject/Public/StructUtils/InstancedStruct.h
--- a/Engine/Source/Runtime/CoreUObject/Public/StructUtils/InstancedStruct.h
+++ b/Engine/Source/Runtime/CoreUObject/Public/StructUtils/InstancedStruct.h
@@ -153,12 +153,12 @@ public:
 	/** For StructOpsTypeTraits */
 	UE_API bool Serialize(FArchive& Ar, UStruct* DefaultsStruct = nullptr, const void* Defaults = nullptr);
 	UE_API bool Identical(const FInstancedStruct* Other, uint32 PortFlags) const;
-	UE_API void AddStructReferencedObjects(class FReferenceCollector& Collector);
-	UE_API bool ExportTextItem(FString& ValueStr, FInstancedStruct const& DefaultValue, class UObject* Parent, int32 PortFlags, class UObject* ExportRootScope) const;
+	UE_API void AddStructReferencedObjects(FReferenceCollector& Collector);
+	UE_API bool ExportTextItem(FString& ValueStr, FInstancedStruct const& DefaultValue, UObject* Parent, int32 PortFlags, UObject* ExportRootScope) const;
 	UE_API bool ImportTextItem(const TCHAR*& Buffer, int32 PortFlags, UObject* Parent, FOutputDevice* ErrorText, FArchive* InSerializingArchive = nullptr);
 	UE_API bool SerializeFromMismatchedTag(const FPropertyTag& Tag, FStructuredArchive::FSlot Slot);
 	UE_API void GetPreloadDependencies(TArray<UObject*>& OutDeps);
-	UE_API bool NetSerialize(FArchive& Ar, class UPackageMap* Map, bool& bOutSuccess);
+	UE_API bool NetSerialize(FArchive& Ar, UPackageMap* Map, bool& bOutSuccess);
 	UE_API bool FindInnerPropertyInstance(FName PropertyName, const FProperty*& OutProp, const void*& OutData) const;
 	UE_API EPropertyVisitorControlFlow Visit(FPropertyVisitorContext& Context, const TFunctionRef<EPropertyVisitorControlFlow(const FPropertyVisitorContext& /*Context*/)> InFunc) const;
 	UE_API void* ResolveVisitedPathInfo(const FPropertyVisitorInfo& Info) const;
@@ -484,7 +484,7 @@ public:
 		return InstancedStruct != Other.InstancedStruct;
 	}
 
-	void AddReferencedObjects(class FReferenceCollector& Collector)
+	void AddReferencedObjects(FReferenceCollector& Collector)
 	{
 		InstancedStruct.AddStructReferencedObjects(Collector);
 	}
@@ -494,6 +494,11 @@ public:
 		return InstancedStruct.Serialize(Ar);
 	}
 
+	bool NetSerialize(FArchive& Ar, UPackageMap* Map, bool& bOutSuccess)
+	{
+		return InstancedStruct.NetSerialize(Ar, Map, bOutSuccess);
+	}
+
 private:
 	/**
 	 * Note:
