Fixed macros that enable UE_ENABLE_FUOBJECT_ITEM_PACKING

diff --git a/Engine/Source/Runtime/CoreUObject/Public/UObject/UObjectArray.h b/Engine/Source/Runtime/CoreUObject/Public/UObject/UObjectArray.h
--- a/Engine/Source/Runtime/CoreUObject/Public/UObject/UObjectArray.h
+++ b/Engine/Source/Runtime/CoreUObject/Public/UObject/UObjectArray.h
@@ -23,10 +23,10 @@ namespace Verse
 
 #ifndef UE_PACK_FUOBJECT_ITEM
 #	define UE_ENABLE_FUOBJECT_ITEM_PACKING 0
-#elif defined(UE_PACK_FUOBJECT_ITEM) && defined(PLATFORM_CPU_ARM_FAMILY) && defined(FORCE_ANSI_ALLOCATOR) && !UE_WITH_REMOTE_OBJECT_HANDLE
+#elif PLATFORM_CPU_ARM_FAMILY && FORCE_ANSI_ALLOCATOR
 //	disable packing on ARM64 is we use Ansi allocator as it might use MTE or HWASan which uses top byte of a pointer and this feature will discard that
 #	define UE_ENABLE_FUOBJECT_ITEM_PACKING 0
-#elif
+#else
 #	define UE_ENABLE_FUOBJECT_ITEM_PACKING UE_PACK_FUOBJECT_ITEM && !UE_WITH_REMOTE_OBJECT_HANDLE
 #endif
 
