[UE] Placeholder type serialization support is now enabled/disabled as an IDO subfeature.

Notes:
- Placeholder type serialization support now requires IDO.Enable to be enabled as a dependency.
- Consolidated code now uses the existing CVar as an IDO subfeature toggle ("IDO.Placeholder.Enable") rather than deferring to 'bForceDisabled' as before.

#rnx

diff --git a/Engine/Source/Runtime/CoreUObject/Private/UObject/PropertyBagRepository.cpp b/Engine/Source/Runtime/CoreUObject/Private/UObject/PropertyBagRepository.cpp
--- a/Engine/Source/Runtime/CoreUObject/Private/UObject/PropertyBagRepository.cpp
+++ b/Engine/Source/Runtime/CoreUObject/Private/UObject/PropertyBagRepository.cpp
@@ -933,7 +933,7 @@ bool FPropertyBagRepository::IsPropertyBagPlaceholderObject(const UObject* Objec
 namespace Private
 {
 #if WITH_EDITOR
-	static bool bEnablePropertyBagPlaceholderObjectSupport = 0;
+	static bool bEnablePropertyBagPlaceholderObjectSupport = true;
 	static FAutoConsoleVariableRef CVarEnablePropertyBagPlaceholderObjectSupport(
 		TEXT("IDO.Placeholder.Enable"),
 		bEnablePropertyBagPlaceholderObjectSupport,
@@ -980,20 +980,7 @@ namespace Private
 bool FPropertyBagRepository::IsPropertyBagPlaceholderObjectSupportEnabled()
 {
 #if WITH_EDITOR && UE_WITH_OBJECT_HANDLE_TYPE_SAFETY
-	static bool bIsInitialized = false;
-	static bool bForceDisabled = false;
-	if (!bIsInitialized)
-	{
-		Private::bEnablePropertyBagPlaceholderObjectSupport |= FParse::Param(FCommandLine::Get(), TEXT("WithPropertyBagPlaceholderObjects"));
-		Private::CVarEnablePropertyBagPlaceholderObjectSupport->OnChangedDelegate().AddLambda([](IConsoleVariable* CVar)
-		{
-			bForceDisabled = !CVar->GetBool();
-		});
-
-		bIsInitialized = true;
-	}
-	
-	return Private::bEnablePropertyBagPlaceholderObjectSupport || (IsInstanceDataObjectSupportEnabled() && !bForceDisabled);
+	return Private::bEnablePropertyBagPlaceholderObjectSupport && IsInstanceDataObjectSupportEnabled();
 #else
 	return false;
 #endif
