Added new helper to help resolve determinism issues arising from calls to MakeUniqueObjectName during cook/save.
ReleaseUniquelyNamedObject(UObject* Object) : Renames a UObject so that it's name can be resued immediately, also tries to lower the per class suffix annotation so the suffix can also be resued

Changed UHLODEngineSubsystem::CleanupHLOD to use ReleaseUniquelyNamedObject  after removing it's actors so we can recreate identical actors with the same name afterwards.

diff --git a/Engine/Source/Runtime/CoreUObject/Private/UObject/UObjectGlobals.cpp b/Engine/Source/Runtime/CoreUObject/Private/UObject/UObjectGlobals.cpp
--- a/Engine/Source/Runtime/CoreUObject/Private/UObject/UObjectGlobals.cpp
+++ b/Engine/Source/Runtime/CoreUObject/Private/UObject/UObjectGlobals.cpp
@@ -2825,6 +2825,48 @@ namespace NameReuse
 }
 #endif
 
+#if WITH_EDITOR
+void RecycleUniqueObjectSuffixAnnotation(UObject* Parent, UClass* ForClass)
+{
+	int32 MaxIndex = 0;
+	FName BaseName = ForClass->GetFName();
+	BaseName.SetNumber(0);
+
+	ForEachObjectWithOuter(Parent, [&MaxIndex, BaseName](UObject* Object)
+		{
+			FName ObjectName = Object->GetFName();
+			int32 Index = ObjectName.GetNumber();
+			ObjectName.SetNumber(0);
+
+			if (BaseName == ObjectName)
+			{
+				if (Index > MaxIndex)
+				{
+					MaxIndex = Index;
+				}
+			}
+		}, false);
+
+	UpdateSuffixForNextNewObject(Parent, ForClass, [MaxIndex](int32& Index)
+		{
+			Index = MaxIndex;
+		});
+}
+
+void ReleaseUniquelyNamedObject(UObject* Object)
+{
+	UObject* Outer = Object->GetOuter();
+	UClass* Class = Object->GetClass();
+
+	// Make the name available again for other objects
+	Object->Rename(nullptr, GetTransientPackage());
+
+	// Reset the suffix in FPerClassNumberSuffixAnnotation to the lowest we can 
+	RecycleUniqueObjectSuffixAnnotation(Outer, Class);
+}
+#endif
+
+
 FName MakeUniqueObjectName(UObject* Parent, const UClass* Class, FName InBaseName/*=NAME_None*/, EUniqueObjectNameOptions Options /*= EUniqueObjectNameOptions::None*/)
 {
 	TRACE_CPUPROFILER_EVENT_SCOPE(MakeUniqueObjectName);
diff --git a/Engine/Source/Runtime/CoreUObject/Public/UObject/UObjectGlobals.h b/Engine/Source/Runtime/CoreUObject/Public/UObject/UObjectGlobals.h
--- a/Engine/Source/Runtime/CoreUObject/Public/UObject/UObjectGlobals.h
+++ b/Engine/Source/Runtime/CoreUObject/Public/UObject/UObjectGlobals.h
@@ -987,6 +987,22 @@ enum class EUniqueObjectNameOptions
 };
 ENUM_CLASS_FLAGS(EUniqueObjectNameOptions);
 
+
+#if WITH_EDITOR
+/**
+ * Renames a uniquely named objects obtained by MakeUniqueObjectName and 
+ * attempts to lower the arbitrary per-class suffix used to make names unique if possible.
+ * 
+ * This is intended to allow code to recreate sequences of uniquely named objects with identical  
+ * names.
+ *
+ * @param	Object		The objects to rename 
+ */
+
+
+COREUOBJECT_API void ReleaseUniquelyNamedObject(UObject* Object);
+#endif
+
 /**
  * Create a unique name by combining a base name and an arbitrary number string.
  * The object name returned is guaranteed not to exist.
diff --git a/Engine/Source/Runtime/Engine/Private/HLOD/HLODEngineSubsystem.cpp b/Engine/Source/Runtime/Engine/Private/HLOD/HLODEngineSubsystem.cpp
--- a/Engine/Source/Runtime/Engine/Private/HLOD/HLODEngineSubsystem.cpp
+++ b/Engine/Source/Runtime/Engine/Private/HLOD/HLODEngineSubsystem.cpp
@@ -161,6 +161,7 @@ bool UHLODEngineSubsystem::CleanupHLOD(ALODActor* InLODActor)
 	if (bShouldDestroyActor)
 	{
 		InLODActor->GetWorld()->EditorDestroyActor(InLODActor, true);
+		ReleaseUniquelyNamedObject(InLODActor);
 	}
 
 	return bShouldDestroyActor;
