[UObject/GC]
- Avoid contention on GRoots when setting/removing root flags on objects
- Merge Flags and RefCount together to preserve atomicity between root flags
- This improves WeakObjectPtr pinning 5x under heavy contention when not already having another root flag
- This also improves toggling between no root flags to root flags by the same amount
- Cleaned up deprecations from 5.5

diff --git a/Engine/Source/Runtime/CoreUObject/Private/UObject/GarbageCollection.cpp b/Engine/Source/Runtime/CoreUObject/Private/UObject/GarbageCollection.cpp
--- a/Engine/Source/Runtime/CoreUObject/Private/UObject/GarbageCollection.cpp
+++ b/Engine/Source/Runtime/CoreUObject/Private/UObject/GarbageCollection.cpp
@@ -6,6 +6,8 @@
 
 #include "UObject/GarbageCollection.h"
 #include "Algo/Find.h"
+#include "Async/UniqueLock.h"
+#include "Containers/ConsumeAllMpmcQueue.h"
 #include "Containers/Deque.h"
 #include "Containers/StaticBitArray.h"
 #include "HAL/MallocLeakDetection.h"
@@ -621,8 +623,9 @@ namespace UE::GC::Private
 	using FGatherUnreachableObjectsState = TThreadedGather<TArray<UE::GC::FUnreachableObject>>;
 	static FGatherUnreachableObjectsState GGatherUnreachableObjectsState;
 
+	static TConsumeAllMpmcQueue<int32> GDirtyRoots;
 	static TSet<int32> GRoots;
-	static FCriticalSection GRootsCritical;
+	static UE::FMutex  GRootsMutex;
 
 	/** Current EInternalObjectFlags value representing a reachable object */
 	EInternalObjectFlags FGCFlags::ReachableObjectFlag = EInternalObjectFlags::ReachabilityFlag0;
@@ -634,10 +637,62 @@ namespace UE::GC::Private
 	{
 		GSkipDestroyGarbage.store(bSkip, std::memory_order_relaxed);
 	}
+	
+	void ProcessDirtyRootNoLock(int32 Index)
+	{
+		// If this somehow race with a SetRootFlags or ClearRootFlags its not the end of the world
+		// since another MarkRootAsDirty will be queued for later processing.
+		const FUObjectItem* ObjectItem = GUObjectArray.IndexToObjectUnsafeForGC(Index);
+
+		if (ObjectItem->HasAnyFlags(EInternalObjectFlags_RootFlags) || ObjectItem->GetRefCount() != 0)
+		{
+			GRoots.Add(Index);
+		}
+		else
+		{
+			GRoots.Remove(Index);
+		}
+	}
+
+	void ProcessDirtyRootsNoLock()
+	{
+		// The order is not important since the calls are stateless
+		GDirtyRoots.ConsumeAllLifo(
+			[](int32 Index)
+			{
+				ProcessDirtyRootNoLock(Index);
+			}
+		);
+	}
+
+	void MarkRootAsDirty(int32 Index)
+	{
+		if (GRootsMutex.TryLock())
+		{
+			// Process the whole list first to cooperatively help other threads when
+			// we get the lock.
+			ProcessDirtyRootsNoLock();
+
+			// No need to queue anything since we have the lock, process it right away.
+			ProcessDirtyRootNoLock(Index);
+			GRootsMutex.Unlock();
+		}
+		else
+		{
+			// This will be eventually consistent as each thread able to take the lock
+			// will update the value before using the GRoots content.
+			GDirtyRoots.ProduceItem(Index);
+		}
+	}
 }
 
 static bool GatherUnreachableObjects(UE::GC::EGatherOptions Options, double TimeLimit = 0.0);
 
+void FUObjectItem::MarkRootAsDirty()
+{
+	UE::GC::Private::MarkRootAsDirty(GUObjectArray.ObjectToIndex(GetObject()));
+}
+
 bool FUObjectItem::SetRootFlags(EInternalObjectFlags FlagsToSet)
 {
 	using namespace UE::GC;
@@ -646,16 +701,23 @@ bool FUObjectItem::SetRootFlags(EInternalObjectFlags FlagsToSet)
 	constexpr int32 RootFlags = (int32)EInternalObjectFlags_RootFlags;
 	bool bIChangedIt = false;
 	{
-		FScopeLock RootsLock(&GRootsCritical);
-		const int32 OldFlags = GetFlagsInternal();
-		if ((OldFlags & RootFlags) == 0)
+		int32 OldFlags;
+		int32 OldRefCount;
+		bIChangedIt = AtomicallySetFlag_ForGC(FlagsToSet, OldFlags, OldRefCount);
+		if (bIChangedIt)
 		{
-			if (!(GUObjectArray.IsOpenForDisregardForGC() & GUObjectArray.DisregardForGCEnabled())) //-V792
+			// Only dirty the root if there was no root flags before and no refcount
+			if ((OldFlags & RootFlags) == 0 && OldRefCount == 0)
 			{
-				GRoots.Add(GUObjectArray.ObjectToIndex(GetObject()));
+				// RootFlags can only be added while GC is not running, otherwise GC might not see the change
+				checkSlow(IsInGameThread() || IsInParallelGameThread() || IsGarbageCollectionLocked());
+
+				if (!(GUObjectArray.IsOpenForDisregardForGC() & GUObjectArray.DisregardForGCEnabled())) //-V792
+				{
+					MarkRootAsDirty();
+				}
 			}
 		}
-		bIChangedIt = AtomicallySetFlag_ForGC(FlagsToSet);
 	}
 	if (bIChangedIt & GIsIncrementalReachabilityPending) //-V792
 	{
@@ -672,49 +734,23 @@ bool FUObjectItem::ClearRootFlags(EInternalObjectFlags FlagsToClear)
 	using namespace UE::GC::Private;
 
 	constexpr int32 RootFlags = (int32)EInternalObjectFlags_RootFlags;
-	FScopeLock RootsLock(&GRootsCritical);
-	const int32 OldFlags = GetFlagsInternal();
 
-	// Now that we have the lock, we prevent the removal if refcount is not 0
-	// This avoids a race condition between ReleaseRef and AddRef.
-	if ((FlagsToClear & EInternalObjectFlags::RefCounted) == EInternalObjectFlags::RefCounted && GetRefCountInternal())
+	int32 OldFlags;
+	int32 OldRefCount;
+	bool bIChangedIt = AtomicallyClearFlag_ForGC(FlagsToClear, OldFlags, OldRefCount);
+	if (bIChangedIt)
 	{
-		FlagsToClear &= ~EInternalObjectFlags::RefCounted;
-	}
-
-	if ((OldFlags & RootFlags) != 0 && ((OldFlags & ~(int32)FlagsToClear) & RootFlags) == 0)
-	{
-		GRoots.Remove(GUObjectArray.ObjectToIndex(GetObject()));
+		if ((OldFlags & RootFlags) != 0 && ((OldFlags & ~(int32)FlagsToClear) & RootFlags) == 0 && OldRefCount == 0)
+		{
+			if (!(GUObjectArray.IsOpenForDisregardForGC() & GUObjectArray.DisregardForGCEnabled())) //-V792
+			{
+				MarkRootAsDirty();
+			}
+		}
 	}
-	return AtomicallyClearFlag_ForGC((EInternalObjectFlags)FlagsToClear);
-}
-
-// BEGIN Deprecated FUObjectItem functions. These functions will be removed in the next release and are not used internally by the engine
-bool FUObjectItem::IsMaybeUnreachable() const
-{
-	return UE::GC::Private::FGCFlags::IsMaybeUnreachable_ForGC(this);
-}
 
-void FUObjectItem::SetMaybeUnreachable()
-{
-	UE::GC::Private::FGCFlags::SetMaybeUnreachable_ForGC(this);
-}
-
-void FUObjectItem::FastMarkAsReachableAndClearReachaleInClusterInterlocked_ForGC()
-{
-	UE::GC::Private::FGCFlags::FastMarkAsReachableAndClearReachableInClusterInterlocked_ForGC(this);
-}
-
-void FUObjectItem::FastMarkAsReachableInterlocked_ForGC()
-{
-	UE::GC::Private::FGCFlags::FastMarkAsReachableInterlocked_ForGC(this);
-}
-
-bool FUObjectItem::MarkAsReachableInterlocked_ForGC()
-{
-	return UE::GC::Private::FGCFlags::MarkAsReachableInterlocked_ForGC(this);
+	return bIChangedIt;
 }
-// END Deprecated FUObjectItem functions
 
 // BEGIN Hidden (but DLL exported) functions that can only be used for testing GC and GC flags
 namespace UE::GC::Test
@@ -768,7 +804,8 @@ void OnDisregardForGCSetDisabled(int32 NumObjects)
 		}
 	}, (DisregardSetDisabledState.NumWorkerThreads() == 1) ? EParallelForFlags::ForceSingleThread : EParallelForFlags::None);
 
-	FScopeLock RootsLock(&GRootsCritical);
+	UE::TUniqueLock RootsLock(GRootsMutex);
+	ProcessDirtyRootsNoLock();
 	DisregardSetDisabledState.Finish(GRoots);
 }
 
@@ -4207,9 +4244,10 @@ public:
 		FMarkRootsState MarkRootsState;		
 
 		{
-			GRootsCritical.Lock();
+			GRootsMutex.Lock();
+			ProcessDirtyRootsNoLock();
 			TArray<int32> RootsArray(GRoots.Array());				
-			GRootsCritical.Unlock();
+			GRootsMutex.Unlock();
 			MarkRootsState.Start(Options, RootsArray.Num());
 			FMarkRootsState::FThreadIterators& ThreadIterators = MarkRootsState.GetThreadIterators();
 
@@ -6284,7 +6322,8 @@ void PurgeAllUObjectsOnExit()
 
 	{
 		// Since we just destroyed all objects regardless of their flags, make sure the roots list is empty
-		FScopeLock RootsLock(&GRootsCritical);
+		UE::TUniqueLock RootsLock(GRootsMutex);
+		ProcessDirtyRootsNoLock();
 		GRoots.Empty();
 	}
 
diff --git a/Engine/Source/Runtime/CoreUObject/Private/UObject/GarbageCollectionInternalFlags.h b/Engine/Source/Runtime/CoreUObject/Private/UObject/GarbageCollectionInternalFlags.h
--- a/Engine/Source/Runtime/CoreUObject/Private/UObject/GarbageCollectionInternalFlags.h
+++ b/Engine/Source/Runtime/CoreUObject/Private/UObject/GarbageCollectionInternalFlags.h
@@ -86,31 +86,31 @@ public:
 
 	FORCEINLINE static void FastMarkAsReachableInterlocked_ForGC(FUObjectItem* ObjectItem)
 	{
-		FPlatformAtomics::InterlockedAnd(&ObjectItem->Flags, ~int32(MaybeUnreachableObjectFlag));
-		FPlatformAtomics::InterlockedOr(&ObjectItem->Flags, int32(ReachableObjectFlag));
+		FPlatformAtomics::InterlockedAnd(&ObjectItem->FlagsAndRefCount, ~(int64(MaybeUnreachableObjectFlag) << 32));
+		FPlatformAtomics::InterlockedOr(&ObjectItem->FlagsAndRefCount, int64(ReachableObjectFlag) << 32);
 	}
 
 	FORCEINLINE static void FastMarkAsReachableAndClearReachableInClusterInterlocked_ForGC(FUObjectItem* ObjectItem)
 	{
-		FPlatformAtomics::InterlockedAnd(&ObjectItem->Flags, ~int32(MaybeUnreachableObjectFlag | EInternalObjectFlags::ReachableInCluster));
-		FPlatformAtomics::InterlockedOr(&ObjectItem->Flags, int32(ReachableObjectFlag));
+		FPlatformAtomics::InterlockedAnd(&ObjectItem->FlagsAndRefCount, ~(int64(MaybeUnreachableObjectFlag | EInternalObjectFlags::ReachableInCluster) << 32));
+		FPlatformAtomics::InterlockedOr(&ObjectItem->FlagsAndRefCount, int64(ReachableObjectFlag) << 32);
 	}
 
 	FORCEINLINE static bool MarkAsReachableInterlocked_ForGC(FUObjectItem* ObjectItem)
 	{
 		const int32 FlagToClear = int32(MaybeUnreachableObjectFlag);
-		if (FPlatformAtomics::AtomicRead_Relaxed(&ObjectItem->Flags) & FlagToClear)
+		if (ObjectItem->GetFlagsInternal() & FlagToClear)
 		{
-			int32 Old = FPlatformAtomics::InterlockedAnd(&ObjectItem->Flags, ~FlagToClear);
-			FPlatformAtomics::InterlockedOr(&ObjectItem->Flags, int32(ReachableObjectFlag));
-			return Old & FlagToClear;
+			int64 Old = FPlatformAtomics::InterlockedAnd(&ObjectItem->FlagsAndRefCount, ~((int64)FlagToClear << 32));
+			FPlatformAtomics::InterlockedOr(&ObjectItem->FlagsAndRefCount, int64(ReachableObjectFlag) << 32);
+			return FUObjectItem::GetFlagsInternal(Old) & FlagToClear;
 		}
 		return false;
 	}
 
 	FORCEINLINE static constexpr ::size_t OffsetOfFlags_ForGC()
 	{
-		return offsetof(FUObjectItem, Flags);
+		return offsetof(FUObjectItem, FlagsAndRefCount);
 	}
 
 	FORCEINLINE static void SwapReachableAndMaybeUnreachable()
diff --git a/Engine/Source/Runtime/CoreUObject/Private/UObject/ReferenceChainSearch.cpp b/Engine/Source/Runtime/CoreUObject/Private/UObject/ReferenceChainSearch.cpp
--- a/Engine/Source/Runtime/CoreUObject/Private/UObject/ReferenceChainSearch.cpp
+++ b/Engine/Source/Runtime/CoreUObject/Private/UObject/ReferenceChainSearch.cpp
@@ -80,6 +80,7 @@ namespace UE::ReferenceChainSearch
 		{
 			const UObject* Object = VertexToObject(Vertex);
 			return Object->HasAnyInternalFlags(EInternalObjectFlags_RootFlags)
+				|| Object->GetRefCount() != 0
 				|| (GARBAGE_COLLECTION_KEEPFLAGS != RF_NoFlags && Object->HasAnyFlags(GARBAGE_COLLECTION_KEEPFLAGS)
 					&& !(SearchMode & EReferenceChainSearchMode::FullChain));
 		}
@@ -190,6 +191,7 @@ namespace UE::ReferenceChainSearch
 		{
 			ObjectType Object = VertexToObject(Vertex);
 			return Object->HasAnyInternalFlags(EInternalObjectFlags_RootFlags)
+				|| Object->GetRefCount() != 0
 				|| (GARBAGE_COLLECTION_KEEPFLAGS != RF_NoFlags && Object->HasAnyFlags(GARBAGE_COLLECTION_KEEPFLAGS)
 					&& !(SearchMode & EReferenceChainSearchMode::FullChain));
 		}
@@ -1293,7 +1295,7 @@ FString FReferenceChainSearch::GetObjectFlags(const FGCObjectInfo& InObject)
 		Flags += TEXT("(Unreachable) ");
 	}
 
-	if (InObject.HasAnyInternalFlags(EInternalObjectFlags::RefCounted))
+	if (InObject.GetRefCount())
 	{
 		Flags += FString::Printf(TEXT("(refcounted<%d>) "), InObject.GetRefCount());
 	}
@@ -1739,6 +1741,13 @@ int32 FReferenceChainSearch::PrintResults(TFunctionRef<bool(FCallbackParams& Par
 #if !NO_LOGGING
 				GWarn->CategorizedLogf(LogReferenceChain.GetCategoryName(), InVerbosityForPrint,
 					TEXT("%s%s is not currently reachable but it does have some of GARBAGE_COLLECTION_KEEPFLAGS set."), *GetObjectFlags(ObjInfo), *ObjInfo.GetFullName());
+#endif
+			}
+			else if (ObjInfo.GetRefCount() != 0)
+			{
+#if !NO_LOGGING
+				GWarn->CategorizedLogf(LogReferenceChain.GetCategoryName(), InVerbosityForPrint,
+					TEXT("%s%s is not currently reachable but it does have %d refcount set."), *GetObjectFlags(ObjInfo), *ObjInfo.GetFullName(), ObjInfo.GetRefCount());
 #endif
 			}
 			else
@@ -1834,7 +1843,7 @@ void FReferenceChainSearch::Cleanup()
 
 static FORCEINLINE bool HasGarbageCollectionKeepFlags(FGCObjectInfo* ObjectInfo)
 {
-	return ObjectInfo && (ObjectInfo->HasAnyFlags(GARBAGE_COLLECTION_KEEPFLAGS) || ObjectInfo->HasAnyInternalFlags(EInternalObjectFlags_RootFlags));
+	return ObjectInfo && (ObjectInfo->HasAnyFlags(GARBAGE_COLLECTION_KEEPFLAGS) || ObjectInfo->HasAnyInternalFlags(EInternalObjectFlags_RootFlags) || ObjectInfo->GetRefCount() != 0);
 }
 
 static bool PrintStaleReferenceChainsAndFindReferencingObjects(UObject* ObjectToFindReferencesTo, FReferenceChainSearch& RefChainSearch, FGCObjectInfo*& OutGarbageObject, FGCObjectInfo*& OutReferencingObject, ELogVerbosity::Type Verbosity)
diff --git a/Engine/Source/Runtime/CoreUObject/Private/UObject/UObjectArray.cpp b/Engine/Source/Runtime/CoreUObject/Private/UObject/UObjectArray.cpp
--- a/Engine/Source/Runtime/CoreUObject/Private/UObject/UObjectArray.cpp
+++ b/Engine/Source/Runtime/CoreUObject/Private/UObject/UObjectArray.cpp
@@ -290,15 +290,14 @@ void FUObjectArray::AllocateUObjectIndex(UObjectBase* Object, EInternalObjectFla
 	FUObjectItem* ObjectItem = IndexToObject(Index);
 	UE_CLOG(ObjectItem->GetObject() != nullptr, LogUObjectArray, Fatal, TEXT("Attempting to add %s at index %d but another object (0x%016llx) exists at that index!"), *Object->GetFName().ToString(), Index, (int64)(PTRINT)ObjectItem->GetObject());
 	// At this point all not-compiled-in objects are not fully constructed yet and this is the earliest we can mark them as such
-	ObjectItem->Flags = (int32)EInternalObjectFlags::PendingConstruction;
+	ObjectItem->FlagsAndRefCount = ((int64)EInternalObjectFlags::PendingConstruction) << 32;
 	if (!(IsOpenForDisregardForGC() & GUObjectArray.DisregardForGCEnabled())) //-V792
 	{
 		// It's safe to access FGCFlags::GetReachableFlagValue_ForGC() here because creating new objects is being performed 
 		// under the same UObjectArray lock as swapping reachability flags inside of GC, see FGCFlags::SwapReachableAndMaybeUnreachable()
-		ObjectItem->Flags |= (int32)UE::GC::Private::FGCFlags::GetReachableFlagValue_ForGC();
+		ObjectItem->FlagsAndRefCount |= ((int64)UE::GC::Private::FGCFlags::GetReachableFlagValue_ForGC()) << 32;
 	}
 	ObjectItem->SetObject(Object);
-	ObjectItem->RefCount = 0;
 	ObjectItem->ClusterRootIndex = 0;
 	ObjectItem->SerialNumber = SerialNumber;
 #if UE_WITH_REMOTE_OBJECT_HANDLE
@@ -381,10 +380,10 @@ void FUObjectArray::FreeUObjectIndex(UObjectBase* Object)
 	UE_CLOG(ObjectItem->GetObject() != Object, LogUObjectArray, Fatal, TEXT("Removing object (0x%016llx) at index %d but the index points to a different object (0x%016llx)!"), (int64)(PTRINT)Object, Index, (int64)(PTRINT)ObjectItem->GetObject());
 
 	// Can't destroy a refcounted object
-	check((ObjectItem->RefCount == 0 && (ObjectItem->GetFlags() & EInternalObjectFlags::RefCounted) == EInternalObjectFlags::None) || GExitPurge);
+	check(ObjectItem->GetRefCount() == 0 || GExitPurge);
 
 	// Clear root flags to remove this object's index from UE::GC::Private::GRoots array 
-	if ((ObjectItem->Flags & (int32)EInternalObjectFlags_RootFlags) != 0)
+	if ((ObjectItem->GetFlagsInternal() & (int32)EInternalObjectFlags_RootFlags) != 0)
 	{
 		ObjectItem->ThisThreadAtomicallyClearedFlag(EInternalObjectFlags_RootFlags);
 	}
@@ -394,8 +393,7 @@ void FUObjectArray::FreeUObjectIndex(UObjectBase* Object)
 	// so, reset these fields in the open.
 	UE_AUTORTFM_OPEN
 	{
-		ObjectItem->Flags = 0;
-		ObjectItem->RefCount = 0;
+		ObjectItem->FlagsAndRefCount = 0;
 	};
 
 	ObjectItem->SetObject(nullptr);
diff --git a/Engine/Source/Runtime/CoreUObject/Public/UObject/ObjectMacros.h b/Engine/Source/Runtime/CoreUObject/Public/UObject/ObjectMacros.h
--- a/Engine/Source/Runtime/CoreUObject/Public/UObject/ObjectMacros.h
+++ b/Engine/Source/Runtime/CoreUObject/Public/UObject/ObjectMacros.h
@@ -650,7 +650,7 @@ enum class EInternalObjectFlags : int32
 	Async = 1 << 26, ///< Object exists only on a different thread than the game thread.
 	AsyncLoadingPhase2 = 1 << 27, ///< Object is being asynchronously loaded.
 	Unreachable = 1 << 28, ///< Object is not reachable on the object graph.
-	RefCounted = 1 << 29, ///< Object currently has ref-counts associated with it.
+	RefCounted UE_DEPRECATED(5.7, "Use GetRefCount() to determine if a refcount exists instead.") = 1 << 29, ///< Object currently has ref-counts associated with it.
 	RootSet = 1 << 30, ///< Object will not be garbage collected, even if unreferenced.
 	PendingConstruction = 1 << 31, ///< Object didn't have its class constructor called yet (only the UObjectBase one to initialize its most basic members)
 
@@ -662,8 +662,9 @@ ENUM_CLASS_FLAGS(EInternalObjectFlags);
 
 //~ Make sure these macros are up to date!
 #define EInternalObjectFlags_MinFlagBitIndex 14
-#define EInternalObjectFlags_GarbageCollectionKeepFlags (EInternalObjectFlags::Borrowed | EInternalObjectFlags::Native | EInternalObjectFlags::Async | EInternalObjectFlags::AsyncLoadingPhase1 | EInternalObjectFlags::AsyncLoadingPhase2 | EInternalObjectFlags::LoaderImport | EInternalObjectFlags::RefCounted)
-#define EInternalObjectFlags_AllFlags (EInternalObjectFlags::Remote | EInternalObjectFlags::RemoteReference | EInternalObjectFlags::Borrowed | EInternalObjectFlags::ReachabilityFlag0 | EInternalObjectFlags::ReachabilityFlag1 | EInternalObjectFlags::ReachabilityFlag2 | EInternalObjectFlags::AutoRTFMConstructionAborted | EInternalObjectFlags::LoaderImport | EInternalObjectFlags::Garbage | EInternalObjectFlags::ReachableInCluster | EInternalObjectFlags::ClusterRoot | EInternalObjectFlags::Native | EInternalObjectFlags::RefCounted | EInternalObjectFlags::Async | EInternalObjectFlags::AsyncLoadingPhase1 | EInternalObjectFlags::AsyncLoadingPhase2 | EInternalObjectFlags::RootSet | EInternalObjectFlags::PendingConstruction | EInternalObjectFlags::Unreachable)
+#define EInternalObjectFlags_RefCountMask 0xFFFFFFFF
+#define EInternalObjectFlags_GarbageCollectionKeepFlags (EInternalObjectFlags::Borrowed | EInternalObjectFlags::Native | EInternalObjectFlags::Async | EInternalObjectFlags::AsyncLoadingPhase1 | EInternalObjectFlags::AsyncLoadingPhase2 | EInternalObjectFlags::LoaderImport)
+#define EInternalObjectFlags_AllFlags (EInternalObjectFlags::Remote | EInternalObjectFlags::RemoteReference | EInternalObjectFlags::Borrowed | EInternalObjectFlags::ReachabilityFlag0 | EInternalObjectFlags::ReachabilityFlag1 | EInternalObjectFlags::ReachabilityFlag2 | EInternalObjectFlags::AutoRTFMConstructionAborted | EInternalObjectFlags::LoaderImport | EInternalObjectFlags::Garbage | EInternalObjectFlags::ReachableInCluster | EInternalObjectFlags::ClusterRoot | EInternalObjectFlags::Native | EInternalObjectFlags::Async | EInternalObjectFlags::AsyncLoadingPhase1 | EInternalObjectFlags::AsyncLoadingPhase2 | EInternalObjectFlags::RootSet | EInternalObjectFlags::PendingConstruction | EInternalObjectFlags::Unreachable)
 #define EInternalObjectFlags_RootFlags (EInternalObjectFlags::RootSet | EInternalObjectFlags_GarbageCollectionKeepFlags)
 #define EInternalObjectFlags_ReachabilityFlags (EInternalObjectFlags::ReachabilityFlag0 | EInternalObjectFlags::ReachabilityFlag1 | EInternalObjectFlags::ReachabilityFlag2 | EInternalObjectFlags::Unreachable)
 #define EInternalObjectFlags_AsyncLoading (EInternalObjectFlags::AsyncLoadingPhase1 | EInternalObjectFlags::AsyncLoadingPhase2)
diff --git a/Engine/Source/Runtime/CoreUObject/Public/UObject/UObjectArray.h b/Engine/Source/Runtime/CoreUObject/Public/UObject/UObjectArray.h
--- a/Engine/Source/Runtime/CoreUObject/Public/UObject/UObjectArray.h
+++ b/Engine/Source/Runtime/CoreUObject/Public/UObject/UObjectArray.h
@@ -51,13 +51,19 @@ struct FUObjectItem
 	friend class FUObjectArray;
 	friend class UE::GC::Private::FGCFlags;
 
+private:
+	// Stores EInternalObjectFlags (and higher 13 bits of the UObject pointer packed together if UE_PACK_FUOBJECT_ITEM is set to 1)
+	// These can only be changed via Set* and Clear* functions
+	// The Flags are now stored on high 32-bit so that we can use InterlockedInc/InterlockedDec directly for RefCount which is stored on the low 32-bit
+	// while preserving atomicity of the whole thing so RootFlags+RefCount can be evaluated in a lock-less way.
+	// If we want to add more Flags, we can reduce the size of RefCount to 24 bits to give us more bits for Flags but will require that we convert
+	// EInternalObjectFlags to a 64-bit.
+	int64 FlagsAndRefCount;
+
 #if !UE_ENABLE_FUOBJECT_ITEM_PACKING
 	// Pointer to the allocated object
-	UE_DEPRECATED(5.5, "Use GetObject() and SetObject() to access Object.")
 	class UObjectBase* Object;
-
 #else
-private:
 	// Currently we assume UObjects are aligned by 8 bytes, that gives us 3 lower bits as zeros that we can discard.
 	// This will give us total 45 bits in a pointer that we pack into a int32 and the remaining 13 bits we pack with Flags
 	// EInternalObjectFlags_MinFlagBitIndex at the time of writing this is 14 and we have only 1 bit left in the EInternalObjectFlags for future use
@@ -68,22 +74,15 @@ private:
 	constexpr static int32 FlagsMask = 0xFFFFFFFF << int(EInternalObjectFlags_MinFlagBitIndex);
 	constexpr static int32 PtrMask = ~FlagsMask;
 
-	// Stores lower 32 bits of UObject pointer shifted by 3 to the left as all our allocations are at leas 8 bytes aligned and lower 3 bits will always be 0
+	// Stores lower 32 bits of UObject pointer shifted by 3 to the left as all our allocations are at least 8 bytes aligned and lower 3 bits will always be 0
 	uint32 ObjectPtrLow;
 #endif
 
-private:
-	// Stores EInternalObjectFlags (and higher 13 bits of the UObject pointer packed together if UE_PACK_FUOBJECT_ITEM is set to 1)
-	// These can only be changed via Set* and Clear* functions
-	int32 Flags;
-
 public:
 	// UObject Owner Cluster Index
-	int32 ClusterRootIndex;	
+	int32 ClusterRootIndex;
 	// Weak Object Pointer Serial number associated with the object
 	int32 SerialNumber;
-	// RefCount associated with the object preventing its destruction.
-	int32 RefCount;
 
 #if UE_WITH_REMOTE_OBJECT_HANDLE
 	// Globally unique id of this object
@@ -100,15 +99,14 @@ public:
 #endif // STATS || ENABLE_STATNAMEDEVENTS
 
 	FUObjectItem()
+		: FlagsAndRefCount(0)
 #if !UE_ENABLE_FUOBJECT_ITEM_PACKING
-		: Object(nullptr)
+		, Object(nullptr)
 #else
-		: ObjectPtrLow(0)
+		, ObjectPtrLow(0)
 #endif
-		, Flags(0)
 		, ClusterRootIndex(0)
 		, SerialNumber(0)
-		, RefCount(0)
 #if ENABLE_STATNAMEDEVENTS_UOBJECT
 		, StatIDStringStorage(nullptr)
 #endif
@@ -134,12 +132,10 @@ public:
 	FORCEINLINE class UObjectBase* GetObject() const
 	{
 #if UE_ENABLE_FUOBJECT_ITEM_PACKING
-		const uintptr_t Object = (uintptr_t(Flags & PtrMask) << (32 + UObjectPtrTrailingZeroes)) | (uintptr_t(ObjectPtrLow) << UObjectPtrTrailingZeroes);
+		const uintptr_t Object = (uintptr_t((FlagsAndRefCount >> 32) & PtrMask) << (32 + UObjectPtrTrailingZeroes)) | (uintptr_t(ObjectPtrLow) << UObjectPtrTrailingZeroes);
 		return (class UObjectBase*)Object;
 #else
-		PRAGMA_DISABLE_DEPRECATION_WARNINGS
 		return Object;
-		PRAGMA_ENABLE_DEPRECATION_WARNINGS
 #endif
 	}
 
@@ -147,12 +143,11 @@ public:
 	{
 #if UE_ENABLE_FUOBJECT_ITEM_PACKING
 		const uint32 ObjectPtrHi = ((uintptr_t)InObject & 0xFFFF00000000ULL) >> (32 + UObjectPtrTrailingZeroes);
-		Flags |= ObjectPtrHi;
+		// This is not thread-safe, so SetObject should ever just be allowed at UObject creation time.
+		FlagsAndRefCount |= ((uint64)ObjectPtrHi) << 32;
 		ObjectPtrLow = ((uintptr_t)InObject >> UObjectPtrTrailingZeroes) & 0xFFFFFFFF;
 #else
-		PRAGMA_DISABLE_DEPRECATION_WARNINGS
 		Object = InObject;
-		PRAGMA_ENABLE_DEPRECATION_WARNINGS
 #endif
 	}
 
@@ -210,17 +205,6 @@ public:
 		ThisThreadAtomicallyClearedFlag(FlagsToClear);
 	}
 
-	/**
-	 * Uses atomics to clear the specified flag(s). GC internal version
-	 * @param FlagsToClear
-	 * @return True if this call cleared the flag, false if it has been cleared by another thread.
-	 */
-	UE_DEPRECATED(5.5, "ThisThreadAtomicallyClearedFlag_ForGC() can only be used by the garbage collector. Use ThisThreadAtomicallyClearedFlag instead.")
-	FORCEINLINE bool ThisThreadAtomicallyClearedFlag_ForGC(EInternalObjectFlags FlagToClear)
-	{
-		return AtomicallyClearFlag_ForGC(FlagToClear);
-	}
-
 	/**
 	 * Uses atomics to clear the specified flag(s).
 	 * @param FlagsToClear
@@ -233,7 +217,6 @@ public:
 		UE_AUTORTFM_OPEN
 		{
 			FlagToClear &= ~EInternalObjectFlags_ReachabilityFlags; // reachability flags can only be cleared by GC through *_ForGC functions
-			FlagToClear &= ~EInternalObjectFlags::RefCounted; // refcounted flag is internal and must only be cleared internally by AddRef/ReleaseRef.
 			if (!!(FlagToClear & EInternalObjectFlags_RootFlags))
 			{
 				Result = ClearRootFlags(FlagToClear);
@@ -247,17 +230,6 @@ public:
 		return Result;
 	}
 
-	/**
-	 * Uses atomics to set the specified flag(s). GC internal version.
-	 * @param FlagToSet
-	 * @return True if this call set the flag, false if it has been set by another thread.
-	 */
-	UE_DEPRECATED(5.5, "ThisThreadAtomicallySetFlag_ForGC() can only be used by the garbage collector. Use ThisThreadAtomicallySetFlag instead.")
-	FORCEINLINE bool ThisThreadAtomicallySetFlag_ForGC(EInternalObjectFlags FlagToSet)
-	{
-		return AtomicallySetFlag_ForGC(FlagToSet);
-	}
-
 	/**
 	 * Uses atomics to set the specified flag(s)
 	 * @param FlagToSet
@@ -270,7 +242,6 @@ public:
 		UE_AUTORTFM_OPEN
 		{
 			FlagToSet &= ~EInternalObjectFlags_ReachabilityFlags; // reachability flags can only be cleared by GC through *_ForGC functions
-			FlagToSet &= ~EInternalObjectFlags::RefCounted; // refcounted flag is internal and must only be set by AddRef/ReleaseRef.
 			if (!!(FlagToSet & EInternalObjectFlags_RootFlags))
 			{
 				Result = SetRootFlags(FlagToSet);
@@ -293,32 +264,11 @@ public:
 		return (GetFlagsInternal() & int32(InFlags)) == int32(InFlags);
 	}
 
-	UE_DEPRECATED(5.5, "SetUnreachable() can only be used by the garbage collector.")
-	FORCEINLINE void SetUnreachable()
-	{
-		AtomicallySetFlag_ForGC(EInternalObjectFlags::Unreachable);
-	}
-	UE_DEPRECATED(5.5, "SetMaybeUnreachable() can only be used by the garbage collector.")
-	COREUOBJECT_API void SetMaybeUnreachable();
-
-	UE_DEPRECATED(5.5, "ClearUnreachable() can only be used by the garbage collector.")
-	FORCEINLINE void ClearUnreachable()
-	{
-		AtomicallyClearFlag_ForGC(EInternalObjectFlags::Unreachable);
-	}
 	FORCEINLINE bool IsUnreachable() const
 	{
 		return !!(GetFlagsInternal() & int32(EInternalObjectFlags::Unreachable));
 	}
 
-	UE_DEPRECATED(5.5, "IsMaybeUnreachable() can only be used by the garbage collector.")
-	COREUOBJECT_API bool IsMaybeUnreachable() const;
-
-	UE_DEPRECATED(5.5, "ThisThreadAtomicallyClearedRFUnreachable() can only be used by the garbage collector.")
-	FORCEINLINE bool ThisThreadAtomicallyClearedRFUnreachable()
-	{
-		return AtomicallyClearFlag_ForGC(EInternalObjectFlags::Unreachable);
-	}
 	FORCEINLINE void SetGarbage()
 	{
 		AtomicallySetFlag_ForGC(EInternalObjectFlags::Garbage);
@@ -332,22 +282,6 @@ public:
 		return !!(GetFlagsInternal() & int32(EInternalObjectFlags::Garbage));
 	}
 
-	UE_DEPRECATED(5.4, "SetPendingKill() should no longer be used. Use SetGarbage() instead.")
-	FORCEINLINE void SetPendingKill()
-	{
-		SetGarbage();
-	}
-	UE_DEPRECATED(5.4, "ClearPendingKill() should no longer be used. Use ClearGarbage() instead.")
-	FORCEINLINE void ClearPendingKill()
-	{
-		ClearGarbage();
-	}
-	UE_DEPRECATED(5.4, "IsPendingKill() should no longer be used. Use IsGarbage() instead.")
-	FORCEINLINE bool IsPendingKill() const
-	{
-		return IsGarbage();
-	}
-
 	FORCEINLINE void SetRootSet()
 	{
 		ThisThreadAtomicallySetFlag(EInternalObjectFlags::RootSet);
@@ -363,17 +297,24 @@ public:
 
 	FORCEINLINE int32 GetRefCount() const
 	{
-		return RefCount;
+		return (int32)(FlagsAndRefCount & EInternalObjectFlags_RefCountMask);
 	}
 
 	void AddRef()
 	{
 		UE_AUTORTFM_OPEN
 		{
-			FPlatformAtomics::InterlockedIncrement(&RefCount);
-			if ((GetFlags() & EInternalObjectFlags::RefCounted) != EInternalObjectFlags::RefCounted)
+			const int64 NewRefCount = FPlatformAtomics::InterlockedIncrement(&FlagsAndRefCount);
+
+			// Only dirty root when doing a transition from no root flags to a refcount.
+			if ((NewRefCount & EInternalObjectFlags_RefCountMask) == 1 && (GetFlagsInternal(NewRefCount) & (int32)EInternalObjectFlags_RootFlags) == 0)
 			{
-				SetRootFlags(EInternalObjectFlags::RefCounted);
+				// Setting any of the root flags on an object during incremental reachability requires a GC barrier
+				// to make sure an object with root flags does not get Garbage Collected
+				checkf(GetObject(), TEXT("Setting an internal object flag on a null object entry"));
+				GetObject()->MarkAsReachable();
+
+				MarkRootAsDirty();
 			}
 		};
 
@@ -388,13 +329,12 @@ public:
 	{
 		UE_AUTORTFM_OPEN
 		{
-			// This alone is not thread-safe as we may race with AddRef and in that case we don't want ClearRootFlags to apply.
-			// We fix this by validating that the refcount is still 0 while inside the root locks in ClearRootFlags.
-			const int32 NewRefCount = FPlatformAtomics::InterlockedDecrement(&RefCount);
-			check(NewRefCount >= 0);
-			if (NewRefCount == 0)
+			// Only dirty root when doing a transition from a refcount to no refcounts and no other root flags present.
+			const int64 NewRefCount = FPlatformAtomics::InterlockedDecrement(&FlagsAndRefCount);
+			check((NewRefCount & EInternalObjectFlags_RefCountMask) >= 0);
+			if ((NewRefCount & EInternalObjectFlags_RefCountMask) == 0 && (GetFlagsInternal(NewRefCount) & (int32)EInternalObjectFlags_RootFlags) == 0)
 			{
-				ClearRootFlags(EInternalObjectFlags::RefCounted);
+				MarkRootAsDirty();
 			}
 		};
 
@@ -416,38 +356,28 @@ public:
 	COREUOBJECT_API void CreateStatID() const;
 #endif
 
-	// Mark this object item as Reachable and clear MaybeUnreachable flag. For GC use only.
-	UE_DEPRECATED(5.5, "FastMarkAsReachableInterlocked_ForGC() can only be used by the garbage collector.")
-	COREUOBJECT_API void FastMarkAsReachableInterlocked_ForGC();
-
-	// Mark this object item as Reachable and clear ReachableInCluster and MaybeUnreachable flags. For GC use only.
-	UE_DEPRECATED(5.5, "FastMarkAsReachableAndClearReachaleInClusterInterlocked_ForGC() can only be used by the garbage collector.")
-	COREUOBJECT_API void FastMarkAsReachableAndClearReachaleInClusterInterlocked_ForGC();
-
-	/**
-	 * Mark this object item as Reachable and clear MaybeUnreachable flag. Only thread-safe for concurrent clear, not concurrent set+clear. Don't use during mark phase. For GC use only.
-	 * @return True if this call cleared MaybeUnreachable flag, false if it has been cleared by another thread.
-	 */
-	UE_DEPRECATED(5.5, "MarkAsReachableInterlocked_ForGC() can only be used by the garbage collector.")
-	COREUOBJECT_API bool MarkAsReachableInterlocked_ForGC();
-
-	UE_DEPRECATED(5.5, "OffsetOfFlags() can only be used by the garbage collector.")
-	FORCEINLINE static constexpr ::size_t OffsetOfFlags()
+private:
+	FORCEINLINE static int32 GetFlagsInternal(int64 FlagsAndRefCount)
 	{
-		return offsetof(FUObjectItem, Flags);
+		return (int32)(FlagsAndRefCount >> 32);
 	}
 
-private:
 	FORCEINLINE int32 GetFlagsInternal() const
 	{
-		return FPlatformAtomics::AtomicRead_Relaxed((int32*)&Flags);
+		return GetFlagsInternal(FPlatformAtomics::AtomicRead_Relaxed((int64*)&FlagsAndRefCount));
 	}
 	
+	FORCEINLINE static int32 GetRefCountInternal(int64 FlagsAndRefCount)
+	{
+		return (int32)(FlagsAndRefCount & EInternalObjectFlags_RefCountMask);
+	}
+
 	FORCEINLINE int32 GetRefCountInternal() const
 	{
-		return FPlatformAtomics::AtomicRead_Relaxed((int32*)&RefCount);
+		return GetRefCountInternal(FPlatformAtomics::AtomicRead_Relaxed((int64*)&FlagsAndRefCount));
 	}
 
+	COREUOBJECT_API void MarkRootAsDirty();
 	COREUOBJECT_API bool SetRootFlags(EInternalObjectFlags FlagsToSet);
 	COREUOBJECT_API bool ClearRootFlags(EInternalObjectFlags FlagsToClear);
 
@@ -458,17 +388,19 @@ private:
 	 */
 	FORCEINLINE bool AtomicallySetFlag_ForGC(EInternalObjectFlags FlagToSet)
 	{
-		static_assert(sizeof(int32) == sizeof(Flags), "Flags must be 32-bit for atomics.");
+		static_assert(sizeof(int64) == sizeof(FlagsAndRefCount), "Flags must be 64-bit for atomics.");
 		bool bIChangedIt = false;
-		while (1)
+		while (true)
 		{
-			int32 StartValue = GetFlagsInternal();
-			if ((StartValue & int32(FlagToSet)) == int32(FlagToSet))
+			const int64 StartValue = FPlatformAtomics::AtomicRead_Relaxed((int64*)&FlagsAndRefCount);
+			const int32 StartFlagsValue = GetFlagsInternal(StartValue);
+			if ((StartFlagsValue & int32(FlagToSet)) == int32(FlagToSet))
 			{
 				break;
 			}
-			int32 NewValue = StartValue | int32(FlagToSet);
-			if ((int32)FPlatformAtomics::InterlockedCompareExchange((int32*)&Flags, NewValue, StartValue) == StartValue)
+			const int32 NewFlagsValue = StartFlagsValue | int32(FlagToSet);
+			const int64 NewValue = ((int64)NewFlagsValue << 32) | GetRefCountInternal(StartValue);
+			if (FPlatformAtomics::InterlockedCompareExchange((int64*)&FlagsAndRefCount, NewValue, StartValue) == StartValue)
 			{
 				bIChangedIt = true;
 				break;
@@ -477,6 +409,40 @@ private:
 		return bIChangedIt;
 	}
 
+	/**
+	 * Uses atomics to set the specified flag(s). GC internal version.
+	 * @param FlagToSet
+	 * @param OldFlags    Set to the previous value of the flags before the change was made
+	 * @param OldRefCount Set to the previous value of the refcount before the change was made
+	 * @return True if this call set the flag, false if it has been set by another thread.
+	 */
+	FORCEINLINE bool AtomicallySetFlag_ForGC(EInternalObjectFlags FlagToSet, int32& OldFlags, int32& OldRefCount)
+	{
+		static_assert(sizeof(int64) == sizeof(FlagsAndRefCount), "Flags must be 64-bit for atomics.");
+		bool bIChangedIt = false;
+		while (true)
+		{
+			const int64 StartValue = FPlatformAtomics::AtomicRead_Relaxed((int64*)&FlagsAndRefCount);
+			const int32 StartFlagsValue = GetFlagsInternal(StartValue);
+			if ((StartFlagsValue & int32(FlagToSet)) == int32(FlagToSet))
+			{
+				OldFlags = StartFlagsValue;
+				OldRefCount = GetRefCountInternal(StartValue);
+				break;
+			}
+			const int32 NewFlagsValue = StartFlagsValue | int32(FlagToSet);
+			const int64 NewValue = ((int64)NewFlagsValue << 32) | GetRefCountInternal(StartValue);
+			if (FPlatformAtomics::InterlockedCompareExchange((int64*)&FlagsAndRefCount, NewValue, StartValue) == StartValue)
+			{
+				OldFlags = StartFlagsValue;
+				OldRefCount = GetRefCountInternal(StartValue);
+				bIChangedIt = true;
+				break;
+			}
+		}
+		return bIChangedIt;
+	}
+
 	/**
 	 * Uses atomics to clear the specified flag(s). GC internal version
 	 * @param FlagsToClear
@@ -484,18 +450,54 @@ private:
 	 */
 	FORCEINLINE bool AtomicallyClearFlag_ForGC(EInternalObjectFlags FlagToClear)
 	{
-		static_assert(sizeof(int32) == sizeof(Flags), "Flags must be 32-bit for atomics.");
+		static_assert(sizeof(int64) == sizeof(FlagsAndRefCount), "Flags must be 64-bit for atomics.");
+		bool bIChangedIt = false;
+		while (true)
+		{
+			const int64 StartValue = FPlatformAtomics::AtomicRead_Relaxed((int64*)&FlagsAndRefCount);
+			const int32 StartFlagsValue = GetFlagsInternal(StartValue);
+			if (!(StartFlagsValue & int32(FlagToClear)))
+			{
+				break;
+			}
+			const int32 NewFlagsValue = StartFlagsValue & ~int32(FlagToClear);
+			const int64 NewValue = ((int64)NewFlagsValue << 32) | GetRefCountInternal(StartValue);
+			if (FPlatformAtomics::InterlockedCompareExchange((int64*)&FlagsAndRefCount, NewValue, StartValue) == StartValue)
+			{
+				bIChangedIt = true;
+				break;
+			}
+		}
+		return bIChangedIt;
+	}
+
+	/**
+	 * Uses atomics to clear the specified flag(s). GC internal version
+	 * @param FlagsToClear
+	 * @param OldFlags    Set to the previous value of the flags before the change was made
+	 * @param OldRefCount Set to the previous value of the refcount before the change was made
+	 * @return True if this call cleared the flag, false if it has been cleared by another thread.
+	 */
+	FORCEINLINE bool AtomicallyClearFlag_ForGC(EInternalObjectFlags FlagToClear, int32& OldFlags, int32& OldRefCount)
+	{
+		static_assert(sizeof(int64) == sizeof(FlagsAndRefCount), "Flags must be 64-bit for atomics.");
 		bool bIChangedIt = false;
-		while (1)
+		while (true)
 		{
-			int32 StartValue = GetFlagsInternal();
-			if (!(StartValue & int32(FlagToClear)))
+			const int64 StartValue = FPlatformAtomics::AtomicRead_Relaxed((int64*)&FlagsAndRefCount);
+			const int32 StartFlagsValue = GetFlagsInternal(StartValue);
+			if (!(StartFlagsValue & int32(FlagToClear)))
 			{
+				OldFlags = StartFlagsValue;
+				OldRefCount = GetRefCountInternal(StartValue);
 				break;
 			}
-			int32 NewValue = StartValue & ~int32(FlagToClear);
-			if ((int32)FPlatformAtomics::InterlockedCompareExchange((int32*)&Flags, NewValue, StartValue) == StartValue)
+			const int32 NewFlagsValue = StartFlagsValue & ~int32(FlagToClear);
+			const int64 NewValue = ((int64)NewFlagsValue << 32) | GetRefCountInternal(StartValue);
+			if (FPlatformAtomics::InterlockedCompareExchange((int64*)&FlagsAndRefCount, NewValue, StartValue) == StartValue)
 			{
+				OldFlags = StartFlagsValue;
+				OldRefCount = GetRefCountInternal(StartValue);
 				bIChangedIt = true;
 				break;
 			}
diff --git a/Engine/Source/Runtime/CoreUObject/Public/UObject/UObjectBaseUtility.h b/Engine/Source/Runtime/CoreUObject/Public/UObject/UObjectBaseUtility.h
--- a/Engine/Source/Runtime/CoreUObject/Public/UObject/UObjectBaseUtility.h
+++ b/Engine/Source/Runtime/CoreUObject/Public/UObject/UObjectBaseUtility.h
@@ -224,6 +224,16 @@ public:
 		return GUObjectArray.IndexToObject(InternalIndex)->IsRootSet();
 	}
 
+	/**
+	 * Returns the refcount of the object
+	 *
+	 * @return the refcount of the object.
+	 */
+	FORCEINLINE int32 GetRefCount() const
+	{
+		return GUObjectArray.IndexToObject(InternalIndex)->GetRefCount();
+	}
+
 	 /**
 	 * Atomically clear the unreachable flag
 	 *
diff --git a/Engine/Source/Runtime/Engine/Private/UnrealEngine.cpp b/Engine/Source/Runtime/Engine/Private/UnrealEngine.cpp
--- a/Engine/Source/Runtime/Engine/Private/UnrealEngine.cpp
+++ b/Engine/Source/Runtime/Engine/Private/UnrealEngine.cpp
@@ -9455,7 +9455,7 @@ bool UEngine::HandleObjCommand( const TCHAR* Cmd, FOutputDevice& Ar )
 					continue;
 				}
 
-				if (bOnlyListRefCountedObjects && !It->HasAnyInternalFlags(EInternalObjectFlags::RefCounted))
+				if (bOnlyListRefCountedObjects && !It->GetRefCount())
 				{
 					continue;
 				}
