[UObjectArray]
- Fix UB because of PendingConstruction being a negative value

diff --git a/Engine/Source/Runtime/CoreUObject/Private/UObject/UObjectArray.cpp b/Engine/Source/Runtime/CoreUObject/Private/UObject/UObjectArray.cpp
--- a/Engine/Source/Runtime/CoreUObject/Private/UObject/UObjectArray.cpp
+++ b/Engine/Source/Runtime/CoreUObject/Private/UObject/UObjectArray.cpp
@@ -290,7 +290,7 @@ void FUObjectArray::AllocateUObjectIndex(UObjectBase* Object, EInternalObjectFla
 	FUObjectItem* ObjectItem = IndexToObject(Index);
 	UE_CLOG(ObjectItem->GetObject() != nullptr, LogUObjectArray, Fatal, TEXT("Attempting to add %s at index %d but another object (0x%016llx) exists at that index!"), *Object->GetFName().ToString(), Index, (int64)(PTRINT)ObjectItem->GetObject());
 	// At this point all not-compiled-in objects are not fully constructed yet and this is the earliest we can mark them as such
-	ObjectItem->FlagsAndRefCount = ((int64)EInternalObjectFlags::PendingConstruction) << 32;
+	ObjectItem->FlagsAndRefCount = (int64)((uint64)EInternalObjectFlags::PendingConstruction << 32);
 	if (!(IsOpenForDisregardForGC() & GUObjectArray.DisregardForGCEnabled())) //-V792
 	{
 		// It's safe to access FGCFlags::GetReachableFlagValue_ForGC() here because creating new objects is being performed 
