[UObject/GC]
- Fix missing condition for incremental GC

diff --git a/Engine/Source/Runtime/CoreUObject/Public/UObject/UObjectArray.h b/Engine/Source/Runtime/CoreUObject/Public/UObject/UObjectArray.h
--- a/Engine/Source/Runtime/CoreUObject/Public/UObject/UObjectArray.h
+++ b/Engine/Source/Runtime/CoreUObject/Public/UObject/UObjectArray.h
@@ -309,10 +309,13 @@ public:
 			// Only dirty root when doing a transition from no root flags to a refcount.
 			if ((NewRefCount & EInternalObjectFlags_RefCountMask) == 1 && (GetFlagsInternal(NewRefCount) & (int32)EInternalObjectFlags_RootFlags) == 0)
 			{
-				// Setting any of the root flags on an object during incremental reachability requires a GC barrier
-				// to make sure an object with root flags does not get Garbage Collected
-				checkf(GetObject(), TEXT("Setting an internal object flag on a null object entry"));
-				GetObject()->MarkAsReachable();
+				if (UE::GC::GIsIncrementalReachabilityPending)
+				{
+					// Setting any of the root flags on an object during incremental reachability requires a GC barrier
+					// to make sure an object with root flags does not get Garbage Collected
+					checkf(GetObject(), TEXT("Setting an internal object flag on a null object entry"));
+					GetObject()->MarkAsReachable();
+				}
 
 				MarkRootAsDirty();
 			}
