Rework on the current physics island computation algorithm. Merging map is stored distributively (each server only stores edges originating from ite owning islands). New process is the following:
1. Local island IDs are computed before forming the edges across servers.
2. Edges connecting islands are computed based on object proximity,
3. After edges are formed, a BFS style tree of RPCs will broadcast edges to correct servers and all servers wait until broadasting is finished.
4. A DFS style tree of RPCs will make sure the edges are traced through the end
5. All object server ID is updated according to the computed edges.

diff --git a/Engine/Source/Runtime/CoreUObject/Internal/UObject/UObjectMigrationContext.h b/Engine/Source/Runtime/CoreUObject/Internal/UObject/UObjectMigrationContext.h
--- a/Engine/Source/Runtime/CoreUObject/Internal/UObject/UObjectMigrationContext.h
+++ b/Engine/Source/Runtime/CoreUObject/Internal/UObject/UObjectMigrationContext.h
@@ -68,6 +68,9 @@ struct FUObjectMigrationContext final
 	/** The Server that this Object's Physics reside on */
 	FRemoteServerId PhysicsServerId;
 
+	/** The Local Island Index the Object belongs to for Physics */
+	uint32 PhysicsLocalIslandId;
+
 	/** Which side of the migration does this context apply to? Sending or Receiving? */
 	EObjectMigrationSide MigrationSide = EObjectMigrationSide::Receive;
 
diff --git a/Engine/Source/Runtime/CoreUObject/Private/UObject/RemoteObject.cpp b/Engine/Source/Runtime/CoreUObject/Private/UObject/RemoteObject.cpp
--- a/Engine/Source/Runtime/CoreUObject/Private/UObject/RemoteObject.cpp
+++ b/Engine/Source/Runtime/CoreUObject/Private/UObject/RemoteObject.cpp
@@ -104,7 +104,8 @@ public:
 		{
 			if (Elem.Value)
 			{
-				Elem.Value->PhysicsOverrideServerId = FRemoteServerId();
+				Elem.Value->PhysicsIslandId.PhysicsServerId = FRemoteServerId();
+				Elem.Value->PhysicsIslandId.PhysicsLocalIslandId = 0;
 			}
 		}
 #endif
@@ -121,17 +122,39 @@ public:
 	{
 		return AssetPaths.Find(ObjectId);
 	}
-	void UpdateAllPhysicsServerId(const TMap<uint32, uint32>& PhysicsServerMergingMap)
+
+	/**
+	* This function is called by UE::RemoteObject::Handle::UpdateAllPhysicsServerId. More info is shown in that function's description.
+	*/
+	void UpdateAllPhysicsServerId(const TMap<UE::RemoteObject::Handle::FPhysicsIslandId, UE::RemoteObject::Handle::FPhysicsIslandId>& PhysicsIdMergingMap)
+	{
+#if UE_WITH_REMOTE_OBJECT_HANDLE
+		using UE::RemoteObject::Handle::FPhysicsIslandId;
+		for (TPair<FRemoteObjectId, UE::RemoteObject::Handle::FRemoteObjectStub*>& Elem : RemoteObjects)
+		{
+			if (Elem.Value)
+			{
+				const FPhysicsIslandId OriginalPhysicsId = Elem.Value->PhysicsIslandId;
+				if (const FPhysicsIslandId* MergedPhysicsId = PhysicsIdMergingMap.Find(OriginalPhysicsId))
+				{
+					Elem.Value->PhysicsIslandId.PhysicsServerId = MergedPhysicsId->PhysicsServerId;
+				}
+			}
+		}
+#endif
+	}
+
+	void UpdateAllPhysicsLocalIslandId(const TMap<uint32, uint32>& PhysicsLocalIslandMergingMap)
 	{
 #if UE_WITH_REMOTE_OBJECT_HANDLE
 		for (TPair<FRemoteObjectId, UE::RemoteObject::Handle::FRemoteObjectStub*>& Elem : RemoteObjects)
 		{
 			if (Elem.Value)
 			{
-				const uint32 OriginalPhysicsId = Elem.Value->PhysicsOverrideServerId.GetIdNumber();
-				if (const uint32* MergedPhysicsId = PhysicsServerMergingMap.Find(OriginalPhysicsId))
+				const uint32 OriginalLocalIslandId = Elem.Value->PhysicsIslandId.PhysicsLocalIslandId;
+				if (const uint32* MergedLocalIslandId = PhysicsLocalIslandMergingMap.Find(OriginalLocalIslandId))
 				{
-					Elem.Value->PhysicsOverrideServerId = FRemoteServerId(*MergedPhysicsId);
+					Elem.Value->PhysicsIslandId.PhysicsLocalIslandId = *MergedLocalIslandId;
 				}
 			}
 		}
@@ -460,38 +483,32 @@ void ChangeOwnerServerId(const UObject* Object, FRemoteServerId NewOwnerServerId
 #endif
 }
 
-FRemoteServerId GetPhysicsServerId(const UObject* Object)
+FPhysicsIslandId GetPhysicsIslandId(const UObject* Object)
 {
 	using namespace UE::CoreUObject::Private;
 	using namespace UE::RemoteObject::Private;
 
-	FRemoteServerId Result;
+	FPhysicsIslandId Result;
 #if UE_WITH_REMOTE_OBJECT_HANDLE
 	UE::RemoteObject::Handle::FRemoteObjectStub* RemoteStub = FindRemoteObjectStub(FObjectHandleUtils::GetRemoteId(Object));
 	if (RemoteStub)
 	{
-		Result = RemoteStub->PhysicsOverrideServerId;
-	}
-	else
-	{
-		//If the object doesn't have stub, then the physics ID is invalid. It should not participate in push migration.
-		Result = FRemoteServerId();
+		Result = RemoteStub->PhysicsIslandId;
 	}
 #endif
 	return Result;
 }
 
-void ChangePhysicsServerId(const UObject* Object, FRemoteServerId NewPhysicsServerId)
+void ChangePhysicsIslandId(const UObject* Object, FPhysicsIslandId NewPhysicsIslandId)
 {
 	using namespace UE::CoreUObject::Private;
 	using namespace UE::RemoteObject::Private;
 
 #if UE_WITH_REMOTE_OBJECT_HANDLE
-	FRemoteObjectId ThisObject = FObjectHandleUtils::GetRemoteId(Object);
 	UE::RemoteObject::Handle::FRemoteObjectStub* RemoteStub = ObjectMaps->FindOrAddRemoteObjectStub(FObjectHandleUtils::GetRemoteId(Object));
 	if (RemoteStub)
 	{
-		RemoteStub->PhysicsOverrideServerId = NewPhysicsServerId;
+		RemoteStub->PhysicsIslandId = NewPhysicsIslandId;
 	}
 #endif
 }
@@ -506,7 +523,7 @@ void ClearAllPhysicsServerId()
 #endif
 }
 
-void UpdateAllPhysicsServerId(const TMap<uint32, uint32>& PhysicsServerMergingMap)
+void UpdateAllPhysicsServerId(const TMap<FPhysicsIslandId, FPhysicsIslandId>& PhysicsServerMergingMap)
 {
 #if UE_WITH_REMOTE_OBJECT_HANDLE
 	if (Private::ObjectMaps)
@@ -516,6 +533,19 @@ void UpdateAllPhysicsServerId(const TMap<uint32, uint32>& PhysicsServerMergingMa
 #endif
 }
 
+/**
+ Updates all FPhysicsIslandId::PhysicsLocalIslandId based on input map. Usage can be found in description of UE::RemoteObject::Handle::UpdateAllPhysicsLocalIslandId
+*/
+void UpdateAllPhysicsLocalIslandId(const TMap<uint32, uint32>& PhysicsLocalIslandMergingMap)
+{
+#if UE_WITH_REMOTE_OBJECT_HANDLE
+	if (Private::ObjectMaps)
+	{
+		Private::ObjectMaps->UpdateAllPhysicsLocalIslandId(PhysicsLocalIslandMergingMap);
+	}
+#endif
+}
+
 UObject* ResolveObject(const FRemoteObjectStub* Stub, ERemoteReferenceType RefType /*= ERemoteReferenceType::Strong*/)
 {
 	using namespace UE::RemoteObject::Transfer;
@@ -820,6 +850,13 @@ FArchive& operator<<(FArchive& Ar, FRemoteServerId& Id)
 	return Ar;
 }
 
+void UE::RemoteObject::Handle::FPhysicsIslandId::Reset()
+{
+	PhysicsServerId = FRemoteServerId();
+	PhysicsLocalIslandId = 0;
+}
+
+
 // Debugging functionality to help us find these objects in debug builds
 #if !UE_BUILD_SHIPPING && UE_WITH_REMOTE_OBJECT_HANDLE
 /**
diff --git a/Engine/Source/Runtime/CoreUObject/Private/UObject/RemoteObjectSerialization.cpp b/Engine/Source/Runtime/CoreUObject/Private/UObject/RemoteObjectSerialization.cpp
--- a/Engine/Source/Runtime/CoreUObject/Private/UObject/RemoteObjectSerialization.cpp
+++ b/Engine/Source/Runtime/CoreUObject/Private/UObject/RemoteObjectSerialization.cpp
@@ -101,7 +101,7 @@ void LoadObjectFromDisk(const FUObjectMigrationContext& MigrationContext)
 	// We have transferred ownership from the Database to the local server
 	FRemoteServerId LocalServerId = UE::RemoteObject::GetGlobalServerId();
 	const FRemoteServerId& DatabaseId = UE::RemoteObject::Transfer::DatabaseId;
-	UE::RemoteObject::Transfer::OnObjectDataReceived(LocalServerId, LocalServerId, MigrationContext.ObjectId, DatabaseId, ObjectData);
+	UE::RemoteObject::Transfer::OnObjectDataReceived(LocalServerId, LocalServerId, 0, MigrationContext.ObjectId, DatabaseId, ObjectData);
 }
 
 void SaveObjectToDisk(const UE::RemoteObject::Transfer::FMigrateSendParams& Params)
diff --git a/Engine/Source/Runtime/CoreUObject/Private/UObject/RemoteObjectTransfer.cpp b/Engine/Source/Runtime/CoreUObject/Private/UObject/RemoteObjectTransfer.cpp
--- a/Engine/Source/Runtime/CoreUObject/Private/UObject/RemoteObjectTransfer.cpp
+++ b/Engine/Source/Runtime/CoreUObject/Private/UObject/RemoteObjectTransfer.cpp
@@ -199,12 +199,13 @@ public:
 						*HighestPriorityRequest->RequestPriority.ToString());
 
 					FRemoteServerId OwnerServerId = Handle::GetOwnerServerId(RootObject);
-					FRemoteServerId PhysicsServerId = Handle::GetPhysicsServerId(RootObject);
+					FPhysicsIslandId PhysicsIslandId = Handle::GetPhysicsIslandId(RootObject);
 					FUObjectMigrationContext MigrationContext {
 						.ObjectId = HighestPriorityRequest->ObjectId,
 						.RemoteServerId = HighestPriorityRequest->DestinationServerId,
 						.OwnerServerId = OwnerServerId,
-						.PhysicsServerId = PhysicsServerId,
+						.PhysicsServerId = PhysicsIslandId.PhysicsServerId,
+						.PhysicsLocalIslandId = PhysicsIslandId.PhysicsLocalIslandId,
 						.MigrationSide = EObjectMigrationSide::Send
 					};
 					FScopedObjectMigrationContext ScopedMigrationContext(MigrationContext);
@@ -596,14 +597,15 @@ public:
 		for (UObject* SentObject : MultiServerCommitSentObjects)
 		{
 			FRemoteServerId OwnerServerId = Handle::GetOwnerServerId(SentObject);
-			FRemoteServerId PhysicsServerId = Handle::GetPhysicsServerId(SentObject);
+			UE::RemoteObject::Handle::FPhysicsIslandId PhysicsIslandId = Handle::GetPhysicsIslandId(SentObject);
 			FRemoteObjectId SentObjectId = UE::CoreUObject::Private::FObjectHandleUtils::GetRemoteId(SentObject);
 
 			FUObjectMigrationContext MigrationContext {
 				.ObjectId = SentObjectId,
 				.RemoteServerId = OwnerServerId,
 				.OwnerServerId = OwnerServerId,
-				.PhysicsServerId = PhysicsServerId,
+				.PhysicsServerId = PhysicsIslandId.PhysicsServerId,
+				.PhysicsLocalIslandId = PhysicsIslandId.PhysicsLocalIslandId,
 				.MigrationSide = EObjectMigrationSide::Send
 			};
 			FScopedObjectMigrationContext ScopedMigrationContext(MigrationContext);
@@ -887,8 +889,7 @@ UObject* MigrateObjectFromObjectData(FRemoteObjectData& ObjectData, const FUObje
 	for (UObject* ReceivedObject : ReceivedObjects)
 	{
 		UE::RemoteObject::Handle::ChangeOwnerServerId(ReceivedObject, MigrationContext.OwnerServerId);
-		UE::RemoteObject::Handle::ChangePhysicsServerId(ReceivedObject, MigrationContext.PhysicsServerId);
-
+		UE::RemoteObject::Handle::ChangePhysicsIslandId(ReceivedObject, UE::RemoteObject::Handle::FPhysicsIslandId(MigrationContext.PhysicsServerId, MigrationContext.PhysicsLocalIslandId));
 #if UE_WITH_REMOTE_OBJECT_HANDLE
 		ensureAlwaysMsgf(ReceivedObject->HasAnyInternalFlags(EInternalObjectFlags::Borrowed) == !UE::RemoteObject::Handle::IsOwned(ReceivedObject),
 			TEXT("Object %s (%s) is %s owned but is %s borrowed"),
@@ -1044,12 +1045,13 @@ void TransferObjectOwnershipToRemoteServer(UObject* Object, FRemoteServerId Dest
 	const FRemoteObjectId CanonicalRootObjectId = UE::CoreUObject::Private::FObjectHandleUtils::GetRemoteId(RootObject);
 
 	// Figure out the ownership semantics
-	FRemoteServerId PhysicsServerId = Handle::GetPhysicsServerId(RootObject);
+	UE::RemoteObject::Handle::FPhysicsIslandId PhysicsIslandId = Handle::GetPhysicsIslandId(RootObject);
 	FUObjectMigrationContext MigrationContext {
 		.ObjectId = ObjectId,
 		.RemoteServerId = DestinationServerId,
 		.OwnerServerId = DestinationServerId,
-		.PhysicsServerId = PhysicsServerId,
+		.PhysicsServerId = PhysicsIslandId.PhysicsServerId,
+		.PhysicsLocalIslandId = PhysicsIslandId.PhysicsLocalIslandId,
 		.MigrationSide = EObjectMigrationSide::Send
 	};
     FScopedObjectMigrationContext ScopedMigrationContext(MigrationContext);
@@ -1145,7 +1147,7 @@ void MigrateObjectToRemoteServerWithExplicitPriority(FRemoteWorkPriority Request
 	GTransferQueue.RequestObjectMigration(RequestPriority, ObjectId, DestinationServerId);
 }
 
-void OnObjectDataReceived(FRemoteServerId OwnerServerId, FRemoteServerId PhysicsId, FRemoteObjectId ObjectId, FRemoteServerId RemoteServerId, FRemoteObjectData& Data)
+void OnObjectDataReceived(FRemoteServerId OwnerServerId, FRemoteServerId PhysicsId, uint32 PhysicsLocalIslandId, FRemoteObjectId ObjectId, FRemoteServerId RemoteServerId, FRemoteObjectData& Data)
 {
 	check(!AutoRTFM::IsClosed());
 
@@ -1154,6 +1156,7 @@ void OnObjectDataReceived(FRemoteServerId OwnerServerId, FRemoteServerId Physics
 		.RemoteServerId = RemoteServerId,
 		.OwnerServerId = OwnerServerId,
 		.PhysicsServerId = PhysicsId,
+		.PhysicsLocalIslandId = PhysicsLocalIslandId,
 		.MigrationSide = EObjectMigrationSide::Receive
 	};
 	FScopedObjectMigrationContext ScopedMigrationContext(MigrationContext);
diff --git a/Engine/Source/Runtime/CoreUObject/Public/UObject/RemoteObject.h b/Engine/Source/Runtime/CoreUObject/Public/UObject/RemoteObject.h
--- a/Engine/Source/Runtime/CoreUObject/Public/UObject/RemoteObject.h
+++ b/Engine/Source/Runtime/CoreUObject/Public/UObject/RemoteObject.h
@@ -27,6 +27,64 @@ namespace UE::RemoteObject
 
 namespace UE::RemoteObject::Handle
 {
+	struct FPhysicsIslandId
+	{
+		friend struct FRemoteObjectId;
+
+		FRemoteServerId PhysicsServerId;
+
+		//IMPORTANT: LocalIslandIndex = 0 means it's invalid. 
+		uint32 PhysicsLocalIslandId = 0;
+
+		void Reset();
+
+		FPhysicsIslandId() = default;
+
+		FPhysicsIslandId(const FRemoteServerId& PhysicsServerIdIn, const uint32 PhysicsLocalIslandIdIn)
+			: PhysicsServerId(PhysicsServerIdIn), PhysicsLocalIslandId(PhysicsLocalIslandIdIn)
+		{
+		}
+
+		bool operator==(const FPhysicsIslandId& Other) const
+		{
+			return PhysicsServerId == Other.PhysicsServerId && PhysicsLocalIslandId == Other.PhysicsLocalIslandId;
+		}
+
+		bool IsValid() const
+		{
+			return PhysicsServerId.IsValid() && PhysicsLocalIslandId != INDEX_NONE;
+		}
+
+		bool operator<(const FPhysicsIslandId& Other) const
+		{
+			return PhysicsServerId < Other.PhysicsServerId ||
+				(PhysicsServerId == Other.PhysicsServerId && PhysicsLocalIslandId < Other.PhysicsLocalIslandId);
+		}
+
+		bool operator<=(const FPhysicsIslandId& Other) const
+		{
+			return (*this) < Other || (*this) == Other;
+		}
+
+		friend FArchive& operator<<(FArchive& Ar, FPhysicsIslandId& ID)
+		{
+			Ar << ID.PhysicsServerId;
+			Ar << ID.PhysicsLocalIslandId;
+			return Ar;
+		}
+
+		FString ToString() const
+		{
+			return FString::Printf(TEXT("[%s, %d]"), *(PhysicsServerId.ToString()), PhysicsLocalIslandId);
+		}
+	};
+
+	FORCEINLINE uint32 GetTypeHash(const FPhysicsIslandId& Key)
+	{
+		// Combine hashes of the members
+		return HashCombine(::GetTypeHash(Key.PhysicsServerId.GetIdNumber()), ::GetTypeHash(Key.PhysicsLocalIslandId));
+	}
+
 	// Structure that holds basic information about a remote object
 	// This is what FObjectPtr that references a remote object actually points to after the remote object's memory has been claimed by GC
 	struct FRemoteObjectStub
@@ -45,9 +103,7 @@ namespace UE::RemoteObject::Handle
 			(note: only valid if the object is Local) */
 		FRemoteServerId OwningServerId;
 
-		/** Server id of the which server the object should be on for physics
-			Default (invalid) means physics ID should not be considered when migrating the object.   */
-		FRemoteServerId PhysicsOverrideServerId;
+		FPhysicsIslandId PhysicsIslandId;
 
 		FRemoteObjectStub() = default;
 		COREUOBJECT_API explicit FRemoteObjectStub(UObject* Object);
@@ -123,16 +179,21 @@ namespace UE::RemoteObject::Handle
 	COREUOBJECT_API void ChangeOwnerServerId(const UObject* Object, FRemoteServerId NewOwnerServerId);
 
 	/**
-	* Get the physics server id where the object should be simulated on
+	* Get PhysicsIslandId of the input object
 	*/
-	COREUOBJECT_API FRemoteServerId GetPhysicsServerId(const UObject* Object);
+	COREUOBJECT_API FPhysicsIslandId GetPhysicsIslandId(const UObject* Object);
 
-	/**
-	* Sets the physics id where the object should go to
-	*/
-	COREUOBJECT_API void ChangePhysicsServerId(const UObject* Object, FRemoteServerId NewPhysicsServerId);
+	COREUOBJECT_API void ChangePhysicsIslandId(const UObject* Object, FPhysicsIslandId NewPhysicsIslandId);
 
 	COREUOBJECT_API void ClearAllPhysicsServerId();
 
-	COREUOBJECT_API void UpdateAllPhysicsServerId(const TMap<uint32, uint32>& PhysicsServerMergingMap);
+	/**
+	* Updates FPhysicsIslandId::PhysicsServerId using the input map. Updating PhysicsServerId only because it is the only data needed for physics migration. (SeeDstmPhysicsMigration.cpp)
+	*/
+	COREUOBJECT_API void UpdateAllPhysicsServerId(const TMap<FPhysicsIslandId, FPhysicsIslandId>& PhysicsServerMergingMap);
+
+	/**
+	* Updates PhysicsLocalIslandId of FPhysicsIslandId using the input map. Usage in DstmPhysicsMigration.cpp for local island computation.
+	*/
+	COREUOBJECT_API void UpdateAllPhysicsLocalIslandId(const TMap<uint32, uint32>& PhysicsIslandMergingMap);
 }
diff --git a/Engine/Source/Runtime/CoreUObject/Public/UObject/RemoteObjectTransfer.h b/Engine/Source/Runtime/CoreUObject/Public/UObject/RemoteObjectTransfer.h
--- a/Engine/Source/Runtime/CoreUObject/Public/UObject/RemoteObjectTransfer.h
+++ b/Engine/Source/Runtime/CoreUObject/Public/UObject/RemoteObjectTransfer.h
@@ -67,7 +67,7 @@ namespace UE::RemoteObject::Transfer
 	* @param RemoteServerId Server that sent object data
 	* @param Data Remote object data. Data ownership is transferred to the remote object queue
 	*/
-	COREUOBJECT_API void OnObjectDataReceived(FRemoteServerId OwnerServerId, FRemoteServerId PhysicsId, FRemoteObjectId ObjectId, FRemoteServerId RemoteServerId, FRemoteObjectData& Data);
+	COREUOBJECT_API void OnObjectDataReceived(FRemoteServerId OwnerServerId, FRemoteServerId PhysicsId, uint32 PhysicsLocalIslandId, FRemoteObjectId ObjectId, FRemoteServerId RemoteServerId, FRemoteObjectData& Data);
 
 	/**
 	* Called when a remote object request was denied by a remote server
diff --git a/Engine/Source/Runtime/CoreUObject/Public/UObject/RemoteObjectTypes.h b/Engine/Source/Runtime/CoreUObject/Public/UObject/RemoteObjectTypes.h
--- a/Engine/Source/Runtime/CoreUObject/Public/UObject/RemoteObjectTypes.h
+++ b/Engine/Source/Runtime/CoreUObject/Public/UObject/RemoteObjectTypes.h
@@ -198,3 +198,8 @@ FORCEINLINE uint32 GetTypeHash(const FRemoteObjectId& ObjectId)
 {
 	return ObjectId.GetTypeHash();
 }
+
+FORCEINLINE uint32 GetTypeHash(const FRemoteServerId& ServerId)
+{
+	return ServerId.GetIdNumber();
+}
