Fix FClassProperty::GetCPPType to return correct string for properties such as TObjectPtr<UClass>.
Remove GetCPPTypeCustom from object properties.

diff --git a/Engine/Source/Runtime/CoreUObject/Private/UObject/PropertyClass.cpp b/Engine/Source/Runtime/CoreUObject/Private/UObject/PropertyClass.cpp
--- a/Engine/Source/Runtime/CoreUObject/Private/UObject/PropertyClass.cpp
+++ b/Engine/Source/Runtime/CoreUObject/Private/UObject/PropertyClass.cpp
@@ -158,26 +158,18 @@ const TCHAR* FClassProperty::ImportText_Internal( const TCHAR* Buffer, void* Con
 FString FClassProperty::GetCPPType(FString* ExtendedTypeText, uint32 CPPExportFlags) const
 {
 	check(MetaClass);
-	return GetCPPTypeCustom(ExtendedTypeText, CPPExportFlags,
-		FString::Printf(TEXT("%s%s"), MetaClass->GetPrefixCPP(), *MetaClass->GetName()));
-}
-
-FString FClassProperty::GetCPPTypeCustom(FString* ExtendedTypeText, uint32 CPPExportFlags, const FString& InnerNativeTypeName) const
-{
 	if (EnumHasAnyFlags(PropertyFlags, CPF_TObjectPtr))
 	{
 		if (!EnumHasAnyFlags((EPropertyExportCPPFlags)CPPExportFlags, CPPF_NoTObjectPtr))
 		{
-			ensure(!InnerNativeTypeName.IsEmpty());
-			return FString::Printf(TEXT("TObjectPtr<%s>"), *InnerNativeTypeName);
+			return FString::Printf(TEXT("TObjectPtr<%s%s>"), PropertyClass->GetPrefixCPP(), *PropertyClass->GetName());
 		}
 	}
 	else if (EnumHasAnyFlags(PropertyFlags, CPF_UObjectWrapper))
 	{
-		ensure(!InnerNativeTypeName.IsEmpty());
-		return FString::Printf(TEXT("TSubclassOf<%s>"), *InnerNativeTypeName);
+		return FString::Printf(TEXT("TSubclassOf<%s%s>"), MetaClass->GetPrefixCPP(), *MetaClass->GetName());
 	}
-	return TEXT("UClass*");
+	return FString::Printf(TEXT("%s%s*"), PropertyClass->GetPrefixCPP(), *PropertyClass->GetName());
 }
 
 FString FClassProperty::GetCPPMacroType( FString& ExtendedTypeText ) const
diff --git a/Engine/Source/Runtime/CoreUObject/Private/UObject/PropertyLazyObjectPtr.cpp b/Engine/Source/Runtime/CoreUObject/Private/UObject/PropertyLazyObjectPtr.cpp
--- a/Engine/Source/Runtime/CoreUObject/Private/UObject/PropertyLazyObjectPtr.cpp
+++ b/Engine/Source/Runtime/CoreUObject/Private/UObject/PropertyLazyObjectPtr.cpp
@@ -30,14 +30,8 @@ FLazyObjectProperty::FLazyObjectProperty(UField* InField)
 
 FString FLazyObjectProperty::GetCPPType(FString* ExtendedTypeText/*=NULL*/, uint32 CPPExportFlags/*=0*/) const
 {
-	return GetCPPTypeCustom(ExtendedTypeText, CPPExportFlags,
-		FString::Printf(TEXT("%s%s"), PropertyClass->GetPrefixCPP(), *PropertyClass->GetName()));
-}
-
-FString FLazyObjectProperty::GetCPPTypeCustom(FString* ExtendedTypeText, uint32 CPPExportFlags, const FString& InnerNativeTypeName) const
-{
-	ensure(!InnerNativeTypeName.IsEmpty());
-	return FString::Printf(TEXT("TLazyObjectPtr<%s>"), *InnerNativeTypeName);
+	check(PropertyClass);
+	return FString::Printf(TEXT("TLazyObjectPtr<%s%s>"), PropertyClass->GetPrefixCPP(), *PropertyClass->GetName());
 }
 
 FString FLazyObjectProperty::GetCPPMacroType( FString& ExtendedTypeText ) const
diff --git a/Engine/Source/Runtime/CoreUObject/Private/UObject/PropertyObject.cpp b/Engine/Source/Runtime/CoreUObject/Private/UObject/PropertyObject.cpp
--- a/Engine/Source/Runtime/CoreUObject/Private/UObject/PropertyObject.cpp
+++ b/Engine/Source/Runtime/CoreUObject/Private/UObject/PropertyObject.cpp
@@ -34,14 +34,14 @@ FObjectProperty::FObjectProperty(UField* InField)
 }
 #endif // WITH_EDITORONLY_DATA
 
-FString FObjectProperty::GetCPPTypeCustom(FString* ExtendedTypeText, uint32 CPPExportFlags, const FString& InnerNativeTypeName)  const
+FString FObjectProperty::GetCPPType(FString* ExtendedTypeText, uint32 CPPExportFlags)  const
 {
-	ensure(!InnerNativeTypeName.IsEmpty());
+	check(PropertyClass);
 	if (HasAnyPropertyFlags(CPF_TObjectPtr) && !(CPPExportFlags & CPPF_NoTObjectPtr))
 	{
-		return FString::Printf(TEXT("TObjectPtr<%s>"), *InnerNativeTypeName);
+		return FString::Printf(TEXT("TObjectPtr<%s%s>"), PropertyClass->GetPrefixCPP(), *PropertyClass->GetName());
 	}
-	return FString::Printf(TEXT("%s*"), *InnerNativeTypeName);
+	return FString::Printf(TEXT("%s%s*"), PropertyClass->GetPrefixCPP(), *PropertyClass->GetName());
 }
 
 FString FObjectProperty::GetCPPMacroType( FString& ExtendedTypeText ) const
diff --git a/Engine/Source/Runtime/CoreUObject/Private/UObject/PropertySoftClassPtr.cpp b/Engine/Source/Runtime/CoreUObject/Private/UObject/PropertySoftClassPtr.cpp
--- a/Engine/Source/Runtime/CoreUObject/Private/UObject/PropertySoftClassPtr.cpp
+++ b/Engine/Source/Runtime/CoreUObject/Private/UObject/PropertySoftClassPtr.cpp
@@ -55,14 +55,9 @@ void FSoftClassProperty::PostDuplicate(const FField& InField)
 FString FSoftClassProperty::GetCPPType(FString* ExtendedTypeText, uint32 CPPExportFlags) const
 {
 	check(MetaClass);
-	return GetCPPTypeCustom(ExtendedTypeText, CPPExportFlags, 
-		FString::Printf(TEXT("%s%s"), MetaClass->GetPrefixCPP(), *MetaClass->GetName()));
-}
-FString FSoftClassProperty::GetCPPTypeCustom(FString* ExtendedTypeText, uint32 CPPExportFlags, const FString& InnerNativeTypeName) const
-{
-	ensure(!InnerNativeTypeName.IsEmpty());
-	return FString::Printf(TEXT("TSoftClassPtr<%s> "), *InnerNativeTypeName);
+	return FString::Printf(TEXT("TSoftClassPtr<%s%s> "), MetaClass->GetPrefixCPP(), *MetaClass->GetName());
 }
+
 FString FSoftClassProperty::GetCPPMacroType( FString& ExtendedTypeText ) const
 {
 	ExtendedTypeText = FString::Printf(TEXT("TSoftClassPtr<%s%s> "),MetaClass->GetPrefixCPP(),*MetaClass->GetName());
diff --git a/Engine/Source/Runtime/CoreUObject/Private/UObject/PropertySoftObjectPtr.cpp b/Engine/Source/Runtime/CoreUObject/Private/UObject/PropertySoftObjectPtr.cpp
--- a/Engine/Source/Runtime/CoreUObject/Private/UObject/PropertySoftObjectPtr.cpp
+++ b/Engine/Source/Runtime/CoreUObject/Private/UObject/PropertySoftObjectPtr.cpp
@@ -35,17 +35,11 @@ FSoftObjectProperty::FSoftObjectProperty(UField* InField)
 }
 #endif // WITH_EDITORONLY_DATA
 
-FString FSoftObjectProperty::GetCPPTypeCustom(FString* ExtendedTypeText, uint32 CPPExportFlags, const FString& InnerNativeTypeName) const
-{
-	ensure(!InnerNativeTypeName.IsEmpty());
-	return FString::Printf(TEXT("TSoftObjectPtr<%s>"), *InnerNativeTypeName);
-}
-
 FString FSoftObjectProperty::GetCPPType(FString* ExtendedTypeText, uint32 CPPExportFlags) const
 {
 	if (ensureMsgf(PropertyClass, TEXT("Soft object property missing PropertyClass: %s"), *GetFullNameSafe(this)))
 	{
-		return Super::GetCPPType(ExtendedTypeText, CPPExportFlags);
+		return FString::Printf(TEXT("TSoftObjectPtr<%s%s>"), PropertyClass->GetPrefixCPP(), *PropertyClass->GetName());
 	}
 	else
 	{
diff --git a/Engine/Source/Runtime/CoreUObject/Private/UObject/PropertyWeakObjectPtr.cpp b/Engine/Source/Runtime/CoreUObject/Private/UObject/PropertyWeakObjectPtr.cpp
--- a/Engine/Source/Runtime/CoreUObject/Private/UObject/PropertyWeakObjectPtr.cpp
+++ b/Engine/Source/Runtime/CoreUObject/Private/UObject/PropertyWeakObjectPtr.cpp
@@ -32,18 +32,12 @@ FWeakObjectProperty::FWeakObjectProperty(UField* InField)
 
 FString FWeakObjectProperty::GetCPPType( FString* ExtendedTypeText/*=NULL*/, uint32 CPPExportFlags/*=0*/ ) const
 {
-	return GetCPPTypeCustom(ExtendedTypeText, CPPExportFlags,
-		FString::Printf(TEXT("%s%s"), PropertyClass->GetPrefixCPP(), *PropertyClass->GetName()));
-}
-
-FString FWeakObjectProperty::GetCPPTypeCustom(FString* ExtendedTypeText, uint32 CPPExportFlags, const FString& InnerNativeTypeName) const
-{
-	ensure(!InnerNativeTypeName.IsEmpty());
+	check(PropertyClass);
 	if (PropertyFlags & CPF_AutoWeak)
 	{
-		return FString::Printf(TEXT("TAutoWeakObjectPtr<%s>"), *InnerNativeTypeName);
+		return FString::Printf(TEXT("TAutoWeakObjectPtr<%s%s>"), PropertyClass->GetPrefixCPP(), *PropertyClass.GetName());
 	}
-	return FString::Printf(TEXT("TWeakObjectPtr<%s>"), *InnerNativeTypeName);
+	return FString::Printf(TEXT("TWeakObjectPtr<%s%s>"), PropertyClass->GetPrefixCPP(), *PropertyClass.GetName());
 }
 
 FString FWeakObjectProperty::GetCPPMacroType( FString& ExtendedTypeText ) const
diff --git a/Engine/Source/Runtime/CoreUObject/Public/UObject/UnrealType.h b/Engine/Source/Runtime/CoreUObject/Public/UObject/UnrealType.h
--- a/Engine/Source/Runtime/CoreUObject/Public/UObject/UnrealType.h
+++ b/Engine/Source/Runtime/CoreUObject/Public/UObject/UnrealType.h
@@ -2564,7 +2564,8 @@ public:
 	// FObjectPropertyBase interface
 public:
 
-	UE_API virtual FString GetCPPTypeCustom(FString* ExtendedTypeText, uint32 CPPExportFlags, const FString& InnerNativeTypeName) const PURE_VIRTUAL(FObjectPropertyBase::GetCPPTypeCustom, return TEXT(""););
+	UE_DEPRECATED(5.7, "GetCPPTypeCustom is deprecated, and object properties now implement GetCPPType directly.")
+	UE_API FString GetCPPTypeCustom(FString* ExtendedTypeText, uint32 CPPExportFlags, const FString& InnerNativeTypeName) const { return {}; }
 
 	/**
 	 * Parses a text buffer into an object reference.
@@ -2839,14 +2840,6 @@ public:
 			|| (EnumHasAnyFlags(InReferenceType, EPropertyObjectReferenceType::Soft) && TIsSoftObjectPointerType<InTCppType>::Value);
 	}
 	// End of FProperty interface
-
-	// TProperty::GetCPPType should not be used here
-	virtual FString GetCPPType(FString* ExtendedTypeText, uint32 CPPExportFlags) const override
-	{
-		check(this->PropertyClass);
-		return this->GetCPPTypeCustom(ExtendedTypeText, CPPExportFlags, 
-			FString::Printf(TEXT("%s%s"), this->PropertyClass->GetPrefixCPP(), *this->PropertyClass->GetName()));
-	}
 };
 
 enum class EObjectPropertyOptions
@@ -2877,6 +2870,7 @@ class FObjectProperty : public TFObjectPropertyBase<TObjectPtr<UObject>>
 #endif // WITH_EDITORONLY_DATA
 
 	// UHT interface
+	UE_API virtual FString GetCPPType(FString* ExtendedTypeText, uint32 CPPExportFlags) const override;
 	UE_API virtual FString GetCPPMacroType( FString& ExtendedTypeText ) const  override;
 	// End of UHT interface
 
@@ -2913,7 +2907,6 @@ protected:
 	UE_API virtual void SetObjectPropertyValueUnchecked_InContainer(void* ContainerAddress, UObject* Value, int32 ArrayIndex = 0) const override;
 	UE_API virtual void SetObjectPtrPropertyValueUnchecked_InContainer(void* ContainerAddress, TObjectPtr<UObject> Ptr, int32 ArrayIndex = 0) const override;
 public:
-	UE_API virtual FString GetCPPTypeCustom(FString* ExtendedTypeText, uint32 CPPExportFlags, const FString& InnerNativeTypeName)  const override;
 	// End of FObjectPropertyBase interface
 	
 	/**
@@ -2989,7 +2982,6 @@ class FWeakObjectProperty : public TFObjectPropertyBase<FWeakObjectPtr>
 	// UHT interface
 	UE_API virtual FString GetCPPMacroType( FString& ExtendedTypeText ) const  override;
 	UE_API virtual FString GetCPPType( FString* ExtendedTypeText, uint32 CPPExportFlags ) const override;
-	UE_API virtual FString GetCPPTypeCustom(FString* ExtendedTypeText, uint32 CPPExportFlags, const FString& InnerNativeTypeName) const override;
 	// End of UHT interface
 
 	// FProperty interface
@@ -3048,7 +3040,6 @@ class FLazyObjectProperty : public TFObjectPropertyBase<FLazyObjectPtr>
 	// UHT interface
 	UE_API virtual FString GetCPPMacroType( FString& ExtendedTypeText ) const  override;
 	UE_API virtual FString GetCPPType( FString* ExtendedTypeText, uint32 CPPExportFlags ) const override;
-	UE_API virtual FString GetCPPTypeCustom(FString* ExtendedTypeText, uint32 CPPExportFlags, const FString& InnerNativeTypeName) const override;
 	// End of UHT interface
 
 	// FProperty interface
@@ -3140,7 +3131,6 @@ protected:
 	UE_API virtual void SetObjectPtrPropertyValueUnchecked_InContainer(void* ContainerAddress, TObjectPtr<UObject> Ptr, int32 ArrayIndex = 0) const override;
 public:
 	UE_API virtual bool AllowCrossLevel() const override;
-	UE_API virtual FString GetCPPTypeCustom(FString* ExtendedTypeText, uint32 CPPExportFlags, const FString& InnerNativeTypeName) const override;
 	UE_API virtual FString GetCPPType(FString* ExtendedTypeText, uint32 CPPExportFlags) const override;
 
 private:
@@ -3202,7 +3192,6 @@ public:
 #endif
 	// End of FProperty interface
 
-	UE_API virtual FString GetCPPTypeCustom(FString* ExtendedTypeText, uint32 CPPExportFlags, const FString& InnerNativeTypeName)  const override;
 
 	/**
 	 * Setter function for this property's MetaClass member. Favor this function 
@@ -3267,7 +3256,6 @@ public:
 	UE_API virtual bool SameType(const FProperty* Other) const override;
 	// End of FProperty interface
 
-	UE_API virtual FString GetCPPTypeCustom(FString* ExtendedTypeText, uint32 CPPExportFlags, const FString& InnerNativeTypeName)  const override;
 
 	/**
 	 * Setter function for this property's MetaClass member. Favor this function 
diff --git a/Engine/Source/Runtime/CoreUObject/Tests/ObjectPtrTestClass.cpp b/Engine/Source/Runtime/CoreUObject/Tests/ObjectPtrTestClass.cpp
--- a/Engine/Source/Runtime/CoreUObject/Tests/ObjectPtrTestClass.cpp
+++ b/Engine/Source/Runtime/CoreUObject/Tests/ObjectPtrTestClass.cpp
@@ -163,10 +163,10 @@ IMPLEMENT_CORE_INTRINSIC_CLASS(UObjectWithClassProperty, UObject,
 			Params.ClassFunc = nullptr;
 			Params.MetaClassFunc = []()
 			{
-				return UClass::StaticClass();
+				return UObject::StaticClass();
 			};
 			auto Property = new FClassProperty(Class, Params);
-			Property->PropertyClass = UObject::StaticClass();
+			Property->PropertyClass = UClass::StaticClass();
 		}
 	
 		{
@@ -178,10 +178,10 @@ IMPLEMENT_CORE_INTRINSIC_CLASS(UObjectWithClassProperty, UObject,
 			Params.ClassFunc = nullptr;
 			Params.MetaClassFunc = []()
 			{
-				return UClass::StaticClass();
+				return UObject::StaticClass();
 			};
 			auto Property = new FClassProperty(Class, Params);
-			Property->PropertyClass = UObject::StaticClass();
+			Property->PropertyClass = UClass::StaticClass();
 		}
 
 		{
@@ -196,7 +196,7 @@ IMPLEMENT_CORE_INTRINSIC_CLASS(UObjectWithClassProperty, UObject,
 				return UObjectPtrTestClass::StaticClass();
 			};
 			auto Property = new FClassProperty(Class, Params);
-			Property->PropertyClass = UObject::StaticClass();
+			Property->PropertyClass = UClass::StaticClass();
 		}
 	}
 );
