Add missing scope lock to move constructor of object annotation.

#tests editor

diff --git a/Engine/Source/Runtime/CoreUObject/Public/UObject/UObjectAnnotation.h b/Engine/Source/Runtime/CoreUObject/Public/UObject/UObjectAnnotation.h
--- a/Engine/Source/Runtime/CoreUObject/Public/UObject/UObjectAnnotation.h
+++ b/Engine/Source/Runtime/CoreUObject/Public/UObject/UObjectAnnotation.h
@@ -314,6 +314,7 @@ private:
 		}
 		else
 		{
+			UE::TRWScopeLock AnnotationCacheLock(AnnotationCacheCritical, SLT_Write);
 			AnnotationCacheKey = Key;
 			AnnotationCacheValue = MoveTemp(Value);
 		}
