Fixing a crash in CppStructOps caused by unloading a dll at an inopportune time resulting in a dangling pointer to a vtable.

diff --git a/Engine/Source/Runtime/CoreUObject/Private/UObject/Class.cpp b/Engine/Source/Runtime/CoreUObject/Private/UObject/Class.cpp
--- a/Engine/Source/Runtime/CoreUObject/Private/UObject/Class.cpp
+++ b/Engine/Source/Runtime/CoreUObject/Private/UObject/Class.cpp
@@ -3104,6 +3104,12 @@ void UScriptStruct::PrepareCppStructOps()
 			bPrepareCppStructOpsCompleted = true;
 			return;
 		}
+		else
+		{
+			// Clean it up so it doesn't linger if we unload this dll
+			GetDeferredCppStructOps().Remove(GetFlattenedStructPathName());
+		}
+
 		StructFlags = EStructFlags(StructFlags | STRUCT_Native);
 		
 #if !(UE_BUILD_SHIPPING || UE_BUILD_TEST)
