[UE] Enable dynamic instancing during instanced subobject construction based on the source subobject class.

Notes:
- Fixes an oversight where subobjects in the instancing graph were not being instanced based on their source type, but rather the root object's source type.

diff --git a/Engine/Source/Runtime/CoreUObject/Private/UObject/CoreNative.cpp b/Engine/Source/Runtime/CoreUObject/Private/UObject/CoreNative.cpp
--- a/Engine/Source/Runtime/CoreUObject/Private/UObject/CoreNative.cpp
+++ b/Engine/Source/Runtime/CoreUObject/Private/UObject/CoreNative.cpp
@@ -293,8 +293,12 @@ UObject* FObjectInstancingGraph::GetInstancedSubobject( UObject* SourceSubobject
 
 							if (!InstancedSubobject)
 							{
+								// use dynamic instancing if necessary
+								const UClass* SubobjectClass = SourceSubobject->GetClass();
+								TGuardValue<bool> UseDynamicInstancing(bCanUseDynamicInstancing, SubobjectClass->ShouldUseDynamicSubobjectInstancing());
+
 								// finally, create the subobject instance
-								FStaticConstructObjectParameters Params(SourceSubobject->GetClass());
+								FStaticConstructObjectParameters Params(SubobjectClass);
 								Params.Outer = SubobjectOuter;
 								Params.Name = SubobjectName;
 								Params.SetFlags = SubobjectOuter->GetMaskedFlags(RF_PropagateToSubObjects);
