Fix incorrect lock used when accessing object annotation cache

#tests editor

diff --git a/Engine/Source/Runtime/CoreUObject/Public/UObject/UObjectAnnotation.h b/Engine/Source/Runtime/CoreUObject/Public/UObject/UObjectAnnotation.h
--- a/Engine/Source/Runtime/CoreUObject/Public/UObject/UObjectAnnotation.h
+++ b/Engine/Source/Runtime/CoreUObject/Public/UObject/UObjectAnnotation.h
@@ -247,7 +247,7 @@ public:
 
 		{
 			// Check cache
-			UE::TRWScopeLock AnnotationCacheLock(AnnotationMapCritical, SLT_ReadOnly);
+			UE::TRWScopeLock AnnotationCacheLock(AnnotationCacheCritical, SLT_ReadOnly);
 			if (Object == AnnotationCacheKey)
 			{
 				return AnnotationCacheValue;
