fix the PendingTransaction counter in the traceview

diff --git a/Engine/Source/Runtime/CoreUObject/Private/UObject/RemoteExecutor.cpp b/Engine/Source/Runtime/CoreUObject/Private/UObject/RemoteExecutor.cpp
--- a/Engine/Source/Runtime/CoreUObject/Private/UObject/RemoteExecutor.cpp
+++ b/Engine/Source/Runtime/CoreUObject/Private/UObject/RemoteExecutor.cpp
@@ -17,9 +17,11 @@ TDelegate<void(FRemoteTransactionId, FRemoteServerId /*ServerId*/)> ReadyRemoteM
 TDelegate<void(FRemoteTransactionId, FRemoteServerId /*ServerId*/)> AbortRemoteMultiServerCommitDelegate;
 TDelegate<TOptional<TTuple<FName, FRemoteWorkPriority, bool, TFunction<void(void)>>>(void)> FetchNextDeferredRPCDelegate;
 
+TMulticastDelegate<void(FRemoteTransactionId, FName)> OnTransactionQueuedDelegate;
 TMulticastDelegate<void(FRemoteTransactionId, FName)> OnTransactionStartingDelegate;
 TMulticastDelegate<void(FRemoteTransactionId)> OnTransactionCompletedDelegate;
 TMulticastDelegate<void(FRemoteTransactionId)> OnTransactionAbortedDelegate;
+TMulticastDelegate<void(FRemoteTransactionId)> OnTransactionReleasedDelegate;
 }
 
 enum class ERemoteExecutorAbortReason : uint8
@@ -98,6 +100,11 @@ public:
 		NewWork.Work = InWork;
 		NewWork.bIsTransactional = bIsTransactional;
 		NewWork.RequestId = GenerateNextTransactionId();
+
+		if (bIsTransactional)
+		{
+			UE::RemoteExecutor::OnTransactionQueuedDelegate.Broadcast(NewWork.RequestId, NewWork.Name);
+		}
 	}
 
 	void ExecutePendingWork()
@@ -555,6 +562,8 @@ public:
 							Subsystem->DestroyRequest(ExecutingWork->RequestId);
 						}
 
+						UE::RemoteExecutor::OnTransactionReleasedDelegate.Broadcast(ExecutingWork->RequestId);
+
 						// remove this work from the list and do not increment PendingWorkIndex
 						PendingWorks.RemoveAt(PendingWorkIndex, EAllowShrinking::No);
 					}
diff --git a/Engine/Source/Runtime/CoreUObject/Public/UObject/RemoteExecutor.h b/Engine/Source/Runtime/CoreUObject/Public/UObject/RemoteExecutor.h
--- a/Engine/Source/Runtime/CoreUObject/Public/UObject/RemoteExecutor.h
+++ b/Engine/Source/Runtime/CoreUObject/Public/UObject/RemoteExecutor.h
@@ -203,6 +203,9 @@ namespace UE::RemoteExecutor
 	/** Fetch the next deferred RPC request **/
 	extern COREUOBJECT_API TDelegate<TOptional<TTuple<FName, FRemoteWorkPriority, bool, TFunction<void(void)>>>(void)> FetchNextDeferredRPCDelegate;
 
+	/** Delegate executed when a transaction is queued in the executor for execution */
+	extern COREUOBJECT_API TMulticastDelegate<void(FRemoteTransactionId, FName)> OnTransactionQueuedDelegate;
+
 	/** Delegate executed when a ExecuteTransactional starts an AutoRTFM transaction */
 	extern COREUOBJECT_API TMulticastDelegate<void(FRemoteTransactionId, FName)> OnTransactionStartingDelegate;
 
@@ -212,6 +215,9 @@ namespace UE::RemoteExecutor
 	/** Delegate executed when a ExecuteTransactional observes an aborted AutoRTFM transaction */
 	extern COREUOBJECT_API TMulticastDelegate<void(FRemoteTransactionId)> OnTransactionAbortedDelegate;
 
+	/** Delegate executed when a transaction is released (execution is finished and the work is removed from the executor) */
+	extern COREUOBJECT_API TMulticastDelegate<void(FRemoteTransactionId)> OnTransactionReleasedDelegate;
+
 	COREUOBJECT_API void RegisterRemoteSubsystem(FRemoteSubsystemBase* Subsystem);
 
 	COREUOBJECT_API void ExecuteTransactional(FName WorkName, const TFunctionRef<void(void)>& Work);
