Moving some intrinsic classes for CoreUObject tests to be non-intrinsic.

diff --git a/Engine/Source/Programs/Shared/EpicGames.UHT/Utils/UhtPropertyResolveArgs.cs b/Engine/Source/Programs/Shared/EpicGames.UHT/Utils/UhtPropertyResolveArgs.cs
--- a/Engine/Source/Programs/Shared/EpicGames.UHT/Utils/UhtPropertyResolveArgs.cs
+++ b/Engine/Source/Programs/Shared/EpicGames.UHT/Utils/UhtPropertyResolveArgs.cs
@@ -348,6 +348,12 @@ namespace EpicGames.UHT.Utils
 				}
 			}
 
+			// HACK for CoreUObject tests - allow raw pointers in non-intrinsic classes there.
+			if (module.ShortName == "CoreUObject" && PropertySettings.Outer.HeaderFile.ModuleRelativeFilePath.StartsWith("Tests", StringComparison.InvariantCultureIgnoreCase))
+			{
+				behavior = UhtIssueBehavior.AllowSilently;
+			}
+
 			if (behavior == UhtIssueBehavior.AllowSilently)
 			{
 				return;
diff --git a/Engine/Source/Runtime/CoreUObject/Private/UObject/GarbageCollectionTesting.cpp b/Engine/Source/Runtime/CoreUObject/Private/UObject/GarbageCollectionTesting.cpp
--- a/Engine/Source/Runtime/CoreUObject/Private/UObject/GarbageCollectionTesting.cpp
+++ b/Engine/Source/Runtime/CoreUObject/Private/UObject/GarbageCollectionTesting.cpp
@@ -84,8 +84,3 @@ void UnlinkReachabilityStressData(TArray<UObjectReachabilityStressData*>& Data)
 		RootData->RemoveFromRoot();
 	}
 }
-
-IMPLEMENT_CORE_INTRINSIC_CLASS(UObjectReachabilityStressData, UObject,
-	{
-		UE::GC::DeclareIntrinsicMembers(Class, { UE_GC_MEMBER(UObjectReachabilityStressData, Children) });
-	});
diff --git a/Engine/Source/Runtime/CoreUObject/Private/UObject/GarbageCollectionTesting.h b/Engine/Source/Runtime/CoreUObject/Private/UObject/GarbageCollectionTesting.h
--- a/Engine/Source/Runtime/CoreUObject/Private/UObject/GarbageCollectionTesting.h
+++ b/Engine/Source/Runtime/CoreUObject/Private/UObject/GarbageCollectionTesting.h
@@ -7,15 +7,16 @@
 #include "UObject/ObjectMacros.h"
 #include "UObject/UObjectGlobals.h"
 
+#include "GarbageCollectionTesting.generated.h"
+
+UCLASS(Transient)
 class UObjectReachabilityStressData : public UObject
 {
-	DECLARE_CLASS_INTRINSIC(UObjectReachabilityStressData,
-		UObject,
-		CLASS_Transient,
-		TEXT("/Script/CoreUObject"));
+	GENERATED_BODY()
 
 public:
-	TArray<UObjectReachabilityStressData*> Children;
+	UPROPERTY()
+	TArray<TObjectPtr<UObjectReachabilityStressData>> Children;
 };
 
 void GenerateReachabilityStressData(TArray<UObjectReachabilityStressData*>& Data);
diff --git a/Engine/Source/Runtime/CoreUObject/Tests/ObjectPropertyTest.cpp b/Engine/Source/Runtime/CoreUObject/Tests/ObjectPropertyTest.cpp
--- a/Engine/Source/Runtime/CoreUObject/Tests/ObjectPropertyTest.cpp
+++ b/Engine/Source/Runtime/CoreUObject/Tests/ObjectPropertyTest.cpp
@@ -95,6 +95,9 @@ TEST_CASE("UE::CoreUObject::FObjectProperty::CheckValidAddressNonNullable")
 	UClass* Class = UObjectPtrTestClassWithRef::StaticClass();
 	FObjectProperty* Property = CastField<FObjectProperty>(Class->FindPropertyByName(TEXT("ObjectPtrNonNullable")));
 	FObjectProperty* AbstractProperty = CastField<FObjectProperty>(Class->FindPropertyByName(TEXT("ObjectPtrAbstractNonNullable")));
+	// Force non-nullable flag on because UHT does not yet support TNonNullPtr<TObjectPtr<UObject>>
+	Property->SetPropertyFlags(CPF_NonNullable);
+	AbstractProperty->SetPropertyFlags(CPF_NonNullable);
 	REQUIRE(Property != nullptr);
 	REQUIRE(AbstractProperty != nullptr);
 
diff --git a/Engine/Source/Runtime/CoreUObject/Tests/ObjectPtrTest.cpp b/Engine/Source/Runtime/CoreUObject/Tests/ObjectPtrTest.cpp
--- a/Engine/Source/Runtime/CoreUObject/Tests/ObjectPtrTest.cpp
+++ b/Engine/Source/Runtime/CoreUObject/Tests/ObjectPtrTest.cpp
@@ -217,7 +217,7 @@ TEST_CASE_METHOD(FObjectPtrTestBase, "CoreUObject::TObjectPtr::Default Serialize
 
 	TestPackage->RemoveFromRoot();
 }
-#endif
+#endif // UE_WITH_OBJECT_HANDLE_LATE_RESOLVE || UE_WITH_OBJECT_HANDLE_TRACKING
 
 #if UE_WITH_OBJECT_HANDLE_LATE_RESOLVE || UE_WITH_OBJECT_HANDLE_TRACKING
 
@@ -252,7 +252,7 @@ TEST_CASE_METHOD(FObjectPtrTestBase, "CoreUObject::TObjectPtr::Soft Object Path"
 
 	TestPackage->RemoveFromRoot();
 }
-#endif
+#endif // UE_WITH_OBJECT_HANDLE_LATE_RESOLVE || UE_WITH_OBJECT_HANDLE_TRACKING
 
 TEST_CASE_METHOD(FObjectPtrTestBase, "CoreUObject::TObjectPtr::Forward Declared", "[CoreUObject][ObjectPtr]")
 {
@@ -315,7 +315,7 @@ TEST_CASE_METHOD(FObjectPtrTestBase, "CoreUObject::TObjectPtr::Hash Consistency"
 
 	TestPackage1->RemoveFromRoot();
 }
-#endif
+#endif // UE_WITH_OBJECT_HANDLE_LATE_RESOLVE
 
 #if UE_WITH_OBJECT_HANDLE_LATE_RESOLVE
 TEST_CASE_METHOD(FObjectPtrTestBase, "CoreUObject::TObjectPtr::Long Path", "[CoreUObject][ObjectPtr]")
@@ -343,7 +343,7 @@ TEST_CASE_METHOD(FObjectPtrTestBase, "CoreUObject::TObjectPtr::Long Path", "[Cor
 	TEST_EQUAL(TEXT("Resolved path from FObjectPathId should have TestObject3 at element 2"), ResolvedNames[2], TestObject3->GetFName());
 	TEST_EQUAL(TEXT("Resolved path from FObjectPathId should have TestObject4 at element 3"), ResolvedNames[3], TestObject4->GetFName());
 }
-#endif
+#endif // UE_WITH_OBJECT_HANDLE_LATE_RESOLVE
 
 TEST_CASE_METHOD(FObjectPtrTestBase, "CoreUObject::TObjectPtr::Type Specialization", "[CoreUObject][ObjectPtr]")
 {
@@ -745,7 +745,7 @@ TEST_CASE("CoreUObject::TObjectPtr::SwapArray")
 	CHECK(ArrayRaw[0] == RawPtrA);
 	CHECK(ArrayPtr[0] == RawPtrB);
 }
-#endif
+#endif // !UE_DEPRECATE_MUTABLE_TOBJECTPTR
 
 TEST_CASE("CoreUObject::TObjectPtr::Move")
 {
@@ -948,7 +948,7 @@ TEST_CASE("CoreUObject::TObjectPtr::GetTypeHash")
 }
 
 
-#endif
+#endif // UE_WITH_OBJECT_HANDLE_LATE_RESOLVE
 
 template<typename TObj>
 void TestArrayConversion()
@@ -1368,7 +1368,7 @@ TEST_CASE("CoreUObject::TObjectPtr::GetOuter")
 // }
 
 
-#endif
+#endif // WITH_LOW_LEVEL_TESTS
 
 class UForwardDeclaredObjDerived: public UObject {};
 class FForwardDeclaredNotObjDerived {};
@@ -1377,50 +1377,6 @@ class FForwardDeclaredNotObjDerived {};
 
 #if WITH_LOW_LEVEL_TESTS
 
-//bunch of class to reproduce multiple inheritance issues.
-//need to live outside of the namespace be IMPLEMENT_CORE_INTRINSIC_CLASS don't like namespaces
-
-class FTestBaseClass
-{
-public:
-	virtual ~FTestBaseClass() = default;
-	virtual void VirtFunc() { };
-};
-
-class UMiddleClass : public UObject, public FTestBaseClass
-{
-	DECLARE_CLASS_INTRINSIC(UMiddleClass, UObject, CLASS_MatchedSerializers, TEXT("/Script/CoreUObject"))
-
-public:
-	virtual void VirtFunc() override { };
-};
-
-
-class FAnotherBaseClass
-{
-public:
-	virtual ~FAnotherBaseClass() = default;
-	virtual void AnotherVirtFunc() { };
-};
-
-
-class UDerrivedClass : public UMiddleClass, public FAnotherBaseClass
-{
-	DECLARE_CLASS_INTRINSIC(UDerrivedClass, UMiddleClass, CLASS_MatchedSerializers, TEXT("/Script/CoreUObject"))
-public:
-	virtual void AnotherVirtFunc() override { };
-};
-
-IMPLEMENT_CORE_INTRINSIC_CLASS(UMiddleClass, UObject,
-	{
-	}
-);
-
-IMPLEMENT_CORE_INTRINSIC_CLASS(UDerrivedClass, UMiddleClass,
-	{
-	}
-);
-
 TEST_CASE("CoreUObject::TObjectPtr::TestEquals")
 {
 	const FName TestPackageName(TEXT("/Engine/Test/TestEquals/Transient"));
diff --git a/Engine/Source/Runtime/CoreUObject/Tests/ObjectPtrTestClass.cpp b/Engine/Source/Runtime/CoreUObject/Tests/ObjectPtrTestClass.cpp
deleted file mode 100644
diff --git a/Engine/Source/Runtime/CoreUObject/Tests/ObjectPtrTestClass.h b/Engine/Source/Runtime/CoreUObject/Tests/ObjectPtrTestClass.h
--- a/Engine/Source/Runtime/CoreUObject/Tests/ObjectPtrTestClass.h
+++ b/Engine/Source/Runtime/CoreUObject/Tests/ObjectPtrTestClass.h
@@ -1,96 +1,145 @@
 // Copyright Epic Games, Inc. All Rights Reserved.
 
 #pragma once
-#if WITH_LOW_LEVEL_TESTS
+
+#if WITH_TESTS
 
 #include "UObject/Object.h"
 #include "UObject/ObjectPtr.h"
+#include "Templates/NonNullPointer.h"
 #include "Templates/SubclassOf.h"
 
+#include "ObjectPtrTestClass.generated.h"
+
 //simple test class for testing TObjectPtr resolve behavior
+UCLASS(meta=(LoadBehavior=LazyOnDemand))
 class UObjectPtrTestClass : public UObject
 {
-	DECLARE_CLASS_INTRINSIC(UObjectPtrTestClass, UObject, CLASS_MatchedSerializers, TEXT("/Script/CoreUObject"))
-
-public:
-
+	GENERATED_BODY()
 };
 
 //abstract test class for testing TObjectPtr resolve behavior
+UCLASS(Abstract, meta=(LoadBehavior=LazyOnDemand))
 class UObjectPtrAbstractTestClass : public UObject
 {
-	DECLARE_CLASS_INTRINSIC(UObjectPtrAbstractTestClass, UObject, CLASS_MatchedSerializers | CLASS_Abstract, TEXT("/Script/CoreUObject"))
+	GENERATED_BODY()
 };
 
 //derived-from-abstract test class for testing TObjectPtr resolve behavior
+UCLASS(meta=(LoadBehavior=LazyOnDemand))
 class UObjectPtrAbstractDerivedTestClass : public UObjectPtrAbstractTestClass
 {
-	DECLARE_CLASS_INTRINSIC(UObjectPtrAbstractDerivedTestClass, UObjectPtrAbstractTestClass, CLASS_MatchedSerializers, TEXT("/Script/CoreUObject"))
+	GENERATED_BODY()
 };
 
 //test class with typed reference to another class
+UCLASS(meta=(LoadBehavior=LazyOnDemand))
 class UObjectPtrTestClassWithRef : public UObject
 {
-	DECLARE_CLASS_INTRINSIC(UObjectPtrTestClassWithRef, UObject, CLASS_MatchedSerializers, TEXT("/Script/CoreUObject"))
+	GENERATED_BODY()
 
 public:
+	UPROPERTY()
 	TObjectPtr<UObjectPtrTestClass> ObjectPtr;
+
+	UPROPERTY()
 	TObjectPtr<UObjectPtrTestClass> ObjectPtrNonNullable;
+
+	UPROPERTY()
 	TObjectPtr<UObjectPtrAbstractTestClass> ObjectPtrAbstractNonNullable;
+
+	UPROPERTY()
 	TArray<TObjectPtr<UObjectPtrTestClass>> ArrayObjPtr;
 };
 
 
 //test class with typed reference to another class
+UCLASS(meta=(LoadBehavior=LazyOnDemand))
 class UObjectWithClassProperty : public UObject
 {
-	DECLARE_CLASS_INTRINSIC(UObjectWithClassProperty, UObject, CLASS_MatchedSerializers, TEXT("/Script/CoreUObject"))
+	GENERATED_BODY()
 
 public:
+	UPROPERTY()
 	TObjectPtr<UClass> ClassPtr;
+	UPROPERTY()
 	TSubclassOf<UObjectPtrTestClass> SubClass;
+	UPROPERTY()
 	UClass* ClassRaw;
 };
 
 //test class with raw pointer
+UCLASS(meta=(LoadBehavior=LazyOnDemand))
 class UObjectWithRawProperty : public UObject
 {
-	DECLARE_CLASS_INTRINSIC(UObjectWithRawProperty, UObject, CLASS_MatchedSerializers, TEXT("/Script/CoreUObject"))
+	GENERATED_BODY()
 
 public:
+	UPROPERTY()
 	UObjectPtrTestClass* ObjectPtr;
+	UPROPERTY()
 	UObjectPtrTestClass* ObjectPtrNonNullable;
 };
 
 
 //derived test class
+UCLASS(meta=(LoadBehavior=LazyOnDemand))
 class UObjectPtrDerrivedTestClass : public UObjectPtrTestClass
 {
-	DECLARE_CLASS_INTRINSIC(UObjectPtrDerrivedTestClass, UObjectPtrTestClass, CLASS_MatchedSerializers, TEXT("/Script/CoreUObject"))
-
-public:
-
+	GENERATED_BODY()
 };
 
 
 //non lazy test class
+UCLASS()
 class UObjectPtrNotLazyTestClass : public UObject
 {
-	DECLARE_CLASS_INTRINSIC(UObjectPtrNotLazyTestClass, UObject, CLASS_MatchedSerializers, TEXT("/Script/CoreUObject"))
-
-public:
-
+	GENERATED_BODY()
 };
 
 
 //stress testing class
+UCLASS(meta=(LoadBehavior=LazyOnDemand))
 class UObjectPtrStressTestClass : public UObject
 {
-	DECLARE_CLASS_INTRINSIC(UObjectPtrStressTestClass, UObject, CLASS_MatchedSerializers, TEXT("/Script/CoreUObject"))
+	GENERATED_BODY()
 
 public:
 	uint8 Data[PLATFORM_CACHE_LINE_SIZE];
 };
 
+class FTestBaseClass
+{
+public:
+	virtual ~FTestBaseClass() = default;
+	virtual void VirtFunc() { };
+};
+
+UCLASS()
+class UMiddleClass : public UObject, public FTestBaseClass
+{
+	GENERATED_BODY()
+
+public:
+	virtual void VirtFunc() override { };
+};
+
+
+class FAnotherBaseClass
+{
+public:
+	virtual ~FAnotherBaseClass() = default;
+	virtual void AnotherVirtFunc() { };
+};
+
+
+UCLASS()
+class UDerrivedClass : public UMiddleClass, public FAnotherBaseClass
+{
+	GENERATED_BODY()
+public:
+	virtual void AnotherVirtFunc() override { };
+};
+
 
 #endif
\ No newline at end of file
