[CrashFix] IDO archetypes (InstanceDataObjectArchetypeChain) was causing crashes in serialization. This is because it caused array sizes to mismatch which overriadable serailization is picky about. To solve this, I changed how the templates are created/selected so that they can properly be passed in to the constructor as a template rather than relying on the copy hack from before

diff --git a/Engine/Source/Runtime/CoreUObject/Private/UObject/InstanceDataObjectUtils.cpp b/Engine/Source/Runtime/CoreUObject/Private/UObject/InstanceDataObjectUtils.cpp
--- a/Engine/Source/Runtime/CoreUObject/Private/UObject/InstanceDataObjectUtils.cpp
+++ b/Engine/Source/Runtime/CoreUObject/Private/UObject/InstanceDataObjectUtils.cpp
@@ -1487,6 +1487,12 @@ bool IsInstanceDataObject(const UObject* Object)
 
 UObject* CreateInstanceDataObject(UObject* Owner)
 {
+	// If an Ido already exists, skip the uneeded serialization and just return it
+	if (UObject* Found = FPropertyBagRepository::Get().FindInstanceDataObject(Owner))
+	{
+		return Found;
+	}
+
 	TArray<uint8> OwnerData;
 
 	FObjectWriter Writer(OwnerData);
diff --git a/Engine/Source/Runtime/CoreUObject/Private/UObject/PropertyBagRepository.cpp b/Engine/Source/Runtime/CoreUObject/Private/UObject/PropertyBagRepository.cpp
--- a/Engine/Source/Runtime/CoreUObject/Private/UObject/PropertyBagRepository.cpp
+++ b/Engine/Source/Runtime/CoreUObject/Private/UObject/PropertyBagRepository.cpp
@@ -837,17 +837,31 @@ UObject* FPropertyBagRepository::CreateInstanceDataObject(UObject* Owner, FArchi
 	UObject* Template = nullptr;
 	if (IsInstanceDataObjectArchetypeChainEnabled())
 	{
-		UObject* OwnerArchetype = Owner->GetArchetype();
-
-		// If an archetype is not the CDO then it's a sub-object archetype.
-		bool bIsSubobjectArchetype = OwnerArchetype && !OwnerArchetype->HasAllFlags(RF_ClassDefaultObject);
-
-		// Subobject archetypes (under a CDO) might have initialized properties that their own CDO does
-		// not initialize (e.g., Verse overrides). For this reason we create those IDO archetypes on demand.
-		if (bIsSubobjectArchetype)
+		if (UObject* OwnerArchetype = Owner->GetArchetype())
 		{
-			Template = UE::CreateInstanceDataObject(OwnerArchetype);
-			check(Template);
+			if (OwnerArchetype == Owner->GetClass()->GetDefaultObject(false))
+			{
+				// if the archetype is a CDO, we can simply use the IDO's CDO because it's data should match.
+				Template = InstanceDataObjectClass->GetDefaultObject(false);
+			}
+			else
+			{
+				// attempt to find/create an IDO for the archetype. Unfortunately if there's unkown properties involved, the types may not match.
+				UObject* ArchetypeIdo = UE::CreateInstanceDataObject(OwnerArchetype);
+				if (ArchetypeIdo && ArchetypeIdo->GetClass() && ArchetypeIdo->IsA(InstanceDataObjectClass))
+				{
+					Template = ArchetypeIdo;
+				}
+				else if (OwnerArchetype->GetArchetype())
+				{
+					// we need a archetype that is the same type as the IDO but has the same data as the owner's archetype
+					FStaticConstructObjectParameters TemplateParams(InstanceDataObjectClass);
+					TemplateParams.SetFlags |= EObjectFlags::RF_ArchetypeObject;
+					TemplateParams.Outer = GetTransientPackage();
+					Template = StaticConstructObject_Internal(TemplateParams);
+					CopyTaggedProperties(OwnerArchetype, Template);
+				}
+			}
 		}
 	}
 
@@ -862,17 +876,10 @@ UObject* FPropertyBagRepository::CreateInstanceDataObject(UObject* Owner, FArchi
 
 	Params.Name = MakeUniqueObjectName(Outer, Params.Class, Owner->GetFName(), EUniqueObjectNameOptions::UniversallyUnique);
 	Params.Outer = Outer;
+	Params.Template = Template;
 
 	UObject* InstanceDataObject = StaticConstructObject_Internal(Params);
 
-	// We are not passing the template in FStaticConstructObjectParameters because of potential type mismatches
-	// between the IDO and its template. This can happen when both objects have the same 'owner' type but a
-	// different number of unknown properties. For this reason we do a copy instead.
-	if (Template)
-	{
-		CopyTaggedProperties(Template, InstanceDataObject);
-	}
-
 	FPropertyBagRepositoryLock LockRepo(this);
 	InstanceDataObjectToOwner.Add(InstanceDataObject, Owner);
 
