Skipping property initialization and serialization tracking during SubObjectsShadowing.
We want to skip tracking during SubObjectsShadowing because during that phase we are not properly initializing or serializing properties.

#rnx

diff --git a/Engine/Source/Runtime/CoreUObject/Private/UObject/Class.cpp b/Engine/Source/Runtime/CoreUObject/Private/UObject/Class.cpp
--- a/Engine/Source/Runtime/CoreUObject/Private/UObject/Class.cpp
+++ b/Engine/Source/Runtime/CoreUObject/Private/UObject/Class.cpp
@@ -1399,7 +1399,7 @@ void UStruct::LoadTaggedPropertiesFromText(FStructuredArchive::FSlot Slot, uint8
 				if (Tag.SerializeType == EPropertyTagSerializeType::Skipped)
 				{
 				#if WITH_EDITORONLY_DATA
-					if (SerializeContext->bTrackInitializedProperties)
+					if (SerializeContext->bTrackInitializedProperties && (Tag.OverrideOperation != EOverriddenPropertyOperation::SubObjectsShadowing))
 					{
 						UE::FInitializedPropertyValueState(this, Data).Set(Property, ItemIndex);
 					}
@@ -1458,13 +1458,16 @@ void UStruct::LoadTaggedPropertiesFromText(FStructuredArchive::FSlot Slot, uint8
 			#if WITH_EDITORONLY_DATA
 				if (bSerialized)
 				{
-					if (SerializeContext->bTrackInitializedProperties)
+					if ((Tag.OverrideOperation != EOverriddenPropertyOperation::SubObjectsShadowing))
 					{
-						UE::FInitializedPropertyValueState(this, Data).Set(Property, ItemIndex);
-					}
-					if (SerializeContext->bTrackSerializedProperties)
-					{
-						UE::FSerializedPropertyValueState(this, Data).Set(Property, ItemIndex);
+						if (SerializeContext->bTrackInitializedProperties)
+						{
+							UE::FInitializedPropertyValueState(this, Data).Set(Property, ItemIndex);
+						}
+						if (SerializeContext->bTrackSerializedProperties)
+						{
+							UE::FSerializedPropertyValueState(this, Data).Set(Property, ItemIndex);
+						}
 					}
 				}
 			#endif
@@ -1737,7 +1740,7 @@ void UStruct::SerializeVersionedTaggedProperties(FStructuredArchive::FSlot Slot,
 				if (Tag.SerializeType == EPropertyTagSerializeType::Skipped)
 				{
 				#if WITH_EDITORONLY_DATA
-					if (SerializeContext->bTrackInitializedProperties && Property)
+					if (SerializeContext->bTrackInitializedProperties && Property && (Tag.OverrideOperation != EOverriddenPropertyOperation::SubObjectsShadowing))
 					{
 						UE::FInitializedPropertyValueState(this, Data).Set(Property, Tag.ArrayIndex);
 					}
@@ -1937,13 +1940,16 @@ void UStruct::SerializeVersionedTaggedProperties(FStructuredArchive::FSlot Slot,
 					}
 
 				#if WITH_EDITORONLY_DATA
-					if (SerializeContext->bTrackInitializedProperties && Property)
+					if ((Tag.OverrideOperation != EOverriddenPropertyOperation::SubObjectsShadowing))
 					{
-						UE::FInitializedPropertyValueState(this, Data).Set(Property, Tag.ArrayIndex);
-					}
-					if (SerializeContext->bTrackSerializedProperties && Property)
-					{
-						UE::FSerializedPropertyValueState(this, Data).Set(Property, Tag.ArrayIndex);
+						if (SerializeContext->bTrackInitializedProperties && Property)
+						{
+							UE::FInitializedPropertyValueState(this, Data).Set(Property, Tag.ArrayIndex);
+						}
+						if (SerializeContext->bTrackSerializedProperties && Property)
+						{
+							UE::FSerializedPropertyValueState(this, Data).Set(Property, Tag.ArrayIndex);
+						}
 					}
 				#endif
 				}
