Don't include the Reserved bit in FRemoteObjectId comparisons

diff --git a/Engine/Source/Runtime/CoreUObject/Public/UObject/RemoteObjectTypes.h b/Engine/Source/Runtime/CoreUObject/Public/UObject/RemoteObjectTypes.h
--- a/Engine/Source/Runtime/CoreUObject/Public/UObject/RemoteObjectTypes.h
+++ b/Engine/Source/Runtime/CoreUObject/Public/UObject/RemoteObjectTypes.h
@@ -120,6 +120,11 @@ private:
 		uint64 Id = 0;
 	};
 
+	FORCEINLINE uint64 GetComparableValue() const
+	{
+		return uint64(SerialNumber) | (uint64(ServerId) << REMOTE_OBJECT_SERIAL_NUMBER_BIT_SIZE);
+	}
+
 public:
 
 	FRemoteObjectId() = default;
@@ -132,7 +137,8 @@ public:
 	}
 
 	FRemoteObjectId(FRemoteServerId InServerId, uint64 InSerialNumber)
-		: SerialNumber(InSerialNumber)
+		: Reserved(0)
+		, SerialNumber(InSerialNumber)
 		, ServerId(InServerId.Id)
 	{
 		// Make sure the arguments will fit into the different bitfields.
@@ -144,32 +150,34 @@ public:
 
 	FORCEINLINE uint32 GetTypeHash() const
 	{
-		return ::GetTypeHash(Id);
+		return ::GetTypeHash(GetComparableValue());
 	}
 
 	FORCEINLINE bool operator==(const FRemoteObjectId& Other) const
 	{
-		return Id == Other.Id;
+		// Don't compare against the Reserved bit
+		return GetComparableValue() == Other.GetComparableValue();
 	}
 	FORCEINLINE bool operator!=(const FRemoteObjectId& Other) const
 	{
-		return Id != Other.Id;
+		// Don't compare against the Reserved bit
+		return GetComparableValue() != Other.GetComparableValue();
 	}
 	FORCEINLINE bool operator<(const FRemoteObjectId& Other) const
 	{
-		return Id < Other.Id;
+		return GetComparableValue() < Other.GetComparableValue();
 	}
 	FORCEINLINE bool operator<=(const FRemoteObjectId& Other) const
 	{
-		return Id <= Other.Id;
+		return GetComparableValue() <= Other.GetComparableValue();
 	}
 	FORCEINLINE bool operator>(const FRemoteObjectId& Other) const
 	{
-		return Id > Other.Id;
+		return GetComparableValue() > Other.GetComparableValue();
 	}
 	FORCEINLINE bool operator>=(const FRemoteObjectId& Other) const
 	{
-		return Id >= Other.Id;
+		return GetComparableValue() >= Other.GetComparableValue();
 	}
 	FORCEINLINE bool IsValid() const
 	{
