[AsyncLoading]

Improve the error message when we can't find script imports, so we return the expected name from the nametable rather than only the import index

#rnx

diff --git a/Engine/Source/Runtime/CoreUObject/Private/Serialization/AsyncLoading2.cpp b/Engine/Source/Runtime/CoreUObject/Private/Serialization/AsyncLoading2.cpp
--- a/Engine/Source/Runtime/CoreUObject/Private/Serialization/AsyncLoading2.cpp
+++ b/Engine/Source/Runtime/CoreUObject/Private/Serialization/AsyncLoading2.cpp
@@ -7642,8 +7642,9 @@ UObject* FAsyncPackage2::EventDrivenIndexToObject(const FAsyncPackageHeaderData&
 	else if (Index.IsImport())
 	{
 		Result = ImportStore.FindOrGetImportObject(Header, Index);
-		UE_CLOG(!Result, LogStreaming, Warning, TEXT("Missing Dependency, missing %s import 0x%llX for package %s"),
+		UE_CLOG(!Result, LogStreaming, Warning, TEXT("Missing Dependency, missing %s import %s(0x%llX) for package %s"),
 			Index.IsScriptImport() ? TEXT("script") : TEXT("package"),
+			*GetFNameFromPackageObjectIndex(Index).ToString(),
 			Index.Value(),
 			*Desc.PackagePathToLoad.GetPackageFName().ToString());
 	}
