Improved logging for async loading assertion in FAsyncPackage2::EventDrivenCreateExport

-this improvement was added during the course of debugging the check in question being hit.

diff --git a/Engine/Source/Runtime/CoreUObject/Private/Serialization/AsyncLoading2.cpp b/Engine/Source/Runtime/CoreUObject/Private/Serialization/AsyncLoading2.cpp
--- a/Engine/Source/Runtime/CoreUObject/Private/Serialization/AsyncLoading2.cpp
+++ b/Engine/Source/Runtime/CoreUObject/Private/Serialization/AsyncLoading2.cpp
@@ -8132,7 +8132,10 @@ void FAsyncPackage2::EventDrivenCreateExport(const FAsyncPackageHeaderData& Head
 			}
 			else
 			{
-				check(ExportObject.TemplateObject->IsA(LoadClass));
+				checkf(ExportObject.TemplateObject->IsA(LoadClass), 
+					TEXT("ExportObject.TemplateObject: '%s' is not a child of class: '%s'"), 
+					*GetNameSafe(ExportObject.TemplateObject),
+					*GetNameSafe(LoadClass));
 			}
 		}
 #endif
