Fix comments about FGCScopeLock which should have been FGCScopeGuard.

diff --git a/Engine/Source/Runtime/CoreUObject/Private/Serialization/AsyncLoading.cpp b/Engine/Source/Runtime/CoreUObject/Private/Serialization/AsyncLoading.cpp
--- a/Engine/Source/Runtime/CoreUObject/Private/Serialization/AsyncLoading.cpp
+++ b/Engine/Source/Runtime/CoreUObject/Private/Serialization/AsyncLoading.cpp
@@ -501,7 +501,7 @@ static FName StaticGetNativeClassName(UClass* InClass)
 	return InClass ? InClass->GetFName() : NAME_None;
 }
 
-/** Returns true if we're inside a FGCScopeLock */
+/** Returns true if we're inside a FGCScopeGuard */
 bool IsGarbageCollectionLocked();
 
 #if UE_ENABLE_TRACKED_IO
diff --git a/Engine/Source/Runtime/CoreUObject/Public/Serialization/AsyncPackageLoader.h b/Engine/Source/Runtime/CoreUObject/Public/Serialization/AsyncPackageLoader.h
--- a/Engine/Source/Runtime/CoreUObject/Public/Serialization/AsyncPackageLoader.h
+++ b/Engine/Source/Runtime/CoreUObject/Public/Serialization/AsyncPackageLoader.h
@@ -39,7 +39,7 @@ enum class ENotifyRegistrationPhase;
 
 extern const FName PrestreamPackageClassNameLoad;
 
-/** Returns true if we're inside a FGCScopeLock */
+/** Returns true if we're inside a FGCScopeGuard */
 extern bool IsGarbageCollectionLocked();
 
 bool IsNativeCodePackage(UPackage* Package);
