[GCScopeLock]
- Further simplification and optimization of the GC Lock
- About 1.33x faster in microbenchmark

diff --git a/Engine/Source/Runtime/CoreUObject/Private/UObject/GCScopeLock.h b/Engine/Source/Runtime/CoreUObject/Private/UObject/GCScopeLock.h
--- a/Engine/Source/Runtime/CoreUObject/Private/UObject/GCScopeLock.h
+++ b/Engine/Source/Runtime/CoreUObject/Private/UObject/GCScopeLock.h
@@ -7,15 +7,14 @@
 #pragma once
 
 #include "CoreMinimal.h"
-#include "HAL/ThreadSafeCounter.h"
 #include "HAL/PlatformProcess.h"
 #include "Misc/ScopeLock.h"
-#include "HAL/ThreadSafeBool.h"
+#include "Async/ManualResetEvent.h"
 #include "Async/SharedMutex.h"
 #include "Async/SharedLock.h"
 #include "Async/UniqueLock.h"
-#include "HAL/Event.h"
 
+#include <atomic>
 
 /**
 * Garbage Collection synchronization objects
@@ -24,23 +23,18 @@
 */
 class FGCCSyncObject
 {
-	/** Non zero if any of the non-game threads is blocking GC */
-	FThreadSafeCounter AsyncCounter;
-	/** Non zero if GC is running */
-	FThreadSafeCounter GCCounter;
-	/** Non zero if GC wants to run but is blocked by some other thread \
-	    This flag is not automatically enforced on the async threads, instead
-			threads have to manually implement support for it. */
-	TAtomic<int32> GCWantsToRunCounter {0};
 	/** Shared mutex for thread safe operations */
 	UE::FSharedMutex SharedMutex;
 	/** Event used to block non-game threads when GC is running */
-	FEvent* GCUnlockedEvent;
+	UE::FManualResetEvent GCUnlockedEvent;
+	/** True if GC wants to run but is blocked by some other thread \
+	    This flag is not automatically enforced on the async threads, instead
+			threads have to manually implement support for it. */
+	std::atomic<bool> GCWantsToRunCounter = false;
 
-public:
 
-	FGCCSyncObject();
-	~FGCCSyncObject();
+public:
+	FGCCSyncObject() = default;
 
 	/** Creates the singleton object */
 	static void Create();
@@ -48,53 +42,45 @@ public:
 	/** Gets the singleton object */
 	static FGCCSyncObject& Get();
 
-	/** Lock on non-game thread. Will block if GC is running. */
-	void LockAsync()
+	/** Lock on non-game thread. Will block if GC is running.
+	    Returns true if a matching call to UnlockAsyncUnconditional is be required.
+	*/
+	bool LockAsync()
 	{
 		if (!IsInGameThread())
 		{
 			// Wait until GC is done if it was running when entering this function
-			bool bLocked = false;
-			do
+			while (true)
 			{
-				if (GCCounter.GetValue() > 0)
+				if (SharedMutex.IsLocked())
 				{
-					GCUnlockedEvent->Wait();
+					// Using an event instead of just doing LockShared() because we want oversubscription
+					GCUnlockedEvent.Wait();
 				}
+				if (SharedMutex.TryLockShared())
 				{
-					// We only need a read lock here since AsyncCounter is already thread-safe
-					// with other threads that want to lock but we need to guard against
-					// the GC trying to take the write lock.
-					UE::TSharedLock SharedLock(SharedMutex);
-					if (GCCounter.GetValue() == 0)
-					{
-						AsyncCounter.Increment();
-						bLocked = true;
-					}
+					return true;
 				}
-			} while (!bLocked);
+			}
 		}
+
+		return false;
 	}
 
 	/** Attempt async Lock **/
 	bool TryLockAsync()
 	{
-		bool bLocked = true;
 		if (!IsInGameThread())
 		{
-			bLocked = false;
-			if (GCCounter.GetValue() == 0)
+			if (SharedMutex.TryLockShared())
 			{
-				UE::TSharedLock SharedLock(SharedMutex);
-				if (GCCounter.GetValue() == 0)
-				{
-					AsyncCounter.Increment();
-					bLocked = true;
-				}
+				return true;
 			}
+
+			return false;
 		}
 
-		return bLocked;
+		return true;
 	}
 
 	/** Release lock from non-game thread */
@@ -102,91 +88,72 @@ public:
 	{
 		if (!IsInGameThread())
 		{
-			AsyncCounter.Decrement();
+			SharedMutex.UnlockShared();
 		}
 	}
+
+	/** Release lock unconditionally when you know you already have a lock */
+	void UnlockAsyncUnconditional()
+	{
+		SharedMutex.UnlockShared();
+	}
+
 	/** Lock for GC. Will block if any other thread has locked. */
 	void GCLock()
 	{
 		// Signal other threads that GC wants to run
 		SetGCIsWaiting();
 
-		// Wait until all other threads are done if they're currently holding the lock
-		bool bLocked = false;
-		do
-		{
-			FPlatformProcess::ConditionalSleep([&]()
-			{
-				return AsyncCounter.GetValue() == 0;
-			});
-			{
-				// Guard against any reader locks owned by other threads.
-				UE::TUniqueLock ExclusiveLock(SharedMutex);
-				if (AsyncCounter.GetValue() == 0)
-				{
-					GCUnlockedEvent->Reset();
-					int32 GCCounterValue = GCCounter.Increment();
-					check(GCCounterValue == 1); // GCLock doesn't support recursive locks
-					// At this point GC can run so remove the signal that it's waiting
-					FPlatformMisc::MemoryBarrier();
-					ResetGCIsWaiting();
-					bLocked = true;
-				}
-			}
-		} while (!bLocked);
+		SharedMutex.Lock();
+		GCUnlockedEvent.Reset();
+		// At this point GC can run so remove the signal that it's waiting
+		FPlatformMisc::MemoryBarrier();
+		ResetGCIsWaiting();
 	}
 	/** Checks if any async thread has a lock */
 	bool IsAsyncLocked() const
 	{
-		return AsyncCounter.GetValue() != 0;
+		return SharedMutex.IsLockShared();
 	}
 	/** Checks if GC has a lock */
 	bool IsGCLocked() const
 	{
-		return GCCounter.GetValue() != 0;
+		return SharedMutex.IsLocked();
 	}
 	/** Lock for GC. Will not block and return false if any other thread has already locked. */
 	bool TryGCLock()
 	{
-		bool bSuccess = false;
-		if (AsyncCounter.GetValue() != 0)
-		{
-			return bSuccess;
-		}
-		// Guard against any reader locks owned by other threads.
-		UE::TUniqueLock ExclusiveLock(SharedMutex);
-		// If any other thread is currently locking we just exit
-		if (AsyncCounter.GetValue() == 0)
+		if (SharedMutex.TryLock())
 		{
-			GCUnlockedEvent->Reset();
-			int32 GCCounterValue = GCCounter.Increment();
-			check(GCCounterValue == 1); // GCLock doesn't support recursive locks
-			bSuccess = true;
+			GCUnlockedEvent.Reset();
+			return true;
 		}
-		return bSuccess;
+
+		return false;
 	}
+
 	/** Unlock GC */
 	void GCUnlock()
 	{
-		GCUnlockedEvent->Trigger();
-		GCCounter.Decrement();
+		SharedMutex.Unlock();
+		GCUnlockedEvent.Notify();
 	}
 
 	/** Manually mark GC state as 'waiting to run' */
 	void SetGCIsWaiting()
 	{
-		GCWantsToRunCounter++;
+		GCWantsToRunCounter.store(true, std::memory_order_relaxed);
 	}
 
 	/** Manually reset GC 'waiting to run' state*/
 	void ResetGCIsWaiting()
 	{
-		GCWantsToRunCounter.Store(0);
+		GCWantsToRunCounter.store(false, std::memory_order_relaxed);
 	}
 
 	/** True if GC wants to run on the game thread but is maybe blocked by some other thread */
 	FORCEINLINE bool IsGCWaiting() const
 	{
-		return !!GCWantsToRunCounter.Load(EMemoryOrder::Relaxed);
+		return GCWantsToRunCounter.load(std::memory_order_relaxed);
 	}
 };
diff --git a/Engine/Source/Runtime/CoreUObject/Private/UObject/GarbageCollection.cpp b/Engine/Source/Runtime/CoreUObject/Private/UObject/GarbageCollection.cpp
--- a/Engine/Source/Runtime/CoreUObject/Private/UObject/GarbageCollection.cpp
+++ b/Engine/Source/Runtime/CoreUObject/Private/UObject/GarbageCollection.cpp
@@ -150,16 +150,6 @@ public:
 	}
 };
 
-FGCCSyncObject::FGCCSyncObject()
-{
-	GCUnlockedEvent = FPlatformProcess::GetSynchEventFromPool(true);
-}
-FGCCSyncObject::~FGCCSyncObject()
-{
-	FPlatformProcess::ReturnSynchEventToPool(GCUnlockedEvent);
-	GCUnlockedEvent = nullptr;
-}
-
 FGCCSyncObject* GGCSingleton;
 
 void FGCCSyncObject::Create()
@@ -184,12 +174,15 @@ FGCCSyncObject& FGCCSyncObject::Get()
 #define UE_LOG_FGCScopeGuard_LockAsync_Time 0
 
 FGCScopeGuard::FGCScopeGuard()
+#if !UE_LOG_FGCScopeGuard_LockAsync_Time
+	: bNeedUnlock(FGCCSyncObject::Get().LockAsync())
+#endif
 {
 #if UE_LOG_FGCScopeGuard_LockAsync_Time
 	const double StartTime = FPlatformTime::Seconds();
-#endif
-	FGCCSyncObject::Get().LockAsync();
-#if UE_LOG_FGCScopeGuard_LockAsync_Time
+
+	bNeedUnlock = FGCCSyncObject::Get().LockAsync();
+
 	const double ElapsedTime = FPlatformTime::Seconds() - StartTime;
 	if (FPlatformProperties::RequiresCookedData() && ElapsedTime > 0.001)
 	{
@@ -201,7 +194,10 @@ FGCScopeGuard::FGCScopeGuard()
 
 FGCScopeGuard::~FGCScopeGuard()
 {
-	FGCCSyncObject::Get().UnlockAsync();
+	if (bNeedUnlock)
+	{
+		FGCCSyncObject::Get().UnlockAsyncUnconditional();
+	}
 }
 
 FGCScopeTryGuard::FGCScopeTryGuard()
diff --git a/Engine/Source/Runtime/CoreUObject/Public/UObject/GarbageCollection.h b/Engine/Source/Runtime/CoreUObject/Public/UObject/GarbageCollection.h
--- a/Engine/Source/Runtime/CoreUObject/Public/UObject/GarbageCollection.h
+++ b/Engine/Source/Runtime/CoreUObject/Public/UObject/GarbageCollection.h
@@ -111,6 +111,7 @@ COREUOBJECT_API void MarkNativeStructAsReachable(::Verse::VNativeStruct* NativeS
 /** Prevent GC from running in the current scope */
 class FGCScopeGuard
 {
+	bool bNeedUnlock;
 public:
 	COREUOBJECT_API FGCScopeGuard();
 	COREUOBJECT_API ~FGCScopeGuard();
