Rehash map and set properties after instancing key subobjects.

diff --git a/Engine/Source/Runtime/CoreUObject/Private/UObject/PropertyMap.cpp b/Engine/Source/Runtime/CoreUObject/Private/UObject/PropertyMap.cpp
--- a/Engine/Source/Runtime/CoreUObject/Private/UObject/PropertyMap.cpp
+++ b/Engine/Source/Runtime/CoreUObject/Private/UObject/PropertyMap.cpp
@@ -1397,6 +1397,12 @@ void FMapProperty::InstanceSubobjects(void* Data, void const* DefaultData, TNotN
 			}
 		}
 	}
+
+	// Instancing keys will likely have invalidated the hash, so rehash
+	if (bInstancedKey)
+	{
+		MapHelper.Rehash();
+	}
 }
 
 bool FMapProperty::SameType(const FProperty* Other) const
diff --git a/Engine/Source/Runtime/CoreUObject/Private/UObject/PropertySet.cpp b/Engine/Source/Runtime/CoreUObject/Private/UObject/PropertySet.cpp
--- a/Engine/Source/Runtime/CoreUObject/Private/UObject/PropertySet.cpp
+++ b/Engine/Source/Runtime/CoreUObject/Private/UObject/PropertySet.cpp
@@ -838,6 +838,10 @@ void FSetProperty::InstanceSubobjects(void* Data, void const* DefaultData, TNotN
 			ElementProp->InstanceSubobjects(ElementPtr, nullptr, InOwner, InstanceGraph);
 		}
 	}
+
+
+	// Instancing keys will likely have invalidated the hash, so rehash
+	SetHelper.Rehash();
 }
 
 bool FSetProperty::SameType(const FProperty* Other) const
