InstancedStruct: Fix operator= when copying from a TConstStructView

Previously was comparing incompatible types (&TInstancedStruct<T> != &TConstStructView<T>), now correctly compares views

diff --git a/Engine/Source/Runtime/CoreUObject/Public/StructUtils/InstancedStruct.h b/Engine/Source/Runtime/CoreUObject/Public/StructUtils/InstancedStruct.h
--- a/Engine/Source/Runtime/CoreUObject/Public/StructUtils/InstancedStruct.h
+++ b/Engine/Source/Runtime/CoreUObject/Public/StructUtils/InstancedStruct.h
@@ -337,9 +337,9 @@ public:
 	template<typename T = BaseStructT, typename = std::enable_if_t<std::is_base_of_v<BaseStructT, std::decay_t<T>>>>
 	TInstancedStruct& operator=(const TConstStructView<T> InOther)
 	{
-		if (this != &InOther)
+		if (TConstStructView<T>(*this) != InOther)
 		{
-			InstancedStruct = InOther.InstancedStruct;
+			InstancedStruct.InitializeAs(InOther.GetScriptStruct(), InOther.GetMemory());
 		}
 		return *this;
 	}
