[AutoRTFM] Use a transactionally-safe lock for property bag access.

This lock can be hit while running the StateTree tests.

diff --git a/Engine/Source/Runtime/CoreUObject/Private/StructUtils/PropertyBag.cpp b/Engine/Source/Runtime/CoreUObject/Private/StructUtils/PropertyBag.cpp
--- a/Engine/Source/Runtime/CoreUObject/Private/StructUtils/PropertyBag.cpp
+++ b/Engine/Source/Runtime/CoreUObject/Private/StructUtils/PropertyBag.cpp
@@ -3,6 +3,8 @@
 #include "StructUtils/PropertyBag.h"
 
 #include "HAL/IConsoleManager.h"
+#include "Misc/ScopeLock.h"
+#include "Misc/TransactionallySafeCriticalSection.h"
 #include "UObject/EnumProperty.h"
 #include "UObject/Package.h"
 #include "UObject/TextProperty.h"
@@ -3815,7 +3817,7 @@ EPropertyBagResult FPropertyBagSetRef::AddValueSoftPath(const FSoftObjectPath& I
 namespace UE::StructUtils::Private
 {
 	// Lock to prevent concurrent access to lazily-constructed UPropertyBag objects in UPropertyBag::GetOrCreateFromDescs
-	static FCriticalSection GPropertyBagLock;
+	static FTransactionallySafeCriticalSection GPropertyBagLock;
 }
 
 const UPropertyBag* UPropertyBag::GetOrCreateFromDescs(const TConstArrayView<FPropertyBagPropertyDesc> PropertyDescs, const TCHAR* PrefixName)
@@ -3827,7 +3829,7 @@ const UPropertyBag* UPropertyBag::GetOrCreateFromDescs(const TConstArrayView<FPr
 
 	// We need to linearize this entire function otherwise threads that create bags of identical layouts can view
 	// partially-constructed objects 
-	FScopeLock ScopeLock(&UE::StructUtils::Private::GPropertyBagLock);
+	UE::TScopeLock ScopeLock(UE::StructUtils::Private::GPropertyBagLock);
 
 	// we need to use StaticFindObjectFastInternal with ExclusiveInternalFlags = EInternalObjectFlags::None
 	// here because objects with RF_NeedPostLoad cannot be found by regular FindObject calls as they will have
