For types which do not have an intrusive unset optional state, avoid instantiating TOptional from UScriptStruct operations.
This avoids instantiating TOptional for abstract types and allows TOptional to contain a typed member.

diff --git a/Engine/Source/Runtime/CoreUObject/Public/UObject/Class.h b/Engine/Source/Runtime/CoreUObject/Public/UObject/Class.h
--- a/Engine/Source/Runtime/CoreUObject/Public/UObject/Class.h
+++ b/Engine/Source/Runtime/CoreUObject/Public/UObject/Class.h
@@ -1694,19 +1694,41 @@ public:
 		/** Construct an unset optional value */
 		virtual void InitializeIntrusiveUnsetOptionalValue(void* Data) const override
 		{
-			new (Data) TOptional<CPPSTRUCT>();
+			if constexpr(::HasIntrusiveUnsetOptionalState<CPPSTRUCT>())
+			{
+				new (Data) TOptional<CPPSTRUCT>();
+			}
+			else
+			{
+				checkf(false, TEXT("InitializeIntrusiveUnsetOptionalValue should not be called for types which do not have an intrusive unset optional state"));
+			}
 		}
 
 		/** Return true if the optional value at Data is in an unset state */
 		virtual bool IsIntrusiveOptionalValueSet(const void* Data) const override
 		{
-			return reinterpret_cast<const TOptional<CPPSTRUCT>*>(Data)->IsSet();
+			if constexpr(::HasIntrusiveUnsetOptionalState<CPPSTRUCT>())
+			{
+				return reinterpret_cast<const TOptional<CPPSTRUCT>*>(Data)->IsSet();
+			}
+			else 
+			{
+				checkf(false, TEXT("IsIntrusiveOptionalValueSet should not be called for types which do not have an intrusive unset optional state"));
+				return false;
+			}
 		}
 
 		/** Reset an optional value to its unset state */
 		virtual void ClearIntrusiveOptionalValue(void* Data) const override
 		{
-			reinterpret_cast<TOptional<CPPSTRUCT>*>(Data)->Reset();
+			if constexpr(::HasIntrusiveUnsetOptionalState<CPPSTRUCT>())
+			{
+				reinterpret_cast<TOptional<CPPSTRUCT>*>(Data)->Reset();
+			}
+			else
+			{
+				checkf(false, TEXT("ClearIntrusiveOptionalValue should not be called for types which do not have an intrusive unset optional state"));
+			}
 		}
 
 		/** 
