[UE] Disable IDO placeholder features.

#rnx

diff --git a/Engine/Source/Runtime/CoreUObject/Private/UObject/PropertyBagRepository.cpp b/Engine/Source/Runtime/CoreUObject/Private/UObject/PropertyBagRepository.cpp
--- a/Engine/Source/Runtime/CoreUObject/Private/UObject/PropertyBagRepository.cpp
+++ b/Engine/Source/Runtime/CoreUObject/Private/UObject/PropertyBagRepository.cpp
@@ -940,7 +940,7 @@ bool FPropertyBagRepository::IsPropertyBagPlaceholderObject(const UObject* Objec
 namespace Private
 {
 #if WITH_EDITOR
-	static bool bEnablePropertyBagPlaceholderObjectSupport = true;
+	static bool bEnablePropertyBagPlaceholderObjectSupport = false;
 	static FAutoConsoleVariableRef CVarEnablePropertyBagPlaceholderObjectSupport(
 		TEXT("IDO.Placeholder.Enable"),
 		bEnablePropertyBagPlaceholderObjectSupport,
