AssetRegistry: Fix FAssetIdentifier::operator== to require PrimaryAssetType to be equal.
#rnx

diff --git a/Engine/Source/Runtime/CoreUObject/Public/AssetRegistry/AssetIdentifier.h b/Engine/Source/Runtime/CoreUObject/Public/AssetRegistry/AssetIdentifier.h
--- a/Engine/Source/Runtime/CoreUObject/Public/AssetRegistry/AssetIdentifier.h
+++ b/Engine/Source/Runtime/CoreUObject/Public/AssetRegistry/AssetIdentifier.h
@@ -128,7 +128,8 @@ struct FAssetIdentifier
 
 	friend inline bool operator==(const FAssetIdentifier& A, const FAssetIdentifier& B)
 	{
-		return A.PackageName == B.PackageName && A.ObjectName == B.ObjectName && A.ValueName == B.ValueName;
+		return A.PackageName == B.PackageName && A.PrimaryAssetType == B.PrimaryAssetType
+			&& A.ObjectName == B.ObjectName && A.ValueName == B.ValueName;
 	}
 
 	friend inline uint32 GetTypeHash(const FAssetIdentifier& Key)
