Update UObject::SaveConfig() to respect the "DeprecatedProperty" property metadata key and skip saving deprecated properties to disk.
This allows UpdateSinglePropertyInConfigFile() to be properly used to delete deprecated properties from .ini files.

diff --git a/Engine/Source/Runtime/CoreUObject/Private/UObject/Obj.cpp b/Engine/Source/Runtime/CoreUObject/Private/UObject/Obj.cpp
--- a/Engine/Source/Runtime/CoreUObject/Private/UObject/Obj.cpp
+++ b/Engine/Source/Runtime/CoreUObject/Private/UObject/Obj.cpp
@@ -3525,7 +3525,14 @@ void UObject::SaveConfig(uint64 RequiredPropertyFlags, const TCHAR* InFilename,
 
 			// Properties that are the same as the parent class' defaults should not be saved to ini
 			// Before modifying any key in the section, first check to see if it is different from the parent.
+			
+#if WITH_METADATA
+			static const FName NAME_DeprecatedProperty = TEXT("DeprecatedProperty");
+			const bool bPropDeprecated = Property->HasAnyPropertyFlags(CPF_Deprecated) || Property->HasMetaData(NAME_DeprecatedProperty);
+#else
 			const bool bPropDeprecated = Property->HasAnyPropertyFlags(CPF_Deprecated);
+#endif
+
 			const bool bIsPropertyInherited = Property->GetOwnerClass() != GetClass();
 			const bool bShouldCheckIfIdenticalBeforeAdding = !GetClass()->HasAnyClassFlags(CLASS_ConfigDoNotCheckDefaults) && !bPerObject && bIsPropertyInherited;
 			UObject* SuperClassDefaultObject = GetClass()->GetSuperClass()->GetDefaultObject();
