[AsyncLoading]

- Fixes localization issue where import tables were being patched with the localized name which breaks loading (since we then start looking for localized UPackage named which breaks FindObject searches since all UObject references are made by the non-localized UPackage name)
- Fixes incorrect VeryVerbose CLOG check when NotifyConstructedDuringAsyncLoading which would incorrectly trigger every time a UPackage was created by the loader due to the AsyncPackage not yet being initialized
- Fixes NotifyConstructedDuringAsyncLoading to allow adding ConstructedObjects to the current AsyncPackage's ConstructedObject's list while it creates objects during deferred post load. Deferred postload processing will always look for newly created objects when postloading so we only want to prevent adding objects if we are already done deferred postloading.

[RN] Fixed loading localized content when running the editor with -game. Previosly localized imports would not load properly when using zen loader.

diff --git a/Engine/Source/Runtime/CoreUObject/Private/Serialization/AsyncLoading2.cpp b/Engine/Source/Runtime/CoreUObject/Private/Serialization/AsyncLoading2.cpp
--- a/Engine/Source/Runtime/CoreUObject/Private/Serialization/AsyncLoading2.cpp
+++ b/Engine/Source/Runtime/CoreUObject/Private/Serialization/AsyncLoading2.cpp
@@ -591,31 +591,31 @@ static bool AsyncLoading2_ShouldSkipLoadingExport(const EExportFilterFlags Filte
 }
 
 #if WITH_EDITOR
-
+/** Returns the CoreRedirected package name. Otherwise returns the name unchanged. */
 FName ApplyPackageNameRedirections(FName PackageName)
 {
-	if (!PackageName.IsNone())
+	FName RedirectedName = FCoreRedirects::GetRedirectedName(ECoreRedirectFlags::Type_Package, FCoreRedirectObjectName(NAME_None, NAME_None, PackageName)).PackageName;
+	if (RedirectedName != PackageName)
 	{
-		FName NewPackageName = FCoreRedirects::GetRedirectedName(ECoreRedirectFlags::Type_Package, FCoreRedirectObjectName(NAME_None, NAME_None, PackageName)).PackageName;
-		if (NewPackageName != PackageName)
-		{
-			UE_LOG(LogStreaming, Verbose, TEXT("Applying core redirection for package %s to package %s"), *PackageName.ToString(), *NewPackageName.ToString());
-			PackageName = NewPackageName;
-		}
+		UE_LOG(LogStreaming, Verbose, TEXT("Applying core redirection for package %s to package %s"), *PackageName.ToString(), *RedirectedName.ToString());
+	}
+
+	return RedirectedName;
+}
 
-		// Handle localization when running in -game
-		if (!GIsEditor)
+/** Returns localized package name to load when running with -game. Otherwise returns the name unchanged. */
+FName ApplyLocalizationPackageNameRedirects(FName PackageName)
+{
+	if (!GIsEditor)
+	{
+		UE::Core::Private::FPlayInEditorLoadingScope PlayInEditorIDScope(INDEX_NONE);
+		FName RedirectedName = FPackageLocalizationManager::Get().FindLocalizedPackageName(PackageName);
+		if (!RedirectedName.IsNone())
 		{
-			UE::Core::Private::FPlayInEditorLoadingScope PlayInEditorIDScope(INDEX_NONE);
-			NewPackageName = FPackageLocalizationManager::Get().FindLocalizedPackageName(PackageName);
-			if (!NewPackageName.IsNone())
-			{
-				UE_LOG(LogStreaming, Verbose, TEXT("Applying localization redirection for package %s to package %s"), *PackageName.ToString(), *NewPackageName.ToString());
-				PackageName = NewPackageName;
-			}
+			UE_LOG(LogStreaming, Verbose, TEXT("Applying localization redirection for package %s to package %s"), *PackageName.ToString(), *RedirectedName.ToString());
+			PackageName = RedirectedName;
 		}
 	}
-
 	return PackageName;
 }
 
@@ -5474,12 +5474,17 @@ bool FAsyncLoadingThread2::CreateAsyncPackagesFromQueue(FAsyncLoadingThreadState
 				}
 			}
 
+			FName UPackageName = PackageNameToLoad;
 #if WITH_EDITOR
+			// Redirects should redirect both the UPackage name as well as the package name to load
 			PackageNameToLoad = ApplyPackageNameRedirections(PackageNameToLoad);
+			UPackageName = PackageNameToLoad;
+
+			// Localization should only update the package name to load
+			PackageNameToLoad = ApplyLocalizationPackageNameRedirects(PackageNameToLoad);
 #endif
 
 			FPackageId PackageIdToLoad = FPackageId::FromName(PackageNameToLoad);
-			FName UPackageName = PackageNameToLoad;
 			{
 				FName SourcePackageName;
 				FPackageId RedirectedToPackageId;
@@ -6107,6 +6112,7 @@ void FAsyncPackage2::ImportPackagesRecursiveInner(FAsyncLoadingThreadState2& Thr
 #if WITH_EDITOR
 		if (!ImportedPackageNameToLoad.IsNone())
 		{
+			// Redirects should redirect both the UPackage name as well as the package name to load
 			const FName NewPackageNameToLoad = ApplyPackageNameRedirections(ImportedPackageNameToLoad);
 			if (ImportedPackageNameToLoad != NewPackageNameToLoad)
 			{
@@ -6136,7 +6142,6 @@ void FAsyncPackage2::ImportPackagesRecursiveInner(FAsyncLoadingThreadState2& Thr
 					ImportedPackageIdToLoad = FPackageId::FromName(NAME_None);
 					ImportedPackageNameToLoad = NAME_None;
 				}
-
 				ImportedPackageId = FPackageId::FromName(ImportedPackageUPackageName);
 				// Rewrite the import table
 				ImportedPackageIds[LocalImportedPackageIndex] = ImportedPackageId;
@@ -6145,6 +6150,19 @@ void FAsyncPackage2::ImportPackagesRecursiveInner(FAsyncLoadingThreadState2& Thr
 		}
 #endif
 
+#if WITH_EDITOR
+		// Note, localization should apply after remapping from instancing has occurred
+		{
+			FName LocalizedPackageNameToLoad = ApplyLocalizationPackageNameRedirects(ImportedPackageNameToLoad);
+			if (LocalizedPackageNameToLoad != ImportedPackageNameToLoad)
+			{
+				// Localization should only update the package name to load
+				ImportedPackageNameToLoad = LocalizedPackageNameToLoad;
+				ImportedPackageIdToLoad = FPackageId::FromName(ImportedPackageNameToLoad);
+			}
+		}
+#endif
+
 		{
 			FName SourcePackageName;
 			FPackageId RedirectedToPackageId;
@@ -10263,7 +10281,7 @@ void FAsyncLoadingThread2::NotifyConstructedDuringAsyncLoading(UObject* Object,
 	{
 		if (FAsyncPackage2* AsyncPackage = (FAsyncPackage2*)LinkerLoad->AsyncRoot)
 		{
-			if (AsyncPackage->AsyncPackageLoadingState < EAsyncPackageLoadingState2::DeferredPostLoad)
+			if (AsyncPackage->AsyncPackageLoadingState <= EAsyncPackageLoadingState2::DeferredPostLoad)
 			{
 				check(AsyncPackage->Desc.UPackageId == ObjectPackage->GetPackageId());
 				AsyncPackage->AddConstructedObject(Object, bSubObjectThatAlreadyExists);
@@ -10289,7 +10307,7 @@ void FAsyncLoadingThread2::NotifyConstructedDuringAsyncLoading(UObject* Object,
 
 	if (ContextAsyncPackage->AsyncPackageLoadingState <= EAsyncPackageLoadingState2::DeferredPostLoad)
 	{
-		UE_CLOG(ContextAsyncPackage->Desc.UPackageId != ObjectPackage->GetPackageId(), LogStreaming, VeryVerbose,
+		UE_CLOG(ObjectPackage->GetPackageId().IsValid() && ContextAsyncPackage->Desc.UPackageId != ObjectPackage->GetPackageId(), LogStreaming, VeryVerbose,
 			TEXT("Constructed object '%s' is part of package '%s' which is no longer postloading objects. This object will be postloaded with the FUObjectThreadContext AsyncPackage '%s'(%s) instead."),
 			*Object->GetFullName(),
 			*ObjectPackage->GetFullName(),
