Removed deprecated features from UObjectGlobals.

diff --git a/Engine/Source/Runtime/CoreUObject/Private/Serialization/AsyncLoading.cpp b/Engine/Source/Runtime/CoreUObject/Private/Serialization/AsyncLoading.cpp
--- a/Engine/Source/Runtime/CoreUObject/Private/Serialization/AsyncLoading.cpp
+++ b/Engine/Source/Runtime/CoreUObject/Private/Serialization/AsyncLoading.cpp
@@ -8802,11 +8802,6 @@ static FAutoConsoleCommand DumpSerializeCmd(
 	);
 #endif
 
-bool IsEventDrivenLoaderEnabledInCookedBuilds()
-{
-	return true;
-}
-
 bool IsEventDrivenLoaderEnabled()
 {
 	return FPlatformProperties::RequiresCookedData();
diff --git a/Engine/Source/Runtime/CoreUObject/Private/UObject/AnyPackagePrivate.h b/Engine/Source/Runtime/CoreUObject/Private/UObject/AnyPackagePrivate.h
--- a/Engine/Source/Runtime/CoreUObject/Private/UObject/AnyPackagePrivate.h
+++ b/Engine/Source/Runtime/CoreUObject/Private/UObject/AnyPackagePrivate.h
@@ -6,6 +6,8 @@
 
 #pragma once
 
-#define ANY_PACKAGE_DEPRECATED ((UPackage*)-1)
+// HEADER_UNIT_SKIP - Deprecated
 
-#define IS_ANY_PACKAGE_DEPRECATED(Package) (reinterpret_cast<PTRINT>(Package) == -1)
\ No newline at end of file
+#include "CoreTypes.h"
+
+UE_DEPRECATED_HEADER(5.7, "AnyPackagePrivate.h has been deprecated and should no longer be included")
diff --git a/Engine/Source/Runtime/CoreUObject/Private/UObject/GarbageCollection.cpp b/Engine/Source/Runtime/CoreUObject/Private/UObject/GarbageCollection.cpp
--- a/Engine/Source/Runtime/CoreUObject/Private/UObject/GarbageCollection.cpp
+++ b/Engine/Source/Runtime/CoreUObject/Private/UObject/GarbageCollection.cpp
@@ -946,15 +946,6 @@ public:
 	}
 } GUObjectPurge;
 
-/**
-  * (Deprecated) Returns true if this function is called from the async destruction thread.
-  * It will also return true if we're running single-threaded and this function is called on the game thread
-  */
-bool IsInGarbageCollectorThread()
-{
-	return IsInGameThread();
-}
-
 //////////////////////////////////////////////////////////////////////////
 
 namespace UE::GC {
@@ -4558,7 +4549,6 @@ PRAGMA_DISABLE_DEPRECATION_WARNINGS
 		if (!GReachabilityState.IsSuspended())
 		{
 			const double StartTime = FPlatformTime::Seconds();
-			FCoreUObjectDelegates::TraceExternalRootsForReachabilityAnalysis.Broadcast(*this, KeepFlags, !(Options & EGCOptions::Parallel));
 			GGCStats.TraceExternalRootsTime += FPlatformTime::Seconds() - StartTime;
 		}
 PRAGMA_ENABLE_DEPRECATION_WARNINGS
diff --git a/Engine/Source/Runtime/CoreUObject/Private/UObject/Obj.cpp b/Engine/Source/Runtime/CoreUObject/Private/UObject/Obj.cpp
--- a/Engine/Source/Runtime/CoreUObject/Private/UObject/Obj.cpp
+++ b/Engine/Source/Runtime/CoreUObject/Private/UObject/Obj.cpp
@@ -1482,9 +1482,6 @@ void UObject::PreSave(FObjectPreSaveContext SaveContext)
 {
 	SaveContext.Data.bBaseClassCalled = true;
 #if WITH_EDITOR
-	PRAGMA_DISABLE_DEPRECATION_WARNINGS
-	FCoreUObjectDelegates::OnObjectSaved.Broadcast(this);
-	PRAGMA_ENABLE_DEPRECATION_WARNINGS
 	FCoreUObjectDelegates::OnObjectPreSave.Broadcast(this, SaveContext);
 
 	// Validate ClassWithin before save, matches code in StaticAllocateObjectErrorTests
diff --git a/Engine/Source/Runtime/CoreUObject/Private/UObject/SavePackage/SavePackageUtilities.cpp b/Engine/Source/Runtime/CoreUObject/Private/UObject/SavePackage/SavePackageUtilities.cpp
--- a/Engine/Source/Runtime/CoreUObject/Private/UObject/SavePackage/SavePackageUtilities.cpp
+++ b/Engine/Source/Runtime/CoreUObject/Private/UObject/SavePackage/SavePackageUtilities.cpp
@@ -968,17 +968,6 @@ bool IsEditorOnlyObject(const UObject* InObject, bool bCheckRecursive)
 	return IsEditorOnlyObjectInternal(InObject, Flags);
 }
 
-bool IsEditorOnlyObject(const UObject* InObject, bool bCheckRecursive, bool bCheckMarks)
-{
-	using namespace UE::SavePackageUtilities;
-	EEditorOnlyObjectFlags Flags = EEditorOnlyObjectFlags::None;
-	Flags |= bCheckRecursive ? EEditorOnlyObjectFlags::CheckRecursive : EEditorOnlyObjectFlags::None;
-	PRAGMA_DISABLE_DEPRECATION_WARNINGS;
-	Flags |= bCheckMarks ? EEditorOnlyObjectFlags::CheckMarks : EEditorOnlyObjectFlags::None;
-	PRAGMA_ENABLE_DEPRECATION_WARNINGS;
-	return IsEditorOnlyObjectInternal(InObject, Flags);
-}
-
 bool IsEditorOnlyObject(const UObject* InObject, bool bCheckRecursive,
 	TFunctionRef<UE::SavePackageUtilities::EEditorOnlyObjectResult(const UObject*)> LookupInCache,
 	TFunctionRef<void(const UObject*, bool)> AddToCache)
@@ -1001,9 +990,6 @@ bool IsEditorOnlyObjectWithoutWritingCache(const UObject* InObject, EEditorOnlyO
 	bool bCheckRecursive = EnumHasAnyFlags(Flags, EEditorOnlyObjectFlags::CheckRecursive);
 	bool bIgnoreEditorOnlyClass = EnumHasAnyFlags(Flags, EEditorOnlyObjectFlags::ApplyHasNonEditorOnlyReferences) &&
 		InObject->HasNonEditorOnlyReferences();
-	PRAGMA_DISABLE_DEPRECATION_WARNINGS;
-	bool bCheckMarks = EnumHasAnyFlags(Flags, EEditorOnlyObjectFlags::CheckMarks);
-	PRAGMA_ENABLE_DEPRECATION_WARNINGS;
 
 	DECLARE_SCOPE_CYCLE_COUNTER(TEXT("IsEditorOnlyObject"), STAT_IsEditorOnlyObject, STATGROUP_LoadTime);
 
@@ -1011,8 +997,7 @@ bool IsEditorOnlyObjectWithoutWritingCache(const UObject* InObject, EEditorOnlyO
 	// Ignore their value of IsEditorOnly
 	if (!InObject->HasAnyFlags(RF_ClassDefaultObject))
 	{
-		if (!bIgnoreEditorOnlyClass &&
-			((bCheckMarks && InObject->HasAnyMarks(OBJECTMARK_EditorOnly)) || InObject->IsEditorOnly()))
+		if (!bIgnoreEditorOnlyClass && InObject->IsEditorOnly())
 		{
 			return true;
 		}
diff --git a/Engine/Source/Runtime/CoreUObject/Private/UObject/SavePackage/SavePackageUtilities.h b/Engine/Source/Runtime/CoreUObject/Private/UObject/SavePackage/SavePackageUtilities.h
--- a/Engine/Source/Runtime/CoreUObject/Private/UObject/SavePackage/SavePackageUtilities.h
+++ b/Engine/Source/Runtime/CoreUObject/Private/UObject/SavePackage/SavePackageUtilities.h
@@ -200,7 +200,6 @@ enum class EEditorOnlyObjectFlags
 	None = 0,
 	CheckRecursive = 1 << 1,
 	ApplyHasNonEditorOnlyReferences = 1 << 2,
-	CheckMarks UE_DEPRECATED(5.3, "CheckMarks is no longer supported") = 1 << 3,
 
 };
 ENUM_CLASS_FLAGS(EEditorOnlyObjectFlags);
diff --git a/Engine/Source/Runtime/CoreUObject/Private/UObject/UObjectBase.cpp b/Engine/Source/Runtime/CoreUObject/Private/UObject/UObjectBase.cpp
--- a/Engine/Source/Runtime/CoreUObject/Private/UObject/UObjectBase.cpp
+++ b/Engine/Source/Runtime/CoreUObject/Private/UObject/UObjectBase.cpp
@@ -789,9 +789,6 @@ void UClassRegisterAllCompiledInClasses()
 	if (AddedClasses.Num() > 0)
 	{
 		FCoreUObjectDelegates::ReloadAddedClassesDelegate.Broadcast(AddedClasses);
-		PRAGMA_DISABLE_DEPRECATION_WARNINGS
-		FCoreUObjectDelegates::RegisterHotReloadAddedClassesDelegate.Broadcast(AddedClasses);
-		PRAGMA_ENABLE_DEPRECATION_WARNINGS
 	}
 #endif
 }
diff --git a/Engine/Source/Runtime/CoreUObject/Private/UObject/UObjectGlobals.cpp b/Engine/Source/Runtime/CoreUObject/Private/UObject/UObjectGlobals.cpp
--- a/Engine/Source/Runtime/CoreUObject/Private/UObject/UObjectGlobals.cpp
+++ b/Engine/Source/Runtime/CoreUObject/Private/UObject/UObjectGlobals.cpp
@@ -80,7 +80,6 @@
 #include "ProfilingDebugging/AssetMetadataTrace.h"
 #include "Misc/PackageAccessTracking.h"
 #include "UObject/PropertyWithSetterAndGetter.h"
-#include "UObject/AnyPackagePrivate.h"
 #include "Templates/GuardValueAccessors.h"
 #include "UObject/UObjectGlobalsInternal.h"
 #include "Serialization/AsyncPackageLoader.h"
@@ -180,10 +179,8 @@ namespace LoadPackageStats
 
 /** CoreUObject delegates */
 PRAGMA_DISABLE_DEPRECATION_WARNINGS
-FCoreUObjectDelegates::FRegisterHotReloadAddedClassesDelegate FCoreUObjectDelegates::RegisterHotReloadAddedClassesDelegate;
 FCoreUObjectDelegates::FRegisterClassForHotReloadReinstancingDelegate FCoreUObjectDelegates::RegisterClassForHotReloadReinstancingDelegate;
 FCoreUObjectDelegates::FReinstanceHotReloadedClassesDelegate FCoreUObjectDelegates::ReinstanceHotReloadedClassesDelegate;
-FCoreUObjectDelegates::FTraceExternalRootsForReachabilityAnalysisDelegate FCoreUObjectDelegates::TraceExternalRootsForReachabilityAnalysis;
 PRAGMA_ENABLE_DEPRECATION_WARNINGS
 FCoreUObjectDelegates::FReloadReinstancingCompleteDelegate FCoreUObjectDelegates::ReloadReinstancingCompleteDelegate;
 FCoreUObjectDelegates::FReloadCompleteDelegate FCoreUObjectDelegates::ReloadCompleteDelegate;
@@ -194,9 +191,6 @@ FCoreUObjectDelegates::FOnPostInitSparseClassData FCoreUObjectDelegates::OnPostI
 FCoreUObjectDelegates::FOnPackageReloaded FCoreUObjectDelegates::OnPackageReloaded;
 FCoreUObjectDelegates::FNetworkFileRequestPackageReload FCoreUObjectDelegates::NetworkFileRequestPackageReload;
 #if WITH_EDITOR
-PRAGMA_DISABLE_DEPRECATION_WARNINGS;
-FCoreUObjectDelegates::FAutoPackageBackupDelegate FCoreUObjectDelegates::AutoPackageBackupDelegate;
-PRAGMA_ENABLE_DEPRECATION_WARNINGS;
 FCoreUObjectDelegates::FOnPreObjectPropertyChanged FCoreUObjectDelegates::OnPreObjectPropertyChanged;
 FCoreUObjectDelegates::FOnObjectPropertyChanged FCoreUObjectDelegates::OnObjectPropertyChanged;
 TSet<UObject*> FCoreUObjectDelegates::ObjectsModifiedThisFrame;
@@ -210,9 +204,6 @@ FCoreUObjectDelegates::FOnObjectPostCDOCompiled FCoreUObjectDelegates::OnObjectP
 FCoreUObjectDelegates::FOnAssetLoaded FCoreUObjectDelegates::OnAssetLoaded;
 FCoreUObjectDelegates::FOnObjectConstructed FCoreUObjectDelegates::OnObjectConstructed;
 FCoreUObjectDelegates::FOnEndLoadPackage FCoreUObjectDelegates::OnEndLoadPackage;
-PRAGMA_DISABLE_DEPRECATION_WARNINGS
-FCoreUObjectDelegates::FOnObjectSaved FCoreUObjectDelegates::OnObjectSaved;
-PRAGMA_ENABLE_DEPRECATION_WARNINGS
 FCoreUObjectDelegates::FOnObjectPreSave FCoreUObjectDelegates::OnObjectPreSave;
 #endif // WITH_EDITOR
 
@@ -450,27 +441,6 @@ int32 UpdateSuffixForNextNewObject(UObject* Parent, const UClass* Class, TFuncti
 	return UpdateSuffixForNextNewObject(Parent, *Class->GetPathName(), IndexMutator);
 }
 
-//
-// Find an object, path must unqualified
-//
-UObject* StaticFindObjectFast(UClass* ObjectClass, UObject* ObjectPackage, FName ObjectName, bool bExactClass, bool bAnyPackage, EObjectFlags ExclusiveFlags, EInternalObjectFlags ExclusiveInternalFlags)
-{
-	UE_CLOG(UE::IsSavingPackage(nullptr), LogUObjectGlobals, Fatal, TEXT("Illegal call to StaticFindObjectFast() while serializing object data!"));
-	UE_CLOG(IsGarbageCollectingAndLockingUObjectHashTables(), LogUObjectGlobals, Fatal, TEXT("Illegal call to StaticFindObjectFast() while garbage collecting!"));
-
-	// We don't want to return any objects that are currently being background loaded unless we're using FindObject during async loading.
-	ExclusiveInternalFlags |= UE::GetAsyncLoadingInternalFlagsExclusion();
-	PRAGMA_DISABLE_DEPRECATION_WARNINGS
-	UObject* FoundObject = StaticFindObjectFastInternal(ObjectClass, ObjectPackage, ObjectName, bExactClass, bAnyPackage, ExclusiveFlags, ExclusiveInternalFlags);
-	PRAGMA_ENABLE_DEPRECATION_WARNINGS
-	if (!FoundObject)
-	{
-		FoundObject = StaticFindObjectWithChangedLegacyPath(ObjectClass, ObjectPackage, ObjectName, bExactClass, ExclusiveInternalFlags);
-	}
-
-	return FoundObject;
-}
-
 //
 // Find an object, path must unqualified
 //
@@ -491,26 +461,6 @@ UObject* StaticFindObjectFast(UClass* ObjectClass, UObject* ObjectPackage, FName
 	return FoundObject;
 }
 
-UObject* StaticFindObjectFastSafe(UClass* ObjectClass, UObject* ObjectPackage, FName ObjectName, bool bExactClass, bool bAnyPackage, EObjectFlags ExclusiveFlags, EInternalObjectFlags ExclusiveInternalFlags)
-{
-	UObject* FoundObject = nullptr;
-	
-	if (!UE::IsSavingPackage(nullptr) && !IsGarbageCollectingAndLockingUObjectHashTables())
-	{
-		// We don't want to return any objects that are currently being background loaded unless we're using FindObject during async loading.
-		ExclusiveInternalFlags |= UE::GetAsyncLoadingInternalFlagsExclusion();
-		PRAGMA_DISABLE_DEPRECATION_WARNINGS
-		FoundObject = StaticFindObjectFastInternal(ObjectClass, ObjectPackage, ObjectName, bExactClass, bAnyPackage, ExclusiveFlags, ExclusiveInternalFlags);
-		PRAGMA_ENABLE_DEPRECATION_WARNINGS
-		if (!FoundObject)
-		{
-			FoundObject = StaticFindObjectWithChangedLegacyPath(ObjectClass, ObjectPackage, ObjectName, bExactClass, ExclusiveInternalFlags);
-		}
-	}
-
-	return FoundObject;
-}
-
 UObject* StaticFindObjectFastSafe(UClass* ObjectClass, UObject* ObjectPackage, FName ObjectName, bool bExactClass, EObjectFlags ExclusiveFlags, EInternalObjectFlags ExclusiveInternalFlags)
 {
 	UObject* FoundObject = nullptr;
@@ -563,8 +513,7 @@ UObject* StaticFindObject( UClass* ObjectClass, UObject* InObjectPackage, const
 	INC_DWORD_STAT(STAT_FindObject);
 
 	// Resolve the object and package name.
-	const bool bAnyPackage = IS_ANY_PACKAGE_DEPRECATED(InObjectPackage);
-	UObject* ObjectPackage = bAnyPackage ? nullptr : InObjectPackage;
+	UObject* ObjectPackage = InObjectPackage;
 
 	UObject* MatchingObject = nullptr;
 #if WITH_EDITOR
@@ -583,22 +532,15 @@ UObject* StaticFindObject( UClass* ObjectClass, UObject* InObjectPackage, const
 	InName = OrigInName;
 	
 	// Don't resolve the name if we're searching in any package
-	if (!bAnyPackage)
+	if (!ResolveName2(ObjectPackage, InName, false, false))
 	{
-		if (!ResolveName2(ObjectPackage, InName, false, false))
-		{
-			return nullptr;
-		}
-	}
-	else
-	{
-		ConstructorHelpers_StripObjectClass2(InName);
+		return nullptr;
 	}
 	
 	FName ObjectName(InName.ToView(), FNAME_Add);
 	
 	PRAGMA_DISABLE_DEPRECATION_WARNINGS
-	return StaticFindObjectFast(ObjectClass, ObjectPackage, ObjectName, bExactClass, bAnyPackage);
+	return StaticFindObjectFast(ObjectClass, ObjectPackage, ObjectName, bExactClass);
 	PRAGMA_ENABLE_DEPRECATION_WARNINGS
 }
 
@@ -641,7 +583,7 @@ UObject* StaticFindObjectChecked( UClass* ObjectClass, UObject* ObjectParent, co
 #if !(UE_BUILD_SHIPPING || UE_BUILD_TEST)
 	if( !Result )
 	{
-		UE_LOG(LogUObjectGlobals, Fatal, TEXT("Failed to find object '%s %s.%s'"), *ObjectClass->GetName(), IS_ANY_PACKAGE_DEPRECATED(ObjectParent) ? TEXT("Any") : ObjectParent ? *ObjectParent->GetName() : TEXT("None"), InName);
+		UE_LOG(LogUObjectGlobals, Fatal, TEXT("Failed to find object '%s %s.%s'"), *ObjectClass->GetName(), ObjectParent ? *ObjectParent->GetName() : TEXT("None"), InName);
 	}
 #endif
 	return Result;
@@ -1331,7 +1273,7 @@ bool ParseObject( const TCHAR* Stream, const TCHAR* Match, UClass* Class, UObjec
 					? StaticLoadObject(Class, nullptr, ObjNameOrPathName)
 					: StaticFindObject(Class, nullptr, ObjNameOrPathName);
 			}
-			else if (InParent && !IS_ANY_PACKAGE_DEPRECATED(InParent))
+			else if (InParent)
 			{
 				// Try to find the object within its parent
 				return StaticFindObject(Class, InParent, ObjNameOrPathName);
@@ -2789,7 +2731,7 @@ namespace NameReuse
 		}
 
 		static const FName NamePackage(NAME_Package);
-		if (!Parent || IS_ANY_PACKAGE_DEPRECATED(Parent) || !!(Options & EUniqueObjectNameOptions::GloballyUnique) || BaseName == NamePackage || FPlatformProperties::HasEditorOnlyData() || !GFastPathUniqueNameGeneration)
+		if (!Parent || !!(Options & EUniqueObjectNameOptions::GloballyUnique) || BaseName == NamePackage || FPlatformProperties::HasEditorOnlyData() || !GFastPathUniqueNameGeneration)
 		{
 			return FName();
 		}
@@ -2921,7 +2863,7 @@ FName MakeUniqueObjectName(UObject* Parent, const UClass* Class, FName InBaseNam
 				int32 NameNumber = 0;
 				UE_AUTORTFM_OPEN
 				{
-					if (Parent && !IS_ANY_PACKAGE_DEPRECATED(Parent) && !(Options & EUniqueObjectNameOptions::GloballyUnique))
+					if (Parent && !(Options & EUniqueObjectNameOptions::GloballyUnique))
 					{
 						if (!FPlatformProperties::HasEditorOnlyData() && GFastPathUniqueNameGeneration)
 						{
@@ -2956,7 +2898,7 @@ FName MakeUniqueObjectName(UObject* Parent, const UClass* Class, FName InBaseNam
 				TestName = FName(BaseName, NameNumber);
 			}
 
-			if (IS_ANY_PACKAGE_DEPRECATED(Parent) || !!(Options & EUniqueObjectNameOptions::GloballyUnique))
+			if (!!(Options & EUniqueObjectNameOptions::GloballyUnique))
 			{
 				ExistingObject = StaticFindFirstObject(nullptr, *TestName.ToString());
 			}
diff --git a/Engine/Source/Runtime/CoreUObject/Private/UObject/UObjectHash.cpp b/Engine/Source/Runtime/CoreUObject/Private/UObject/UObjectHash.cpp
--- a/Engine/Source/Runtime/CoreUObject/Private/UObject/UObjectHash.cpp
+++ b/Engine/Source/Runtime/CoreUObject/Private/UObject/UObjectHash.cpp
@@ -20,7 +20,6 @@
 #include "Async/ParallelFor.h"
 #include "HAL/IConsoleManager.h"
 #include "HAL/LowLevelMemStats.h"
-#include "UObject/AnyPackagePrivate.h"
 #include "ProfilingDebugging/MetadataTrace.h"
 #include "AutoRTFM.h"
 
@@ -1037,7 +1036,7 @@ UObject* StaticFindObjectInPackageInternal(FUObjectHashTables& ThreadHash, const
 	return Result;
 }
 
-UObject* StaticFindObjectFastInternalThreadSafe(FUObjectHashTables& ThreadHash, const UClass* ObjectClass, const UObject* ObjectPackage, FName ObjectName, bool bExactClass, bool bAnyPackage, EObjectFlags ExcludeFlags, EInternalObjectFlags ExclusiveInternalFlags)
+UObject* StaticFindObjectFastInternalThreadSafe(FUObjectHashTables& ThreadHash, const UClass* ObjectClass, const UObject* ObjectPackage, FName ObjectName, bool bExactClass, EObjectFlags ExcludeFlags, EInternalObjectFlags ExclusiveInternalFlags)
 {
 	ExclusiveInternalFlags |= DefaultInternalExclusionFlags;
 
@@ -1113,7 +1112,7 @@ UObject* StaticFindObjectFastInternalThreadSafe(FUObjectHashTables& ThreadHash,
 					/*If there is no package (no InObjectPackage specified, and InName's package is "")
 					and the caller specified any_package, then accept it, regardless of its package.
 					Or, if the object is a top-level package then accept it immediately.*/
-					&& (bAnyPackage || !Object->GetOuter())
+					&& (!Object->GetOuter())
 
 					/** If a class was specified, check that the object is of the correct class */
 					&& (ObjectClass == nullptr || (bExactClass ? Object->GetClass() == ObjectClass : Object->IsA(ObjectClass)))
@@ -1148,25 +1147,6 @@ UObject* StaticFindObjectFastInternalThreadSafe(FUObjectHashTables& ThreadHash,
 	return Result;
 }
 
-UObject* StaticFindObjectFastInternal(const UClass* ObjectClass, const UObject* ObjectPackage, FName ObjectName, bool bExactClass, bool bAnyPackage, EObjectFlags ExcludeFlags, EInternalObjectFlags ExclusiveInternalFlags)
-{
-	UObject* Result = nullptr;
-
-	// Transactionally, a static find operation has to occur in the open as it touches shared state.
-	UE_AUTORTFM_OPEN
-	{
-		INC_DWORD_STAT(STAT_FindObjectFast);
-
-		check(!IS_ANY_PACKAGE_DEPRECATED(ObjectPackage)); // this could never have returned anything but nullptr
-
-		// If they specified an outer use that during the hashing
-		FUObjectHashTables& ThreadHash = FUObjectHashTables::Get();
-		Result = StaticFindObjectFastInternalThreadSafe(ThreadHash, ObjectClass, ObjectPackage, ObjectName, bExactClass, bAnyPackage, ExcludeFlags, ExclusiveInternalFlags);
-	};
-
-	return Result;
-}
-
 UObject* StaticFindObjectFastInternal(const UClass* ObjectClass, const UObject* ObjectPackage, FName ObjectName, bool bExactClass, EObjectFlags ExcludeFlags, EInternalObjectFlags ExclusiveInternalFlags)
 {
 	UObject* Result = nullptr;
@@ -1176,11 +1156,9 @@ UObject* StaticFindObjectFastInternal(const UClass* ObjectClass, const UObject*
 	{
 		INC_DWORD_STAT(STAT_FindObjectFast);
 
-		check(!IS_ANY_PACKAGE_DEPRECATED(ObjectPackage)); // this could never have returned anything but nullptr
-
 		// If they specified an outer use that during the hashing
 		FUObjectHashTables& ThreadHash = FUObjectHashTables::Get();
-		Result = StaticFindObjectFastInternalThreadSafe(ThreadHash, ObjectClass, ObjectPackage, ObjectName, bExactClass, /*bAnyPackage =*/ false, ExcludeFlags, ExclusiveInternalFlags);
+		Result = StaticFindObjectFastInternalThreadSafe(ThreadHash, ObjectClass, ObjectPackage, ObjectName, bExactClass, ExcludeFlags, ExclusiveInternalFlags);
 	};
 
 	return Result;
diff --git a/Engine/Source/Runtime/CoreUObject/Public/UObject/ObjectMacros.h b/Engine/Source/Runtime/CoreUObject/Public/UObject/ObjectMacros.h
--- a/Engine/Source/Runtime/CoreUObject/Public/UObject/ObjectMacros.h
+++ b/Engine/Source/Runtime/CoreUObject/Public/UObject/ObjectMacros.h
@@ -721,9 +721,6 @@ struct FReferencerInformationList
 	Core macros.
 ----------------------------------------------------------------------------*/
 
-// Special canonical package for FindObject, ParseObject.
-#define ANY_PACKAGE UE_DEPRECATED_MACRO(5.1, "ANY_PACKAGE has been deprecated. Either use full path name of objects (including classes) or provide a valid Outer for *FindObject* functions") ((UPackage*)-1)
-
 // Special prefix for default objects (the UObject in a UClass containing the default values, etc)
 #define DEFAULT_OBJECT_PREFIX TEXT("Default__")
 
diff --git a/Engine/Source/Runtime/CoreUObject/Public/UObject/UObjectGlobals.h b/Engine/Source/Runtime/CoreUObject/Public/UObject/UObjectGlobals.h
--- a/Engine/Source/Runtime/CoreUObject/Public/UObject/UObjectGlobals.h
+++ b/Engine/Source/Runtime/CoreUObject/Public/UObject/UObjectGlobals.h
@@ -294,23 +294,6 @@ COREUOBJECT_API bool SafeLoadError( UObject* Outer, uint32 LoadFlags, const TCHA
 /** Internal function used to update the suffix to be given to the next newly-created unnamed object. */
 COREUOBJECT_API int32 UpdateSuffixForNextNewObject(UObject* Parent, const UClass* Class, TFunctionRef<void(int32&)> IndexMutator);
 
-/**
- * Fast version of StaticFindObject that relies on the passed in FName being the object name without any group/package qualifiers.
- * This will only find top level packages or subobjects nested directly within a passed in outer.
- *
- * @param	Class			The to be found object's class
- * @param	InOuter			Outer object to look inside, if null this will only look for top level packages
- * @param	InName			Object name to look for relative to InOuter
- * @param	bExactClass		Whether to require an exact match with the passed in class
- * @param	bAnyPackage		Whether to look in any package
- * @param	ExclusiveFlags	Ignores objects that contain any of the specified exclusive flags
- * @param	ExclusiveInternalFlags  Ignores objects that contain any of the specified internal exclusive flags
- *
- * @return	Returns a pointer to the found object or null if none could be found
- */
-UE_DEPRECATED(5.1, "Support for searching for objects in ANY_PACKAGE has been deprecated. Please provide the actual Outer of an object you want to find.")
-COREUOBJECT_API UObject* StaticFindObjectFast(UClass* Class, UObject* InOuter, FName InName, bool bExactClass, bool bAnyPackage, EObjectFlags ExclusiveFlags = RF_NoFlags, EInternalObjectFlags ExclusiveInternalFlags = EInternalObjectFlags::None);
-
 /**
  * Fast version of StaticFindObject that relies on the passed in FName being the object name without any group/package qualifiers.
  * This will only find top level packages or subobjects nested directly within a passed in outer.
@@ -326,25 +309,6 @@ COREUOBJECT_API UObject* StaticFindObjectFast(UClass* Class, UObject* InOuter, F
  */
 COREUOBJECT_API UObject* StaticFindObjectFast(UClass* Class, UObject* InOuter, FName InName, bool bExactClass = false, EObjectFlags ExclusiveFlags = RF_NoFlags, EInternalObjectFlags ExclusiveInternalFlags = EInternalObjectFlags::None);
 
-/**
- * Fast and safe version of StaticFindObject that relies on the passed in FName being the object name without any group/package qualifiers.
- * It will not assert on GIsSavingPackage or IsGarbageCollectingAndLockingUObjectHashTables(). If called from within package saving code or GC, will return nullptr
- * This will only find top level packages or subobjects nested directly within a passed in outer.
- *
- * @param	Class			The to be found object's class
- * @param	InOuter			Outer object to look inside, if null this will only look for top level packages
- * @param	InName			Object name to look for relative to InOuter
- * @param	bExactClass		Whether to require an exact match with the passed in class
- * @param	bAnyPackage		Whether to look in any package
- * @param	ExclusiveFlags	Ignores objects that contain any of the specified exclusive flags
- * @param	ExclusiveInternalFlags  Ignores objects that contain any of the specified internal exclusive flags
- *
- * @return	Returns a pointer to the found object or null if none could be found
- */
-UE_DEPRECATED(5.1, "Support for searching for objects in ANY_PACKAGE has been deprecated. Please provide the actual Outer of an object you want to find.")
-COREUOBJECT_API UObject* StaticFindObjectFastSafe(UClass* Class, UObject* InOuter, FName InName, bool bExactClass, bool bAnyPackage, EObjectFlags ExclusiveFlags = RF_NoFlags, EInternalObjectFlags ExclusiveInternalFlags = EInternalObjectFlags::None);
-
-
 /**
  * Fast and safe version of StaticFindObject that relies on the passed in FName being the object name without any group/package qualifiers.
  * It will not assert on GIsSavingPackage or IsGarbageCollectingAndLockingUObjectHashTables(). If called from within package saving code or GC, will return nullptr
@@ -874,13 +838,6 @@ COREUOBJECT_API int32 LoadAssetAsync(FTopLevelAssetPath InAssetPath, FLoadAssetA
 */
 COREUOBJECT_API void CancelAsyncLoading();
 
-/**
-* Returns true if the event driven loader is enabled in cooked builds
-*/
-UE_DEPRECATED(5.0, "Any call to IsEventDrivenLoaderEnabledInCookedBuilds can be removed. "
-		"Cooked packages are always split into different files/segments and headers will always contain preload dependencies.")
-COREUOBJECT_API bool IsEventDrivenLoaderEnabledInCookedBuilds();
-
 /**
 * Returns true if the event driven loader is enabled in the current build
 */
@@ -895,12 +852,6 @@ COREUOBJECT_API bool IsEventDrivenLoaderEnabled();
  */
 COREUOBJECT_API float GetAsyncLoadPercentage( const FName& PackageName );
 
-/**
-* Whether we are running on the Garbage Collector Thread
-*/
-UE_DEPRECATED(5.5, "IsInGarbageCollectorThread() is now equivalent to calling IsInGameThread()")
-COREUOBJECT_API bool IsInGarbageCollectorThread();
-
 /** 
  * Deletes all unreferenced objects, keeping objects that have any of the passed in KeepFlags set. Will wait for other threads to unlock GC.
  *
@@ -1283,15 +1234,6 @@ public:
 	 */
 	COREUOBJECT_API FObjectInitializer(UObject* InObj, UObject* InObjectArchetype, EObjectInitializerOptions InOptions, struct FObjectInstancingGraph* InInstanceGraph = nullptr);
 
-	UE_DEPRECATED(5.0, "Use version that takes EObjectInitializerOptions")
-	FObjectInitializer(UObject* InObj, UObject* InObjectArchetype, bool bInCopyTransientsFromClassDefaults, bool bInShouldInitializeProps, struct FObjectInstancingGraph* InInstanceGraph = nullptr)
-		: FObjectInitializer(InObj, InObjectArchetype, 
-			(bInCopyTransientsFromClassDefaults ? EObjectInitializerOptions::CopyTransientsFromClassDefaults : EObjectInitializerOptions::None) |
-			(bInShouldInitializeProps ? EObjectInitializerOptions::InitializeProperties : EObjectInitializerOptions::None),
-			InInstanceGraph)
-	{
-	}
-
 	/** Special constructor for static construct object internal that passes along the params block directly */
 	COREUOBJECT_API FObjectInitializer(UObject* InObj, const FStaticConstructObjectParameters& StaticConstructParams, UE::RemoteObject::Serialization::FRemoteObjectConstructionOverrides* RemoteObjectOverides = nullptr);
 
@@ -1993,19 +1935,6 @@ inline bool ParseObject( const TCHAR* Stream, const TCHAR* Match, T*& Obj, UObje
 	return ParseObject( Stream, Match, T::StaticClass(), (UObject*&)Obj, Outer, EParseObjectLoadingPolicy::Find, bInvalidObject );
 }
 
-/** 
- * Find an optional object, relies on the name being unqualified 
- * @see StaticFindObjectFast()
- */
-template< class T > 
-UE_DEPRECATED(5.1, "Support for searching for objects in ANY_PACKAGE has been deprecated. Please provide the actual Outer of an object you want to find.")
-inline T* FindObjectFast( UObject* Outer, FName Name, bool ExactClass, bool AnyPackage, EObjectFlags ExclusiveFlags=RF_NoFlags )
-{
-	PRAGMA_DISABLE_DEPRECATION_WARNINGS
-	return (T*)StaticFindObjectFast( T::StaticClass(), Outer, Name, ExactClass, AnyPackage, ExclusiveFlags );
-	PRAGMA_ENABLE_DEPRECATION_WARNINGS
-}
-
 /**
  * Find an optional object, relies on the name being unqualified
  * @see StaticFindObjectFast()
@@ -3258,10 +3187,6 @@ struct FCoreUObjectDelegates
 	DECLARE_MULTICAST_DELEGATE_TwoParams(FOnObjectPostCDOCompiled, UObject*, const FObjectPostCDOCompiledContext&);
 	static COREUOBJECT_API FOnObjectPostCDOCompiled OnObjectPostCDOCompiled;
 
-	DECLARE_MULTICAST_DELEGATE_OneParam(FOnObjectSaved, UObject*);
-	UE_DEPRECATED(5.0, "Use OnObjectPreSave instead.")
-	static COREUOBJECT_API FOnObjectSaved OnObjectSaved;
-
 	/** Callback for when an asset is saved. This is called from UObject::PreSave before it is actually written to disk, for every object saved */
 	DECLARE_MULTICAST_DELEGATE_TwoParams(FOnObjectPreSave, UObject*, FObjectPreSaveContext);
 	static COREUOBJECT_API FOnObjectPreSave OnObjectPreSave;
@@ -3276,11 +3201,6 @@ struct FCoreUObjectDelegates
 	/** Callback when packages end loading in LoadPackage or AsyncLoadPackage. All packages loaded recursively due to imports are included in the single call of the explicitly-loaded package. It is called in all `WITH_EDITOR` cases, including -game. But it is not called in the runtime game compiled without WITH_EDITOR. This difference in behavior between the two game configurations is necessary because WITH_EDITOR uses a different loading mechanism than runtime game. */
 	DECLARE_MULTICAST_DELEGATE_OneParam(FOnEndLoadPackage, const FEndLoadPackageContext&);
 	static COREUOBJECT_API FOnEndLoadPackage OnEndLoadPackage;
-
-	/** Delegate used by SavePackage() to create the package backup */
-	DECLARE_DELEGATE_RetVal_OneParam(bool, FAutoPackageBackupDelegate, const UPackage&);
-	UE_DEPRECATED(5.1, "Backups are no longer used in Unreal Package Saves")
-	static COREUOBJECT_API FAutoPackageBackupDelegate AutoPackageBackupDelegate;
 #endif // WITH_EDITOR
 
 	/** Called when new sparse class data has been created (and the base data initialized) for the given class */
@@ -3311,11 +3231,6 @@ struct FCoreUObjectDelegates
 	DECLARE_MULTICAST_DELEGATE_OneParam(FReloadCompleteDelegate, EReloadCompleteReason);
 	static COREUOBJECT_API FReloadCompleteDelegate ReloadCompleteDelegate;
 
-	/** Delegate for registering hot-reloaded classes that have been added  */
-	DECLARE_MULTICAST_DELEGATE_OneParam(FRegisterHotReloadAddedClassesDelegate, const TArray<UClass*>&);
-	UE_DEPRECATED(5.0, "RegisterHotReloadAddedClassesDelegate has been deprecated, use ReloadAddedClassesDelegate.")
-	static COREUOBJECT_API FRegisterHotReloadAddedClassesDelegate RegisterHotReloadAddedClassesDelegate;
-
 	/** Delegate for registering hot-reloaded classes that changed after hot-reload for reinstancing */
 	DECLARE_MULTICAST_DELEGATE_ThreeParams(FRegisterClassForHotReloadReinstancingDelegate, UClass*, UClass*, EHotReloadedClassFlags);
 	UE_DEPRECATED(5.0, "RegisterClassForHotReloadReinstancingDelegate has been deprecated, use FReload for class re-instancing.")
@@ -3360,10 +3275,6 @@ struct FCoreUObjectDelegates
 	/** Delegate type for reachability analysis external roots callback. First parameter is FGarbageCollectionTracer to use for tracing, second is flags with which objects should be kept alive regardless, third is whether to force single threading */
 	DECLARE_MULTICAST_DELEGATE_ThreeParams(FTraceExternalRootsForReachabilityAnalysisDelegate, FGarbageCollectionTracer&, EObjectFlags, bool);
 
-	/** Called as last phase of reachability analysis. Allow external systems to add UObject roots *after* first reachability pass has been done */
-	UE_DEPRECATED(5.2, "Ability to inject extra root objects will be removed to reduce GC complexity")
-	static COREUOBJECT_API FTraceExternalRootsForReachabilityAnalysisDelegate TraceExternalRootsForReachabilityAnalysis;
-
 	/** Called after reachability analysis, before any purging */
 	static COREUOBJECT_API FSimpleMulticastDelegate PostReachabilityAnalysis;
 
@@ -3531,8 +3442,6 @@ enum class EEditorOnlyObjectResult
  * - if bCheckRecursive is true, if its class, super struct, outer, or archetypes are editor only
  */
 COREUOBJECT_API bool IsEditorOnlyObject(const UObject* InObject, bool bCheckRecursive = true);
-UE_DEPRECATED(5.3, "bCheckMarks argument is no longer supported because we are transitioning away from using ObjectMarks during saving");
-COREUOBJECT_API bool IsEditorOnlyObject(const UObject* InObject, bool bCheckRecursive, bool bCheckMarks);
 COREUOBJECT_API bool IsEditorOnlyObject(const UObject* InObject, bool bCheckRecursive,
 	TFunctionRef<UE::SavePackageUtilities::EEditorOnlyObjectResult(const UObject* Object)> LookupInCache,
 	TFunctionRef<void(const UObject* Object, bool bEditorOnly)> AddToCache);
diff --git a/Engine/Source/Runtime/CoreUObject/Public/UObject/UObjectHash.h b/Engine/Source/Runtime/CoreUObject/Public/UObject/UObjectHash.h
--- a/Engine/Source/Runtime/CoreUObject/Public/UObject/UObjectHash.h
+++ b/Engine/Source/Runtime/CoreUObject/Public/UObject/UObjectHash.h
@@ -30,21 +30,6 @@ DECLARE_STATS_GROUP_VERBOSE(TEXT("UObject Hash"), STATGROUP_UObjectHash, STATCAT
 DECLARE_DWORD_COUNTER_STAT_EXTERN(TEXT("NumObjects"), STAT_Hash_NumObjects, STATGROUP_UObjectHash, COREUOBJECT_API);
 #endif // !UE_BUILD_TEST && !UE_BUILD_SHIPPING
 
-/**
- * Private internal version of StaticFindObjectFast that allows using 0 exclusion flags.
- *
- * @param	Class			The to be found object's class
- * @param	InOuter			The to be found object's outer
- * @param	InName			The to be found object's class
- * @param	ExactClass		Whether to require an exact match with the passed in class
- * @param	AnyPackage		Whether to look in any package
- * @param	ExclusiveFlags	Ignores objects that contain any of the specified exclusive flags
- * @param	ExclusiveInternalFlags	Ignores objects that contain any of the specified internal exclusive flags
- * @return	Returns a pointer to the found object or NULL if none could be found
- */
-UE_DEPRECATED(5.1, "Support for searching for objects in ANY_PACKAGE has been deprecated. Please provide the actual Outer of an object you want to find.")
-COREUOBJECT_API UObject* StaticFindObjectFastInternal(const UClass* Class, const UObject* InOuter, FName InName, bool ExactClass, bool AnyPackage, EObjectFlags ExclusiveFlags = RF_NoFlags, EInternalObjectFlags ExclusiveInternalFlags = EInternalObjectFlags::None);
-
 /**
  * Private internal version of StaticFindObjectFast that allows using 0 exclusion flags.
  *
