Clarified deprecation of RegisterClassForHotReloadReinstancingDelegate and ReinstanceHotReloadedClassesDelegate.

diff --git a/Engine/Source/Runtime/CoreUObject/Private/UObject/UObjectGlobals.cpp b/Engine/Source/Runtime/CoreUObject/Private/UObject/UObjectGlobals.cpp
--- a/Engine/Source/Runtime/CoreUObject/Private/UObject/UObjectGlobals.cpp
+++ b/Engine/Source/Runtime/CoreUObject/Private/UObject/UObjectGlobals.cpp
@@ -179,8 +179,8 @@ namespace LoadPackageStats
 
 /** CoreUObject delegates */
 PRAGMA_DISABLE_DEPRECATION_WARNINGS
-FCoreUObjectDelegates::FRegisterClassForHotReloadReinstancingDelegate FCoreUObjectDelegates::RegisterClassForHotReloadReinstancingDelegate;
-FCoreUObjectDelegates::FReinstanceHotReloadedClassesDelegate FCoreUObjectDelegates::ReinstanceHotReloadedClassesDelegate;
+TMulticastDelegate<void(UClass* OldClass, UClass* NewClass, EHotReloadedClassFlags Flags)> FCoreUObjectDelegates::RegisterClassForHotReloadReinstancingDelegate;
+TMulticastDelegate<void()> FCoreUObjectDelegates::ReinstanceHotReloadedClassesDelegate;
 PRAGMA_ENABLE_DEPRECATION_WARNINGS
 FCoreUObjectDelegates::FReloadReinstancingCompleteDelegate FCoreUObjectDelegates::ReloadReinstancingCompleteDelegate;
 FCoreUObjectDelegates::FReloadCompleteDelegate FCoreUObjectDelegates::ReloadCompleteDelegate;
diff --git a/Engine/Source/Runtime/CoreUObject/Public/UObject/UObjectGlobals.h b/Engine/Source/Runtime/CoreUObject/Public/UObject/UObjectGlobals.h
--- a/Engine/Source/Runtime/CoreUObject/Public/UObject/UObjectGlobals.h
+++ b/Engine/Source/Runtime/CoreUObject/Public/UObject/UObjectGlobals.h
@@ -3232,14 +3232,12 @@ struct FCoreUObjectDelegates
 	static COREUOBJECT_API FReloadCompleteDelegate ReloadCompleteDelegate;
 
 	/** Delegate for registering hot-reloaded classes that changed after hot-reload for reinstancing */
-	DECLARE_MULTICAST_DELEGATE_ThreeParams(FRegisterClassForHotReloadReinstancingDelegate, UClass*, UClass*, EHotReloadedClassFlags);
-	UE_DEPRECATED(5.0, "RegisterClassForHotReloadReinstancingDelegate has been deprecated, use FReload for class re-instancing.")
-	static COREUOBJECT_API FRegisterClassForHotReloadReinstancingDelegate RegisterClassForHotReloadReinstancingDelegate;
+	UE_DEPRECATED(all, "RegisterClassForHotReloadReinstancingDelegate is for internal use only.")
+	static COREUOBJECT_API TMulticastDelegate<void(UClass* OldClass, UClass* NewClass, EHotReloadedClassFlags Flags)> RegisterClassForHotReloadReinstancingDelegate;
 
 	/** Delegate for reinstancing hot-reloaded classes */
-	DECLARE_MULTICAST_DELEGATE(FReinstanceHotReloadedClassesDelegate);
-	UE_DEPRECATED(5.0, "ReinstanceHotReloadedClassesDelegate has been deprecated, use FReload for class re-instancing or ReloadReinstancingCompleteDelegate for notification")
-	static COREUOBJECT_API FReinstanceHotReloadedClassesDelegate ReinstanceHotReloadedClassesDelegate;
+	UE_DEPRECATED(all, "ReinstanceHotReloadedClassesDelegate is for internal use only")
+	static COREUOBJECT_API TMulticastDelegate<void()> ReinstanceHotReloadedClassesDelegate;
 
 	/** 
 	 * Delegate invoked when requests are made to process pending UObject registrations.  This will be invoked during the loading of modules
