EditorBulkData: Fix bug when an EditorBulkData was loaded from one package (e.g. a Template) and moved into another package (e.g. a saved version of the template), where it gets a new guid, and the original package is garbage collected, and then the new package is saved and the EditorBulkData needs to update its payload (e.g. because texture compression has changed). Previously we would incorrectly think the resaved EditorBulkData was from a new package and was a collision, because when its linker was garbage collected we unregistered it even though it was still loaded. Keep it registered, and just in a state of no longer having its linker and in-memory payload available. This allows us to properly unregister it and reregister it later when we we need to update its payload. This fixes occurrences of the warning "LogBulkDataRegistry: Warning: {Package} failed to register BulkData {guid} from SerializeToPackageTrailer in the BulkDataRegistry: it already exists with owner {TempPackage}."
#rnx

diff --git a/Engine/Source/Runtime/CoreUObject/Private/Serialization/EditorBulkData.cpp b/Engine/Source/Runtime/CoreUObject/Private/Serialization/EditorBulkData.cpp
--- a/Engine/Source/Runtime/CoreUObject/Private/Serialization/EditorBulkData.cpp
+++ b/Engine/Source/Runtime/CoreUObject/Private/Serialization/EditorBulkData.cpp
@@ -470,6 +470,9 @@ FEditorBulkData::~FEditorBulkData()
 		AttachedAr = nullptr;
 	}
 
+	// If we have registered with the BulkData registry, keep our ownership of our guid registered,
+	// and possibly loadable from our on-disk location, but let the BulkDataRegistry know that we
+	// no longer have an in-memory copy of the data.
 	OnExitMemory();
 }
 
@@ -672,7 +675,6 @@ void FEditorBulkData::OnExitMemory()
 	{
 		check(!EnumHasAnyFlags(Flags, EFlags::IsTornOff));
 		IBulkDataRegistry::Get().OnExitMemory(*this);
-		EnumRemoveFlags(Flags, EFlags::HasRegistered);
 	}
 #endif
 }
@@ -1785,7 +1787,7 @@ void FEditorBulkData::DetachFromDisk(FArchive* Ar, bool bEnsurePayloadIsLoaded)
 			else
 			{
 				// The package is unloading; we should instruct the bulkdataregistry to keep our old packagepath information,
-				// but will drop our is-registered flag
+				// but let it know that we no longer have an in-memory copy of the data.
 				OnExitMemory();
 			}
 		}
