[UE] Allow users to change the derived type on default value instantiations used to initialize a base type reference member in saved instance data when using override serialization.

Notes:
- This can result in data loss if the user had e.g. modified data on the subobject instance of the previous type if it is not referenced elsewhere. However, this is better than crashing.
- Potentially will further address the data loss issue using IDOs in a future task as part of fixup work.

diff --git a/Engine/Source/Runtime/CoreUObject/Private/UObject/CoreNative.cpp b/Engine/Source/Runtime/CoreUObject/Private/UObject/CoreNative.cpp
--- a/Engine/Source/Runtime/CoreUObject/Private/UObject/CoreNative.cpp
+++ b/Engine/Source/Runtime/CoreUObject/Private/UObject/CoreNative.cpp
@@ -276,6 +276,30 @@ UObject* FObjectInstancingGraph::GetInstancedSubobject( UObject* SourceSubobject
 
 							FName SubobjectName = SourceSubobject->GetFName();
 
+							// If a property serialized a reference to an instanced subobject and it is not type-compatible with the default value that the
+							// serializing (owner) object should be referencing at this point on load, the serialized object needs to be verified against
+							// the value the instancing graph will use during subobject instancing, which gets deferred until PostLoadSubobjects() on load.
+							if (IsLoadingObject()
+								&& CurrentValue == SourceSubobject
+								&& FOverridableManager::Get().IsEnabled(SubobjectOuter)
+								&& SubobjectOuter->GetClass()->HasAnyClassFlags(CLASS_CompiledFromBlueprint))
+							{
+								// Look for an existing instance with the same name at the outer scope. If we find one, we can infer that it's a subobject
+								// that has not been explicitly overridden, but that was serialized on save as a default-instanced subobject. In that case,
+								// we need to check the object's type against the type that we're about to instance below. If there's a mismatch, rename
+								// the existing object out of the way, as subobject recycling will assert in that case. Any serialized data will be lost.
+								UObject* ExistingObject = StaticFindObjectFast(nullptr, SubobjectOuter, SubobjectName);
+								if (ExistingObject
+									&& !ExistingObject->GetClass()->IsChildOf(SourceSubobject->GetClass()))
+								{
+									// Keep the existing base name and scope in case some other property was overridden and serialized this reference.
+									// If nothing else winds up referencing this object after load, the export will be garbage collected (which would
+									// happen whether or not we renamed the object below - that is being done in order to free up the name for instancing).
+									FName NewName = MakeUniqueObjectName(SubobjectOuter, ExistingObject->GetClass(), SubobjectName);
+									ExistingObject->Rename(*NewName.ToString(), nullptr, REN_DoNotDirty | REN_NonTransactional | REN_DontCreateRedirectors);
+								}
+							}
+
 							// Don't search for the existing subobjects on Blueprint-generated classes. What we'll find is a subobject
 							// created by the constructor which may not have all of its fields initialized to the correct value (which
 							// should be coming from a blueprint).
