[UE] Include IDO placeholder support in the "all" IDO features toggle set.

diff --git a/Engine/Source/Runtime/CoreUObject/Internal/UObject/InstanceDataObjectUtils.h b/Engine/Source/Runtime/CoreUObject/Internal/UObject/InstanceDataObjectUtils.h
--- a/Engine/Source/Runtime/CoreUObject/Internal/UObject/InstanceDataObjectUtils.h
+++ b/Engine/Source/Runtime/CoreUObject/Internal/UObject/InstanceDataObjectUtils.h
@@ -35,6 +35,8 @@ UE_API bool IsUninitializedAlertUIEnabled();
 UE_API bool IsInstanceDataObjectImpersonationEnabledOnSave();
 /** Query if InstanceDataObjects should be constructed using an archetype chain. */
 UE_API bool IsInstanceDataObjectArchetypeChainEnabled();
+/** Query if InstanceDataObjects should impose placeholder objects for missing types. */
+UE_API bool IsInstanceDataObjectPlaceholderObjectSupportEnabled();
 
 UE_API bool StructContainsLooseProperties(const UStruct* Struct);
 UE_API bool StructIsInstanceDataObjectStruct(const UStruct* Struct);
diff --git a/Engine/Source/Runtime/CoreUObject/Private/UObject/InstanceDataObjectUtils.cpp b/Engine/Source/Runtime/CoreUObject/Private/UObject/InstanceDataObjectUtils.cpp
--- a/Engine/Source/Runtime/CoreUObject/Private/UObject/InstanceDataObjectUtils.cpp
+++ b/Engine/Source/Runtime/CoreUObject/Private/UObject/InstanceDataObjectUtils.cpp
@@ -349,6 +349,14 @@ namespace UE
 		TEXT("When enabled, IDOs will be constructed using an archetype chain")
 	);
 
+	bool bEnableIDOPlaceholderObjects = false;
+	FAutoConsoleVariableRef EnableIDOPlaceholderObjectsCVar(
+		TEXT("IDO.Placeholder.Enable"),
+		bEnableIDOPlaceholderObjects,
+		TEXT("If true, allows placeholder types to be created in place of missing types in order to allow instance data to be serialized."),
+		ECVF_Default
+	);
+
 	FAutoConsoleCommand EnableIDOUnknownsCommand(
 		TEXT("IDO.Unknowns.Enable"),
 		TEXT("Use this command to toggle IDO.Unknowns.* on or off, or to report their current state."),
@@ -378,6 +386,7 @@ namespace UE
 		EnableIDOsForBlueprintArchetypesCVar->Set(bEnabled);
 		EnableIDOsForBlueprintInstancesCVar->Set(bEnabled);
 		EnableIDOArchetypeChainCVar->Set(bEnabled);
+		EnableIDOPlaceholderObjectsCVar->Set(bEnabled);
 	}
 
 	static bool AreAllIDOFeaturesEnabled()
@@ -391,7 +400,8 @@ namespace UE
 			bEnableIDOImpersonationOnSave &&
 			bEnableIDOsForBlueprintArchetypes &&
 			bEnableIDOsForBlueprintInstances &&
-			bEnableIDOArchetypeChain;
+			bEnableIDOArchetypeChain &&
+			bEnableIDOPlaceholderObjects;
 	}
 
 	FAutoConsoleCommand EnableAllIDOFeaturesCommand(
@@ -538,6 +548,11 @@ namespace UE
 		return bEnableIDOArchetypeChain;
 	}
 
+	bool IsInstanceDataObjectPlaceholderObjectSupportEnabled()
+	{
+		return bEnableIDOPlaceholderObjects;
+	}
+
 	static const UObject* GetBlueprintGeneratedObject(const UObject* InObject)
 	{
 		static const FName NAME_BlueprintGeneratedClass(TEXT("BlueprintGeneratedClass"));
diff --git a/Engine/Source/Runtime/CoreUObject/Private/UObject/PropertyBagRepository.cpp b/Engine/Source/Runtime/CoreUObject/Private/UObject/PropertyBagRepository.cpp
--- a/Engine/Source/Runtime/CoreUObject/Private/UObject/PropertyBagRepository.cpp
+++ b/Engine/Source/Runtime/CoreUObject/Private/UObject/PropertyBagRepository.cpp
@@ -940,14 +940,6 @@ bool FPropertyBagRepository::IsPropertyBagPlaceholderObject(const UObject* Objec
 namespace Private
 {
 #if WITH_EDITOR
-	static bool bEnablePropertyBagPlaceholderObjectSupport = false;
-	static FAutoConsoleVariableRef CVarEnablePropertyBagPlaceholderObjectSupport(
-		TEXT("IDO.Placeholder.Enable"),
-		bEnablePropertyBagPlaceholderObjectSupport,
-		TEXT("If true, allows placeholder types to be created in place of missing types in order to redirect serialization into a property bag."),
-		ECVF_Default
-	);
-
 	using EPlaceholderObjectFeature = UE::FPropertyBagRepository::EPlaceholderObjectFeature;
 	constexpr uint32 PropertyBagPlaceholderObjectFeatureFlag(EPlaceholderObjectFeature Feature)
 	{
@@ -987,7 +979,7 @@ namespace Private
 bool FPropertyBagRepository::IsPropertyBagPlaceholderObjectSupportEnabled()
 {
 #if WITH_EDITOR && UE_WITH_OBJECT_HANDLE_TYPE_SAFETY
-	return Private::bEnablePropertyBagPlaceholderObjectSupport && IsInstanceDataObjectSupportEnabled();
+	return IsInstanceDataObjectSupportEnabled() && IsInstanceDataObjectPlaceholderObjectSupportEnabled();
 #else
 	return false;
 #endif
