Added EFindObjectFlags to replace bExactClass bool parameters for the FindObject family of functions.

diff --git a/Engine/Source/Runtime/CoreUObject/Private/UObject/UObjectGlobals.cpp b/Engine/Source/Runtime/CoreUObject/Private/UObject/UObjectGlobals.cpp
--- a/Engine/Source/Runtime/CoreUObject/Private/UObject/UObjectGlobals.cpp
+++ b/Engine/Source/Runtime/CoreUObject/Private/UObject/UObjectGlobals.cpp
@@ -295,7 +295,7 @@ namespace
 	 *
 	 * @returns Found object.
 	 */
-	UObject* StaticFindObjectWithChangedLegacyPath(UClass* ObjectClass, UObject* ObjectPackage, FName ObjectName, bool ExactClass, EInternalObjectFlags ExclusiveInternalFlags)
+	UObject* StaticFindObjectWithChangedLegacyPath(UClass* ObjectClass, UObject* ObjectPackage, FName ObjectName, EFindObjectFlags Flags, EInternalObjectFlags ExclusiveInternalFlags)
 	{
 		UObject* MatchingObject = nullptr;
 
@@ -312,7 +312,7 @@ namespace
 				|| ObjectClass == UScriptStruct::StaticClass() || ObjectClass == UStruct::StaticClass() // Structs
 				|| (ObjectClass == UFunction::StaticClass() && bHasDelegateSignaturePostfix); // Delegates
 
-			bool bSubclassOfPathChangedClass = !bExactPathChangedClass && !ExactClass
+			bool bSubclassOfPathChangedClass = !bExactPathChangedClass && !(Flags & EFindObjectFlags::ExactClass)
 				&& (ObjectClass == nullptr // Any class
 				|| UEnum::StaticClass()->IsChildOf(ObjectClass) // Enums
 				|| UScriptStruct::StaticClass()->IsChildOf(ObjectClass) || UStruct::StaticClass()->IsChildOf(ObjectClass) // Structs
@@ -323,7 +323,7 @@ namespace
 				return nullptr;
 			}
 
-			MatchingObject = StaticFindObject(ObjectClass, ObjectPackage->GetOutermost(), *ObjectName.ToString(), ExactClass);
+			MatchingObject = StaticFindObject(ObjectClass, ObjectPackage->GetOutermost(), *ObjectName.ToString(), Flags);
 
 			if (MatchingObject)
 			{
@@ -444,24 +444,24 @@ int32 UpdateSuffixForNextNewObject(UObject* Parent, const UClass* Class, TFuncti
 //
 // Find an object, path must unqualified
 //
-UObject* StaticFindObjectFast(UClass* ObjectClass, UObject* ObjectPackage, FName ObjectName, bool bExactClass, EObjectFlags ExclusiveFlags, EInternalObjectFlags ExclusiveInternalFlags)
+UObject* StaticFindObjectFast(UClass* ObjectClass, UObject* ObjectPackage, FName ObjectName, EFindObjectFlags Flags, EObjectFlags ExclusiveFlags, EInternalObjectFlags ExclusiveInternalFlags)
 {
 	UE_CLOG(UE::IsSavingPackage(nullptr), LogUObjectGlobals, Fatal, TEXT("Illegal call to StaticFindObjectFast() while serializing object data!"));
 	UE_CLOG(IsGarbageCollectingAndLockingUObjectHashTables(), LogUObjectGlobals, Fatal, TEXT("Illegal call to StaticFindObjectFast() while garbage collecting!"));
 
 	// We don't want to return any objects that are currently being background loaded unless we're using FindObject during async loading.
 	ExclusiveInternalFlags |= UE::GetAsyncLoadingInternalFlagsExclusion();
-	UObject* FoundObject = StaticFindObjectFastInternal(ObjectClass, ObjectPackage, ObjectName, bExactClass, ExclusiveFlags, ExclusiveInternalFlags);
+	UObject* FoundObject = StaticFindObjectFastInternal(ObjectClass, ObjectPackage, ObjectName, Flags, ExclusiveFlags, ExclusiveInternalFlags);
 
 	if (!FoundObject)
 	{
-		FoundObject = StaticFindObjectWithChangedLegacyPath(ObjectClass, ObjectPackage, ObjectName, bExactClass, ExclusiveInternalFlags);
+		FoundObject = StaticFindObjectWithChangedLegacyPath(ObjectClass, ObjectPackage, ObjectName, Flags, ExclusiveInternalFlags);
 	}
 
 	return FoundObject;
 }
 
-UObject* StaticFindObjectFastSafe(UClass* ObjectClass, UObject* ObjectPackage, FName ObjectName, bool bExactClass, EObjectFlags ExclusiveFlags, EInternalObjectFlags ExclusiveInternalFlags)
+UObject* StaticFindObjectFastSafe(UClass* ObjectClass, UObject* ObjectPackage, FName ObjectName, EFindObjectFlags Flags, EObjectFlags ExclusiveFlags, EInternalObjectFlags ExclusiveInternalFlags)
 {
 	UObject* FoundObject = nullptr;
 
@@ -469,10 +469,10 @@ UObject* StaticFindObjectFastSafe(UClass* ObjectClass, UObject* ObjectPackage, F
 	{
 		// We don't want to return any objects that are currently being background loaded unless we're using FindObject during async loading.
 		ExclusiveInternalFlags |= UE::GetAsyncLoadingInternalFlagsExclusion();;
-		FoundObject = StaticFindObjectFastInternal(ObjectClass, ObjectPackage, ObjectName, bExactClass, ExclusiveFlags, ExclusiveInternalFlags);
+		FoundObject = StaticFindObjectFastInternal(ObjectClass, ObjectPackage, ObjectName, Flags, ExclusiveFlags, ExclusiveInternalFlags);
 		if (!FoundObject)
 		{
-			FoundObject = StaticFindObjectWithChangedLegacyPath(ObjectClass, ObjectPackage, ObjectName, bExactClass, ExclusiveInternalFlags);
+			FoundObject = StaticFindObjectWithChangedLegacyPath(ObjectClass, ObjectPackage, ObjectName, Flags, ExclusiveInternalFlags);
 		}
 	}
 
@@ -508,7 +508,7 @@ static UObject* LoadObjectWhenImportingT3D(UClass* ObjectClass, const TCHAR* Ori
 //
 // Find an optional object.
 //
-UObject* StaticFindObject( UClass* ObjectClass, UObject* InObjectPackage, const TCHAR* OrigInName, bool bExactClass )
+UObject* StaticFindObject( UClass* ObjectClass, UObject* InObjectPackage, const TCHAR* OrigInName, EFindObjectFlags Flags )
 {
 	INC_DWORD_STAT(STAT_FindObject);
 
@@ -540,14 +540,14 @@ UObject* StaticFindObject( UClass* ObjectClass, UObject* InObjectPackage, const
 	FName ObjectName(InName.ToView(), FNAME_Add);
 	
 	PRAGMA_DISABLE_DEPRECATION_WARNINGS
-	return StaticFindObjectFast(ObjectClass, ObjectPackage, ObjectName, bExactClass);
+	return StaticFindObjectFast(ObjectClass, ObjectPackage, ObjectName, Flags);
 	PRAGMA_ENABLE_DEPRECATION_WARNINGS
 }
 
 //
 // Find an optional object.
 //
-UObject* StaticFindObject(UClass* Class, FTopLevelAssetPath ObjectPath, bool ExactClass /*= false*/)
+UObject* StaticFindObject(UClass* Class, FTopLevelAssetPath ObjectPath, EFindObjectFlags Flags)
 {
 	UObject* Result = nullptr;
 	if (!ObjectPath.IsNull())
@@ -555,18 +555,18 @@ UObject* StaticFindObject(UClass* Class, FTopLevelAssetPath ObjectPath, bool Exa
 		UObject* Package = StaticFindObjectFast(UPackage::StaticClass(), nullptr, ObjectPath.GetPackageName());
 		if (Package)
 		{
-			Result = StaticFindObjectFast(Class, Package, ObjectPath.GetAssetName(), ExactClass);
+			Result = StaticFindObjectFast(Class, Package, ObjectPath.GetAssetName(), Flags);
 		}
 	}
 	return Result;
 }
 
-UObject* StaticFindObjectSafe(UClass* Class, FTopLevelAssetPath ObjectPath, bool ExactClass /*= false*/)
+UObject* StaticFindObjectSafe(UClass* Class, FTopLevelAssetPath ObjectPath, EFindObjectFlags Flags)
 {
 	if (!UE::IsSavingPackage(nullptr) && !IsGarbageCollectingAndLockingUObjectHashTables())
 	{
 		FGCScopeGuard GCGuard;
-		return StaticFindObject(Class, ObjectPath, ExactClass);
+		return StaticFindObject(Class, ObjectPath, Flags);
 	}
 	else
 	{
@@ -577,9 +577,9 @@ UObject* StaticFindObjectSafe(UClass* Class, FTopLevelAssetPath ObjectPath, bool
 //
 // Find an object; can't fail.
 //
-UObject* StaticFindObjectChecked( UClass* ObjectClass, UObject* ObjectParent, const TCHAR* InName, bool ExactClass )
+UObject* StaticFindObjectChecked( UClass* ObjectClass, UObject* ObjectParent, const TCHAR* InName, EFindObjectFlags Flags )
 {
-	UObject* Result = StaticFindObject( ObjectClass, ObjectParent, InName, ExactClass );
+	UObject* Result = StaticFindObject( ObjectClass, ObjectParent, InName, Flags );
 #if !(UE_BUILD_SHIPPING || UE_BUILD_TEST)
 	if( !Result )
 	{
@@ -592,12 +592,12 @@ UObject* StaticFindObjectChecked( UClass* ObjectClass, UObject* ObjectParent, co
 //
 // Find an object; won't assert on UE::IsSavingPackage() or IsGarbageCollectingAndLockingUObjectHashTables()
 //
-UObject* StaticFindObjectSafe( UClass* ObjectClass, UObject* ObjectParent, const TCHAR* InName, bool bExactClass )
+UObject* StaticFindObjectSafe( UClass* ObjectClass, UObject* ObjectParent, const TCHAR* InName, EFindObjectFlags Flags )
 {
 	if (!UE::IsSavingPackage(nullptr) && !IsGarbageCollectingAndLockingUObjectHashTables())
 	{
 		FGCScopeGuard GCGuard;
-		return StaticFindObject( ObjectClass, ObjectParent, InName, bExactClass );
+		return StaticFindObject( ObjectClass, ObjectParent, InName, Flags );
 	}
 	else
 	{
@@ -605,28 +605,28 @@ UObject* StaticFindObjectSafe( UClass* ObjectClass, UObject* ObjectParent, const
 	}
 }
 
-bool StaticFindAllObjectsFast(TArray<UObject*>& OutFoundObjects, UClass* ObjectClass, FName ObjectName, bool ExactClass, EObjectFlags ExclusiveFlags, EInternalObjectFlags ExclusiveInternalFlags)
+bool StaticFindAllObjectsFast(TArray<UObject*>& OutFoundObjects, UClass* ObjectClass, FName ObjectName, EFindObjectFlags Flags, EObjectFlags ExclusiveFlags, EInternalObjectFlags ExclusiveInternalFlags)
 {
 	UE_CLOG(UE::IsSavingPackage(nullptr) || IsGarbageCollectingAndLockingUObjectHashTables(), LogUObjectGlobals, Fatal, TEXT("Illegal call to StaticFindAllObjectsFast() while serializing object data or garbage collecting!"));
 
 	// We don't want to return any objects that are currently being background loaded unless we're using FindObject during async loading.
 	ExclusiveInternalFlags |= UE::GetAsyncLoadingInternalFlagsExclusion();
-	return StaticFindAllObjectsFastInternal(OutFoundObjects, ObjectClass, ObjectName, ExactClass, ExclusiveFlags, ExclusiveInternalFlags);
+	return StaticFindAllObjectsFastInternal(OutFoundObjects, ObjectClass, ObjectName, Flags, ExclusiveFlags, ExclusiveInternalFlags);
 }
 
-bool StaticFindAllObjectsFastSafe(TArray<UObject*>& OutFoundObjects, UClass* ObjectClass, FName ObjectName, bool ExactClass, EObjectFlags ExclusiveFlags, EInternalObjectFlags ExclusiveInternalFlags)
+bool StaticFindAllObjectsFastSafe(TArray<UObject*>& OutFoundObjects, UClass* ObjectClass, FName ObjectName, EFindObjectFlags Flags, EObjectFlags ExclusiveFlags, EInternalObjectFlags ExclusiveInternalFlags)
 {
 	bool bFoundObjects = false;
 	if (!UE::IsSavingPackage(nullptr) && !IsGarbageCollectingAndLockingUObjectHashTables())
 	{
 		// We don't want to return any objects that are currently being background loaded unless we're using FindObject during async loading.
 		ExclusiveInternalFlags |= UE::GetAsyncLoadingInternalFlagsExclusion();
-		bFoundObjects = StaticFindAllObjectsFastInternal(OutFoundObjects, ObjectClass, ObjectName, ExactClass, ExclusiveFlags, ExclusiveInternalFlags);
+		bFoundObjects = StaticFindAllObjectsFastInternal(OutFoundObjects, ObjectClass, ObjectName, Flags, ExclusiveFlags, ExclusiveInternalFlags);
 	}
 	return bFoundObjects;
 }
 
-bool StaticFindAllObjects(TArray<UObject*>& OutFoundObjects, UClass* ObjectClass, const TCHAR* OrigInName, bool ExactClass)
+bool StaticFindAllObjects(TArray<UObject*>& OutFoundObjects, UClass* ObjectClass, const TCHAR* OrigInName, EFindObjectFlags Flags)
 {
 	INC_DWORD_STAT(STAT_FindObject);
 
@@ -652,15 +652,15 @@ bool StaticFindAllObjects(TArray<UObject*>& OutFoundObjects, UClass* ObjectClass
 	ConstructorHelpers_StripObjectClass2(InName);
 	FName ObjectName(InName.ToView(), FNAME_Add);
 
-	return StaticFindAllObjectsFast(OutFoundObjects, ObjectClass, ObjectName, ExactClass);
+	return StaticFindAllObjectsFast(OutFoundObjects, ObjectClass, ObjectName, Flags);
 }
 
-bool StaticFindAllObjectsSafe(TArray<UObject*>& OutFoundObjects, UClass* ObjectClass, const TCHAR* OrigInName, bool ExactClass)
+bool StaticFindAllObjectsSafe(TArray<UObject*>& OutFoundObjects, UClass* ObjectClass, const TCHAR* OrigInName, EFindObjectFlags Flags)
 {
 	bool bFoundObjects = false;
 	if (!UE::IsSavingPackage(nullptr) && !IsGarbageCollectingAndLockingUObjectHashTables())
 	{
-		bFoundObjects = StaticFindAllObjects(OutFoundObjects, ObjectClass, OrigInName, ExactClass);
+		bFoundObjects = StaticFindAllObjects(OutFoundObjects, ObjectClass, OrigInName, Flags);
 	}
 	return bFoundObjects;
 }
@@ -682,17 +682,18 @@ UObject* StaticFindFirstObject(UClass* Class, const TCHAR* Name, EFindFirstObjec
 		ObjectName = FName(InName.ToView(), FNAME_Add);
 	}
 
+	EFindObjectFlags FindObjectFlags = !!(Options & EFindFirstObjectOptions::ExactClass) ? EFindObjectFlags::ExactClass : EFindObjectFlags::None;
 	if (AmbiguousMessageVerbosity == ELogVerbosity::NoLogging && !(Options & (EFindFirstObjectOptions::NativeFirst | EFindFirstObjectOptions::EnsureIfAmbiguous)))
 	{
-		Result = StaticFindFirstObjectFastInternal(Class, ObjectName, !!(Options & EFindFirstObjectOptions::ExactClass));
+		Result = StaticFindFirstObjectFastInternal(Class, ObjectName, FindObjectFlags);
 	}
 	else
 	{
 		TArray<UObject*> FoundObjects;
-		if (StaticFindAllObjectsFast(FoundObjects, Class, ObjectName, !!(Options & EFindFirstObjectOptions::ExactClass)))
+		if (StaticFindAllObjectsFast(FoundObjects, Class, ObjectName, FindObjectFlags))
 		{
 			if (FoundObjects.Num() > 1)
-			{				
+			{
 				if (!!(Options & EFindFirstObjectOptions::NativeFirst))
 				{
 					// Prioritize native class instances or native type objects
@@ -1482,7 +1483,7 @@ int32 LoadAssetAsync(FTopLevelAssetPath InAssetPath, FLoadAssetAsyncDelegate InC
 			UObject* LoadedObject = nullptr;
 			if (Result == EAsyncLoadingResult::Succeeded && LoadedPackage)
 			{
-				LoadedObject = StaticFindObjectFast(UObject::StaticClass(), LoadedPackage, InAssetPath.GetAssetName(), false);
+				LoadedObject = StaticFindObjectFast(UObject::StaticClass(), LoadedPackage, InAssetPath.GetAssetName());
 
 				// Package loaded but object was not found inside it, failure
 				if (!LoadedObject)
@@ -3606,7 +3607,7 @@ UObject* StaticAllocateObject
 	else
 	{
 		// See if object already exists.
-		Obj = StaticFindObjectFastInternal( /*Class=*/ NULL, InOuter, InName, true );
+		Obj = StaticFindObjectFastInternal( /*Class=*/ NULL, InOuter, InName, EFindObjectFlags::ExactClass );
 
 		// It is an error if we are trying to replace an object of a different class
 		if (Obj && !Obj->GetClass()->IsChildOf(InClass))
@@ -6645,7 +6646,7 @@ namespace UECodeGen_Private
 			return;
 		}
 
-		UObject* FoundPackage = StaticFindObjectFast(UPackage::StaticClass(), nullptr, FName(UTF8_TO_TCHAR(Params.NameUTF8)), false);
+		UObject* FoundPackage = StaticFindObjectFast(UPackage::StaticClass(), nullptr, FName(UTF8_TO_TCHAR(Params.NameUTF8)), EFindObjectFlags::None);
 
 #if USE_PER_MODULE_UOBJECT_BOOTSTRAP
 		if (!FoundPackage)
diff --git a/Engine/Source/Runtime/CoreUObject/Private/UObject/UObjectHash.cpp b/Engine/Source/Runtime/CoreUObject/Private/UObject/UObjectHash.cpp
--- a/Engine/Source/Runtime/CoreUObject/Private/UObject/UObjectHash.cpp
+++ b/Engine/Source/Runtime/CoreUObject/Private/UObject/UObjectHash.cpp
@@ -872,7 +872,7 @@ static int32 GetObjectOuterHash(FName ObjName, PTRINT Outer)
 	return GetTypeHash(ObjName) + static_cast<int32>(Outer >> 6);
 }
 
-UObject* StaticFindObjectFastExplicitThreadSafe(FUObjectHashTables& ThreadHash, const UClass* ObjectClass, FName ObjectName, const FString& ObjectPathName, bool bExactClass, EObjectFlags ExcludeFlags/*=0*/)
+UObject* StaticFindObjectFastExplicitThreadSafe(FUObjectHashTables& ThreadHash, const UClass* ObjectClass, FName ObjectName, const FString& ObjectPathName, EFindObjectFlags Flags, EObjectFlags ExcludeFlags/*=0*/)
 {
 	const EInternalObjectFlags ExclusiveInternalFlags = DefaultInternalExclusionFlags;
 
@@ -882,6 +882,7 @@ UObject* StaticFindObjectFastExplicitThreadSafe(FUObjectHashTables& ThreadHash,
 	FHashBucket* Bucket = ThreadHash.Hash.Find(Hash);
 	if (Bucket)
 	{
+		bool bExactClass = !!(Flags & EFindObjectFlags::ExactClass);
 		for (FHashBucketIterator It(*Bucket); It; ++It)
 		{
 			UObject* Object = (UObject*)*It;
@@ -922,14 +923,14 @@ UObject* StaticFindObjectFastExplicitThreadSafe(FUObjectHashTables& ThreadHash,
  * @param	ExclusiveFlags	Ignores objects that contain any of the specified exclusive flags
  * @return	Returns a pointer to the found object or nullptr if none could be found
  */
-UObject* StaticFindObjectFastExplicit( const UClass* ObjectClass, FName ObjectName, const FString& ObjectPathName, bool bExactClass, EObjectFlags ExcludeFlags/*=0*/ )
+UObject* StaticFindObjectFastExplicit( const UClass* ObjectClass, FName ObjectName, const FString& ObjectPathName, EFindObjectFlags Flags, EObjectFlags ExcludeFlags/*=0*/ )
 {
 	checkSlow(FPackageName::IsShortPackageName(ObjectName)); //@Package name transition, we aren't checking the name here because we know this is only used for texture
 
 	// Find an object with the specified name and (optional) class, in any package; if bAnyPackage is false, only matches top-level packages
 	const int32 Hash = GetObjectHash( ObjectName );
 	FUObjectHashTables& ThreadHash = FUObjectHashTables::Get();
-	UObject* Result = StaticFindObjectFastExplicitThreadSafe( ThreadHash, ObjectClass, ObjectName, ObjectPathName, bExactClass, ExcludeFlags );
+	UObject* Result = StaticFindObjectFastExplicitThreadSafe( ThreadHash, ObjectClass, ObjectName, ObjectPathName, Flags, ExcludeFlags );
 
 	return Result;
 }
@@ -1002,12 +1003,13 @@ struct FObjectSearchPath
 	}
 };
 
-UObject* StaticFindObjectInPackageInternal(FUObjectHashTables& ThreadHash, const UClass* ObjectClass, const UPackage* ObjectPackage, FName ObjectName, bool bExactClass, EObjectFlags ExcludeFlags, EInternalObjectFlags ExclusiveInternalFlags)
+UObject* StaticFindObjectInPackageInternal(FUObjectHashTables& ThreadHash, const UClass* ObjectClass, const UPackage* ObjectPackage, FName ObjectName, EFindObjectFlags Flags, EObjectFlags ExcludeFlags, EInternalObjectFlags ExclusiveInternalFlags)
 {
 	ExclusiveInternalFlags |= DefaultInternalExclusionFlags;
 	UObject* Result = nullptr;
 	if (FHashBucket* Inners = ThreadHash.PackageToObjectListMap.Find(ObjectPackage))
 	{
+		bool bExactClass = !!(Flags & EFindObjectFlags::ExactClass);
 		for (FHashBucketIterator It(*Inners); It; ++It)
 		{
 			UObject* Object = static_cast<UObject*>(*It);
@@ -1036,7 +1038,7 @@ UObject* StaticFindObjectInPackageInternal(FUObjectHashTables& ThreadHash, const
 	return Result;
 }
 
-UObject* StaticFindObjectFastInternalThreadSafe(FUObjectHashTables& ThreadHash, const UClass* ObjectClass, const UObject* ObjectPackage, FName ObjectName, bool bExactClass, EObjectFlags ExcludeFlags, EInternalObjectFlags ExclusiveInternalFlags)
+UObject* StaticFindObjectFastInternalThreadSafe(FUObjectHashTables& ThreadHash, const UClass* ObjectClass, const UObject* ObjectPackage, FName ObjectName, EFindObjectFlags Flags, EObjectFlags ExcludeFlags, EInternalObjectFlags ExclusiveInternalFlags)
 {
 	ExclusiveInternalFlags |= DefaultInternalExclusionFlags;
 
@@ -1046,6 +1048,7 @@ UObject* StaticFindObjectFastInternalThreadSafe(FUObjectHashTables& ThreadHash,
 	{
 		int32 Hash = GetObjectOuterHash(ObjectName, (PTRINT)ObjectPackage);
 		FHashTableLock HashLock(ThreadHash);
+		bool bExactClass = !!(Flags & EFindObjectFlags::ExactClass);
 		for (TMultiMap<int32, uint32>::TConstKeyIterator HashIt(ThreadHash.HashOuter, Hash); HashIt; ++HashIt)
 		{
 			uint32 InternalIndex = HashIt.Value();
@@ -1085,7 +1088,7 @@ UObject* StaticFindObjectFastInternalThreadSafe(FUObjectHashTables& ThreadHash,
 		// if the search fail and the OuterPackage is a UPackage, lookup potential external package
 		if (Result == nullptr && ObjectPackage->IsA(UPackage::StaticClass()))
 		{
-			Result = StaticFindObjectInPackageInternal(ThreadHash, ObjectClass, static_cast<const UPackage*>(ObjectPackage), ObjectName, bExactClass, ExcludeFlags, ExclusiveInternalFlags);
+			Result = StaticFindObjectInPackageInternal(ThreadHash, ObjectClass, static_cast<const UPackage*>(ObjectPackage), ObjectName, Flags, ExcludeFlags, ExclusiveInternalFlags);
 		}
 #endif
 	}
@@ -1099,6 +1102,7 @@ UObject* StaticFindObjectFastInternalThreadSafe(FUObjectHashTables& ThreadHash,
 		FHashBucket* Bucket = ThreadHash.Hash.Find(Hash);
 		if (Bucket)
 		{
+			bool bExactClass = !!(Flags & EFindObjectFlags::ExactClass);
 			for (FHashBucketIterator It(*Bucket); It; ++It)
 			{
 				UObject* Object = (UObject*)*It;
@@ -1147,7 +1151,7 @@ UObject* StaticFindObjectFastInternalThreadSafe(FUObjectHashTables& ThreadHash,
 	return Result;
 }
 
-UObject* StaticFindObjectFastInternal(const UClass* ObjectClass, const UObject* ObjectPackage, FName ObjectName, bool bExactClass, EObjectFlags ExcludeFlags, EInternalObjectFlags ExclusiveInternalFlags)
+UObject* StaticFindObjectFastInternal(const UClass* ObjectClass, const UObject* ObjectPackage, FName ObjectName, EFindObjectFlags Flags, EObjectFlags ExcludeFlags, EInternalObjectFlags ExclusiveInternalFlags)
 {
 	UObject* Result = nullptr;
 
@@ -1158,7 +1162,7 @@ UObject* StaticFindObjectFastInternal(const UClass* ObjectClass, const UObject*
 
 		// If they specified an outer use that during the hashing
 		FUObjectHashTables& ThreadHash = FUObjectHashTables::Get();
-		Result = StaticFindObjectFastInternalThreadSafe(ThreadHash, ObjectClass, ObjectPackage, ObjectName, bExactClass, ExcludeFlags, ExclusiveInternalFlags);
+		Result = StaticFindObjectFastInternalThreadSafe(ThreadHash, ObjectClass, ObjectPackage, ObjectName, Flags, ExcludeFlags, ExclusiveInternalFlags);
 	};
 
 	return Result;
@@ -1175,7 +1179,7 @@ bool DoesObjectPossiblyExist(const UObject* InOuter, FName ObjectName)
 	return ThreadHash.HashOuter.Contains(Hash);
 }
 
-bool StaticFindAllObjectsFastInternal(TArray<UObject*>& OutFoundObjects, const UClass* ObjectClass, FName ObjectName, bool bExactClass, EObjectFlags ExcludeFlags, EInternalObjectFlags ExclusiveInternalFlags)
+bool StaticFindAllObjectsFastInternal(TArray<UObject*>& OutFoundObjects, const UClass* ObjectClass, FName ObjectName, EFindObjectFlags Flags, EObjectFlags ExcludeFlags, EInternalObjectFlags ExclusiveInternalFlags)
 {
 	INC_DWORD_STAT(STAT_FindObjectFast);
 
@@ -1191,6 +1195,7 @@ bool StaticFindAllObjectsFastInternal(TArray<UObject*>& OutFoundObjects, const U
 	FHashBucket* Bucket = ThreadHash.Hash.Find(Hash);
 	if (Bucket)
 	{
+		bool bExactClass = !!(Flags & EFindObjectFlags::ExactClass);
 		for (FHashBucketIterator It(*Bucket); It; ++It)
 		{
 			UObject* Object = (UObject*)*It;
@@ -1219,7 +1224,7 @@ bool StaticFindAllObjectsFastInternal(TArray<UObject*>& OutFoundObjects, const U
 	return !!NumFoundObjects;
 }
 
-UObject* StaticFindFirstObjectFastInternal(const UClass* ObjectClass, FName ObjectName, bool bExactClass, EObjectFlags ExcludeFlags, EInternalObjectFlags ExclusiveInternalFlags)
+UObject* StaticFindFirstObjectFastInternal(const UClass* ObjectClass, FName ObjectName, EFindObjectFlags Flags, EObjectFlags ExcludeFlags, EInternalObjectFlags ExclusiveInternalFlags)
 {
 	INC_DWORD_STAT(STAT_FindObjectFast);
 
@@ -1236,6 +1241,7 @@ UObject* StaticFindFirstObjectFastInternal(const UClass* ObjectClass, FName Obje
 	FHashBucket* Bucket = ThreadHash.Hash.Find(Hash);
 	if (Bucket)
 	{
+		bool bExactClass = !!(Flags & EFindObjectFlags::ExactClass);
 		for (FHashBucketIterator It(*Bucket); It; ++It)
 		{
 			UObject* Object = (UObject*)*It;
@@ -1669,7 +1675,7 @@ UObjectBase* FindObjectWithOuter(const class UObjectBase* Outer, const class UCl
 
 	if( NameToLookFor != NAME_None )
 	{
-		Result = StaticFindObjectFastInternal(ClassToLookFor, static_cast<const UObject*>(Outer), NameToLookFor, false, RF_NoFlags, ExclusionInternalFlags);
+		Result = StaticFindObjectFastInternal(ClassToLookFor, static_cast<const UObject*>(Outer), NameToLookFor, EFindObjectFlags::None, RF_NoFlags, ExclusionInternalFlags);
 	}
 	else
 	{
diff --git a/Engine/Source/Runtime/CoreUObject/Public/UObject/FindObjectFlags.h b/Engine/Source/Runtime/CoreUObject/Public/UObject/FindObjectFlags.h
new file mode 100644
--- /dev/null
+++ b/Engine/Source/Runtime/CoreUObject/Public/UObject/FindObjectFlags.h
@@ -0,0 +1,16 @@
+// Copyright Epic Games, Inc. All Rights Reserved.
+
+#pragma once
+
+#include "Misc/EnumClassFlags.h"
+
+enum class EFindObjectFlags
+{
+	None,
+
+	/** Whether to require an exact match with the passed in class */
+	ExactClass
+};
+ENUM_CLASS_FLAGS(EFindObjectFlags)
+
+#define UE_EXACTCLASS_BOOL_DEPRECATED(FunctionName) //UE_DEPRECATED(5.7, FunctionName " with a boolean ExactClass has been deprecated - please use the EFindObjectFlags enum instead")
diff --git a/Engine/Source/Runtime/CoreUObject/Public/UObject/UObjectGlobals.h b/Engine/Source/Runtime/CoreUObject/Public/UObject/UObjectGlobals.h
--- a/Engine/Source/Runtime/CoreUObject/Public/UObject/UObjectGlobals.h
+++ b/Engine/Source/Runtime/CoreUObject/Public/UObject/UObjectGlobals.h
@@ -6,6 +6,7 @@
 
 #pragma once
 
+#include "Concepts/SameAs.h"
 #include "Containers/Array.h"
 #include "Containers/ArrayView.h"
 #include "Containers/ContainerAllocationPolicies.h"
@@ -36,6 +37,7 @@
 #include "Templates/UnrealTemplate.h"
 #include "Templates/IsTObjectPtr.h"
 #include "Traits/IsCharEncodingCompatibleWith.h"
+#include "UObject/FindObjectFlags.h"
 #include "UObject/NameTypes.h"
 #include "UObject/ObjectMacros.h"
 #include "UObject/ObjectHandle.h"
@@ -301,13 +303,19 @@ COREUOBJECT_API int32 UpdateSuffixForNextNewObject(UObject* Parent, const UClass
  * @param	Class			The to be found object's class
  * @param	InOuter			Outer object to look inside, if null this will only look for top level packages
  * @param	InName			Object name to look for relative to InOuter
- * @param	bExactClass		Whether to require an exact match with the passed in class
+ * @param	Flags			Flags which control the search
  * @param	ExclusiveFlags	Ignores objects that contain any of the specified exclusive flags
  * @param	ExclusiveInternalFlags  Ignores objects that contain any of the specified internal exclusive flags
  *
  * @return	Returns a pointer to the found object or null if none could be found
  */
-COREUOBJECT_API UObject* StaticFindObjectFast(UClass* Class, UObject* InOuter, FName InName, bool bExactClass = false, EObjectFlags ExclusiveFlags = RF_NoFlags, EInternalObjectFlags ExclusiveInternalFlags = EInternalObjectFlags::None);
+COREUOBJECT_API UObject* StaticFindObjectFast(UClass* Class, UObject* InOuter, FName InName, EFindObjectFlags Flags = EFindObjectFlags::None, EObjectFlags ExclusiveFlags = RF_NoFlags, EInternalObjectFlags ExclusiveInternalFlags = EInternalObjectFlags::None);
+template <UE::CSameAs<bool> ExactClassType>
+UE_EXACTCLASS_BOOL_DEPRECATED("StaticFindObjectFast")
+UE_NODEBUG FORCEINLINE UObject* StaticFindObjectFast(UClass* Class, UObject* InOuter, FName InName, ExactClassType bExactClass, EObjectFlags ExclusiveFlags = RF_NoFlags, EInternalObjectFlags ExclusiveInternalFlags = EInternalObjectFlags::None)
+{
+	return StaticFindObjectFast(Class, InOuter, InName, bExactClass ? EFindObjectFlags::ExactClass : EFindObjectFlags::None, ExclusiveFlags, ExclusiveInternalFlags);
+}
 
 /**
  * Fast and safe version of StaticFindObject that relies on the passed in FName being the object name without any group/package qualifiers.
@@ -317,13 +325,19 @@ COREUOBJECT_API UObject* StaticFindObjectFast(UClass* Class, UObject* InOuter, F
  * @param	Class			The to be found object's class
  * @param	InOuter			Outer object to look inside, if null this will only look for top level packages
  * @param	InName			Object name to look for relative to InOuter
- * @param	bExactClass		Whether to require an exact match with the passed in class
+ * @param	Flags			Flags which control the search
  * @param	ExclusiveFlags	Ignores objects that contain any of the specified exclusive flags
  * @param	ExclusiveInternalFlags  Ignores objects that contain any of the specified internal exclusive flags
  *
  * @return	Returns a pointer to the found object or null if none could be found
  */
-COREUOBJECT_API UObject* StaticFindObjectFastSafe(UClass* Class, UObject* InOuter, FName InName, bool bExactClass = false, EObjectFlags ExclusiveFlags = RF_NoFlags, EInternalObjectFlags ExclusiveInternalFlags = EInternalObjectFlags::None);
+COREUOBJECT_API UObject* StaticFindObjectFastSafe(UClass* Class, UObject* InOuter, FName InName, EFindObjectFlags Flags = EFindObjectFlags::None, EObjectFlags ExclusiveFlags = RF_NoFlags, EInternalObjectFlags ExclusiveInternalFlags = EInternalObjectFlags::None);
+template <UE::CSameAs<bool> ExactClassType>
+UE_EXACTCLASS_BOOL_DEPRECATED("StaticFindObjectFastSafe")
+UE_NODEBUG FORCEINLINE UObject* StaticFindObjectFastSafe(UClass* Class, UObject* InOuter, FName InName, ExactClassType bExactClass, EObjectFlags ExclusiveFlags = RF_NoFlags, EInternalObjectFlags ExclusiveInternalFlags = EInternalObjectFlags::None)
+{
+	return StaticFindObjectFastSafe(Class, InOuter, InName, bExactClass ? EFindObjectFlags::ExactClass : EFindObjectFlags::None, ExclusiveFlags, ExclusiveInternalFlags);
+}
 
 /**
  * Tries to find an object in memory. This will handle fully qualified paths of the form /path/packagename.object:subobject and resolve references for you.
@@ -331,28 +345,52 @@ COREUOBJECT_API UObject* StaticFindObjectFastSafe(UClass* Class, UObject* InOute
  * @param	Class			The to be found object's class
  * @param	InOuter			Outer object to look inside. If this is null then InName should start with a package name
  * @param	InName			The object path to search for an object, relative to InOuter
- * @param	ExactClass		Whether to require an exact match with the passed in class
+ * @param	Flags			Flags which control the search
  *
  * @return	Returns a pointer to the found object or nullptr if none could be found
  */
-COREUOBJECT_API UObject* StaticFindObject( UClass* Class, UObject* InOuter, const TCHAR* Name, bool ExactClass=false );
+COREUOBJECT_API UObject* StaticFindObject( UClass* Class, UObject* InOuter, const TCHAR* Name, EFindObjectFlags Flags = EFindObjectFlags::None );
+template <UE::CSameAs<bool> ExactClassType>
+UE_EXACTCLASS_BOOL_DEPRECATED("StaticFindObject")
+UE_NODEBUG FORCEINLINE UObject* StaticFindObject( UClass* Class, UObject* InOuter, const TCHAR* Name, ExactClassType bExactClass )
+{
+	return StaticFindObject(Class, InOuter, Name, bExactClass ? EFindObjectFlags::ExactClass : EFindObjectFlags::None);
+}
 
 /**
  * Tries to find an object in memory. This version uses FTopLevelAssetPath to find the object.
  *
  * @param	Class			The to be found object's class
  * @param	ObjectPath		FName pair representing the outer package object and the inner top level object (asset)
- * @param	ExactClass		Whether to require an exact match with the passed in class
+ * @param	Flags			Flags which control the search
  *
  * @return	Returns a pointer to the found object or nullptr if none could be found
  */
-COREUOBJECT_API UObject* StaticFindObject(UClass* Class, FTopLevelAssetPath ObjectPath, bool ExactClass /*= false*/);
+COREUOBJECT_API UObject* StaticFindObject(UClass* Class, FTopLevelAssetPath ObjectPath, EFindObjectFlags Flags = EFindObjectFlags::None);
+template <UE::CSameAs<bool> ExactClassType>
+UE_EXACTCLASS_BOOL_DEPRECATED("StaticFindObject")
+UE_NODEBUG FORCEINLINE UObject* StaticFindObject(UClass* Class, FTopLevelAssetPath ObjectPath, ExactClassType bExactClass)
+{
+	return StaticFindObject(Class, ObjectPath, bExactClass ? EFindObjectFlags::ExactClass : EFindObjectFlags::None);
+}
 
 /** Version of StaticFindObject() that will assert if the object is not found */
-COREUOBJECT_API UObject* StaticFindObjectChecked( UClass* Class, UObject* InOuter, const TCHAR* Name, bool ExactClass=false );
+COREUOBJECT_API UObject* StaticFindObjectChecked( UClass* Class, UObject* InOuter, const TCHAR* Name, EFindObjectFlags Flags = EFindObjectFlags::None );
+template <UE::CSameAs<bool> ExactClassType>
+UE_EXACTCLASS_BOOL_DEPRECATED("StaticFindObjectChecked")
+UE_NODEBUG FORCEINLINE UObject* StaticFindObjectChecked( UClass* Class, UObject* InOuter, const TCHAR* Name, ExactClassType bExactClass)
+{
+	return StaticFindObjectChecked(Class, InOuter, Name, bExactClass ? EFindObjectFlags::ExactClass : EFindObjectFlags::None);
+}
 
 /** Internal version of StaticFindObject that will not assert on GIsSavingPackage or IsGarbageCollectingAndLockingUObjectHashTables() */
-COREUOBJECT_API UObject* StaticFindObjectSafe( UClass* Class, UObject* InOuter, const TCHAR* Name, bool ExactClass=false );
+COREUOBJECT_API UObject* StaticFindObjectSafe( UClass* Class, UObject* InOuter, const TCHAR* Name, EFindObjectFlags Flags = EFindObjectFlags::None );
+template <UE::CSameAs<bool> ExactClassType>
+UE_EXACTCLASS_BOOL_DEPRECATED("StaticFindObjectSafe")
+UE_NODEBUG FORCEINLINE UObject* StaticFindObjectSafe( UClass* Class, UObject* InOuter, const TCHAR* Name, ExactClassType bExactClass )
+{
+	return StaticFindObjectSafe(Class, InOuter, Name, bExactClass ? EFindObjectFlags::ExactClass : EFindObjectFlags::None);
+}
 
 /**
  * Tries to find an object in memory. This version uses FTopLevelAssetPath to find the object.
@@ -360,11 +398,17 @@ COREUOBJECT_API UObject* StaticFindObjectSafe( UClass* Class, UObject* InOuter,
  *
  * @param	Class			The to be found object's class
  * @param	ObjectPath		FName pair representing the outer package object and the inner top level object (asset)
- * @param	ExactClass		Whether to require an exact match with the passed in class
+ * @param	Flags			Flags which control the search
  *
  * @return	Returns a pointer to the found object or nullptr if none could be found
  */
-COREUOBJECT_API UObject* StaticFindObjectSafe(UClass* Class, FTopLevelAssetPath ObjectPath, bool ExactClass /*= false*/);
+COREUOBJECT_API UObject* StaticFindObjectSafe(UClass* Class, FTopLevelAssetPath ObjectPath, EFindObjectFlags Flags = EFindObjectFlags::None);
+template <UE::CSameAs<bool> ExactClassType>
+UE_EXACTCLASS_BOOL_DEPRECATED("StaticFindObjectSafe")
+UE_NODEBUG FORCEINLINE UObject* StaticFindObjectSafe(UClass* Class, FTopLevelAssetPath ObjectPath, ExactClassType bExactClass)
+{
+	return StaticFindObjectSafe(Class, ObjectPath, bExactClass ? EFindObjectFlags::ExactClass : EFindObjectFlags::None);
+}
 
 /**
  * Fast version of StaticFindAllObjects that relies on the passed in FName being the object name without any group/package qualifiers.
@@ -373,13 +417,19 @@ COREUOBJECT_API UObject* StaticFindObjectSafe(UClass* Class, FTopLevelAssetPath
  * @param	OutFoundObjects	Array of objects matching the search parameters
  * @param	ObjectClass		The to be found object's class
  * @param	ObjectName		Object name to look for relative to InOuter
- * @param	ExactClass		Whether to require an exact match with the passed in class
+ * @param	Flags			Flags which control the search
  * @param	ExclusiveFlags	Ignores objects that contain any of the specified exclusive flags
  * @param	ExclusiveInternalFlags  Ignores objects that contain any of the specified internal exclusive flags
  *
  * @return	Returns true if any objects were found, false otherwise
  */
-COREUOBJECT_API bool StaticFindAllObjectsFast(TArray<UObject*>& OutFoundObjects, UClass* ObjectClass, FName ObjectName, bool ExactClass = false, EObjectFlags ExclusiveFlags = RF_NoFlags, EInternalObjectFlags ExclusiveInternalFlags = EInternalObjectFlags::None);
+COREUOBJECT_API bool StaticFindAllObjectsFast(TArray<UObject*>& OutFoundObjects, UClass* ObjectClass, FName ObjectName, EFindObjectFlags Flags = EFindObjectFlags::None, EObjectFlags ExclusiveFlags = RF_NoFlags, EInternalObjectFlags ExclusiveInternalFlags = EInternalObjectFlags::None);
+template <UE::CSameAs<bool> ExactClassType>
+UE_EXACTCLASS_BOOL_DEPRECATED("StaticFindAllObjectsFast")
+UE_NODEBUG FORCEINLINE bool StaticFindAllObjectsFast(TArray<UObject*>& OutFoundObjects, UClass* ObjectClass, FName ObjectName, ExactClassType bExactClass, EObjectFlags ExclusiveFlags = RF_NoFlags, EInternalObjectFlags ExclusiveInternalFlags = EInternalObjectFlags::None)
+{
+	return StaticFindAllObjectsFast(OutFoundObjects, ObjectClass, ObjectName, bExactClass ? EFindObjectFlags::ExactClass : EFindObjectFlags::None, ExclusiveFlags, ExclusiveInternalFlags);
+}
 
 /**
  * Fast version of StaticFindAllObjects that relies on the passed in FName being the object name without any group/package qualifiers.
@@ -389,13 +439,19 @@ COREUOBJECT_API bool StaticFindAllObjectsFast(TArray<UObject*>& OutFoundObjects,
  * @param	OutFoundObjects	Array of objects matching the search parameters
  * @param	ObjectClass		The to be found object's class
  * @param	ObjectName		Object name to look for relative to InOuter
- * @param	ExactClass		Whether to require an exact match with the passed in class
+ * @param	Flags			Flags which control the search
  * @param	ExclusiveFlags	Ignores objects that contain any of the specified exclusive flags
  * @param	ExclusiveInternalFlags  Ignores objects that contain any of the specified internal exclusive flags
  *
  * @return	Returns true if any objects were found, false otherwise
  */
-COREUOBJECT_API bool StaticFindAllObjectsFastSafe(TArray<UObject*>& OutFoundObjects, UClass* ObjectClass, FName ObjectName, bool ExactClass = false, EObjectFlags ExclusiveFlags = RF_NoFlags, EInternalObjectFlags ExclusiveInternalFlags = EInternalObjectFlags::None);
+COREUOBJECT_API bool StaticFindAllObjectsFastSafe(TArray<UObject*>& OutFoundObjects, UClass* ObjectClass, FName ObjectName, EFindObjectFlags Flags = EFindObjectFlags::None, EObjectFlags ExclusiveFlags = RF_NoFlags, EInternalObjectFlags ExclusiveInternalFlags = EInternalObjectFlags::None);
+template <UE::CSameAs<bool> ExactClassType>
+UE_EXACTCLASS_BOOL_DEPRECATED("StaticFindAllObjectsFastSafe")
+UE_NODEBUG FORCEINLINE bool StaticFindAllObjectsFastSafe(TArray<UObject*>& OutFoundObjects, UClass* ObjectClass, FName ObjectName, ExactClassType bExactClass, EObjectFlags ExclusiveFlags = RF_NoFlags, EInternalObjectFlags ExclusiveInternalFlags = EInternalObjectFlags::None)
+{
+	return StaticFindAllObjectsFastSafe(OutFoundObjects, ObjectClass, ObjectName, bExactClass ? EFindObjectFlags::ExactClass : EFindObjectFlags::None, ExclusiveFlags, ExclusiveInternalFlags);
+}
 
 /**
  * Tries to find all objects matching the search paramters in memory. This will handle fully qualified paths of the form /path/packagename.object:subobject and resolve references for you.
@@ -403,11 +459,17 @@ COREUOBJECT_API bool StaticFindAllObjectsFastSafe(TArray<UObject*>& OutFoundObje
  * @param	OutFoundObjects	Array of objects matching the search parameters
  * @param	Class			The to be found object's class
  * @param	Name			The object path to search for an object, relative to InOuter
- * @param	ExactClass		Whether to require an exact match with the passed in class
+ * @param	Flags			Flags which control the search
  *
  * @return	Returns true if any objects were found, false otherwise
  */
-COREUOBJECT_API bool StaticFindAllObjects(TArray<UObject*>& OutFoundObjects, UClass* Class, const TCHAR* Name, bool ExactClass = false);
+COREUOBJECT_API bool StaticFindAllObjects(TArray<UObject*>& OutFoundObjects, UClass* Class, const TCHAR* Name, EFindObjectFlags Flags = EFindObjectFlags::None);
+template <UE::CSameAs<bool> ExactClassType>
+UE_EXACTCLASS_BOOL_DEPRECATED("StaticFindAllObjects")
+UE_NODEBUG FORCEINLINE bool StaticFindAllObjects(TArray<UObject*>& OutFoundObjects, UClass* Class, const TCHAR* Name, ExactClassType bExactClass)
+{
+	return StaticFindAllObjects(OutFoundObjects, Class, Name, bExactClass ? EFindObjectFlags::ExactClass : EFindObjectFlags::None);
+}
 
 /**
  * Tries to find all objects matching the search paramters in memory. This will handle fully qualified paths of the form /path/packagename.object:subobject and resolve references for you.
@@ -416,11 +478,17 @@ COREUOBJECT_API bool StaticFindAllObjects(TArray<UObject*>& OutFoundObjects, UCl
  * @param	OutFoundObjects	Array of objects matching the search parameters
  * @param	Class			The to be found object's class
  * @param	Name			The object path to search for an object, relative to InOuter
- * @param	ExactClass		Whether to require an exact match with the passed in class
+ * @param	Flags			Flags which control the search
  *
  * @return	Returns true if any objects were found, false otherwise
  */
-COREUOBJECT_API bool StaticFindAllObjectsSafe(TArray<UObject*>& OutFoundObjects, UClass* Class, const TCHAR* Name, bool ExactClass = false);
+COREUOBJECT_API bool StaticFindAllObjectsSafe(TArray<UObject*>& OutFoundObjects, UClass* Class, const TCHAR* Name, EFindObjectFlags Flags = EFindObjectFlags::None);
+template <UE::CSameAs<bool> ExactClassType>
+UE_EXACTCLASS_BOOL_DEPRECATED("StaticFindAllObjectsSafe")
+UE_NODEBUG FORCEINLINE bool StaticFindAllObjectsSafe(TArray<UObject*>& OutFoundObjects, UClass* Class, const TCHAR* Name, ExactClassType bExactClass)
+{
+	return StaticFindAllObjectsSafe(OutFoundObjects, Class, Name, bExactClass ? EFindObjectFlags::ExactClass : EFindObjectFlags::None);
+}
 
 
 enum class EFindFirstObjectOptions
@@ -1940,9 +2008,15 @@ inline bool ParseObject( const TCHAR* Stream, const TCHAR* Match, T*& Obj, UObje
  * @see StaticFindObjectFast()
  */
 template< class T >
-inline T* FindObjectFast(UObject* Outer, FName Name, bool ExactClass = false, EObjectFlags ExclusiveFlags = RF_NoFlags)
+inline T* FindObjectFast(UObject* Outer, FName Name, EFindObjectFlags Flags = EFindObjectFlags::None, EObjectFlags ExclusiveFlags = RF_NoFlags)
+{
+	return (T*)StaticFindObjectFast(T::StaticClass(), Outer, Name, Flags, ExclusiveFlags);
+}
+template <typename T, UE::CSameAs<bool> ExactClassType>
+UE_EXACTCLASS_BOOL_DEPRECATED("FindObjectFast")
+UE_NODEBUG FORCEINLINE T* FindObjectFast(UObject* Outer, FName Name, ExactClassType bExactClass, EObjectFlags ExclusiveFlags = RF_NoFlags)
 {
-	return (T*)StaticFindObjectFast(T::StaticClass(), Outer, Name, ExactClass, ExclusiveFlags);
+	return FindObjectFast<T>(Outer, Name, bExactClass ? EFindObjectFlags::ExactClass : EFindObjectFlags::None, ExclusiveFlags);
 }
 
 /**
@@ -1950,9 +2024,15 @@ inline T* FindObjectFast(UObject* Outer, FName Name, bool ExactClass = false, EO
  * @see StaticFindObject()
  */
 template< class T > 
-inline T* FindObject( UObject* Outer, const TCHAR* Name, bool ExactClass=false )
+inline T* FindObject( UObject* Outer, const TCHAR* Name, EFindObjectFlags Flags = EFindObjectFlags::None )
 {
-	return (T*)StaticFindObject( T::StaticClass(), Outer, Name, ExactClass );
+	return (T*)StaticFindObject( T::StaticClass(), Outer, Name, Flags );
+}
+template <typename T, UE::CSameAs<bool> ExactClassType>
+UE_EXACTCLASS_BOOL_DEPRECATED("FindObject")
+UE_NODEBUG FORCEINLINE T* FindObject(UObject* Outer, const TCHAR* Name, ExactClassType bExactClass)
+{
+	return FindObject<T>(Outer, Name, bExactClass ? EFindObjectFlags::ExactClass : EFindObjectFlags::None);
 }
 
 /**
@@ -1960,9 +2040,15 @@ inline T* FindObject( UObject* Outer, const TCHAR* Name, bool ExactClass=false )
  * @see StaticFindObject()
  */
 template< class T >
-inline T* FindObject(FTopLevelAssetPath InPath, bool ExactClass = false)
+inline T* FindObject(FTopLevelAssetPath InPath, EFindObjectFlags Flags = EFindObjectFlags::None)
 {
-	return (T*)StaticFindObject(T::StaticClass(), InPath, ExactClass);
+	return (T*)StaticFindObject(T::StaticClass(), InPath, Flags);
+}
+template <typename T, UE::CSameAs<bool> ExactClassType>
+UE_EXACTCLASS_BOOL_DEPRECATED("FindObject")
+UE_NODEBUG FORCEINLINE T* FindObject(FTopLevelAssetPath InPath, ExactClassType bExactClass)
+{
+	return FindObject<T>(InPath, bExactClass ? EFindObjectFlags::ExactClass : EFindObjectFlags::None);
 }
 
 /**
@@ -1970,9 +2056,15 @@ inline T* FindObject(FTopLevelAssetPath InPath, bool ExactClass = false)
  * @see StaticFindObjectChecked()
  */
 template< class T > 
-inline T* FindObjectChecked( UObject* Outer, const TCHAR* Name, bool ExactClass=false )
+inline T* FindObjectChecked( UObject* Outer, const TCHAR* Name, EFindObjectFlags Flags = EFindObjectFlags::None )
 {
-	return (T*)StaticFindObjectChecked( T::StaticClass(), Outer, Name, ExactClass );
+	return (T*)StaticFindObjectChecked( T::StaticClass(), Outer, Name, Flags );
+}
+template <typename T, UE::CSameAs<bool> ExactClassType>
+UE_EXACTCLASS_BOOL_DEPRECATED("FindObjectChecked")
+UE_NODEBUG FORCEINLINE T* FindObjectChecked( UObject* Outer, const TCHAR* Name, ExactClassType bExactClass )
+{
+	return FindObjectChecked<T>(Outer, Name, bExactClass ? EFindObjectFlags::ExactClass : EFindObjectFlags::None);
 }
 
 /**
@@ -1980,9 +2072,15 @@ inline T* FindObjectChecked( UObject* Outer, const TCHAR* Name, bool ExactClass=
  * @see StaticFindObjectSafe()
  */
 template< class T > 
-inline T* FindObjectSafe( UObject* Outer, const TCHAR* Name, bool ExactClass=false )
+inline T* FindObjectSafe( UObject* Outer, const TCHAR* Name, EFindObjectFlags Flags = EFindObjectFlags::None )
+{
+	return (T*)StaticFindObjectSafe( T::StaticClass(), Outer, Name, Flags );
+}
+template <typename T, UE::CSameAs<bool> ExactClassType>
+UE_EXACTCLASS_BOOL_DEPRECATED("FindObjectSafe")
+UE_NODEBUG FORCEINLINE T* FindObjectSafe( UObject* Outer, const TCHAR* Name, ExactClassType bExactClass )
 {
-	return (T*)StaticFindObjectSafe( T::StaticClass(), Outer, Name, ExactClass );
+	return FindObjectSafe<T>(Outer, Name, bExactClass ? EFindObjectFlags::ExactClass : EFindObjectFlags::None);
 }
 
 /**
@@ -1990,9 +2088,15 @@ inline T* FindObjectSafe( UObject* Outer, const TCHAR* Name, bool ExactClass=fal
  * @see StaticFindObject()
  */
 template< class T >
-inline T* FindObjectSafe(FTopLevelAssetPath InPath, bool ExactClass = false)
+inline T* FindObjectSafe(FTopLevelAssetPath InPath, EFindObjectFlags Flags = EFindObjectFlags::None)
+{
+	return (T*)StaticFindObjectSafe(T::StaticClass(), InPath, Flags);
+}
+template <typename T, UE::CSameAs<bool> ExactClassType>
+UE_EXACTCLASS_BOOL_DEPRECATED("FindObjectSafe")
+UE_NODEBUG FORCEINLINE T* FindObjectSafe(FTopLevelAssetPath InPath, ExactClassType bExactClass)
 {
-	return (T*)StaticFindObjectSafe(T::StaticClass(), InPath, ExactClass);
+	return FindObjectSafe<T>(InPath, bExactClass ? EFindObjectFlags::ExactClass : EFindObjectFlags::None);
 }
 
 /**
@@ -2080,19 +2184,18 @@ COREUOBJECT_API UFunction* FindDelegateSignature(FName DelegateSignatureName);
  *
  * @param	ObjectArray		the array to search - must be an array of pointers to instances of a UObject-derived class
  * @param	ClassToCheck	the object class to search for
- * @param	bExactClass		true to consider only those objects that have the class specified, or false to consider objects
- *							of classes derived from the specified SearhClass as well
+ * @param	Flags			Flags which control the search
  * @param	out_Objects		if specified, any objects that match the SearchClass will be added to this array
  */
 template <class T>
-bool ContainsObjectOfClass( const TArray<T*>& ObjectArray, UClass* ClassToCheck, bool bExactClass=false, TArray<T*>* out_Objects=nullptr )
+bool ContainsObjectOfClass( const TArray<T*>& ObjectArray, UClass* ClassToCheck, EFindObjectFlags Flags = EFindObjectFlags::None, TArray<T*>* out_Objects=nullptr )
 {
 	bool bResult = false;
 	for ( int32 ArrayIndex = 0; ArrayIndex < ObjectArray.Num(); ArrayIndex++ )
 	{
 		if ( ObjectArray[ArrayIndex] != nullptr )
 		{
-			bool bMatchesSearchCriteria = bExactClass
+			bool bMatchesSearchCriteria = !!(Flags & EFindObjectFlags::ExactClass)
 				? ObjectArray[ArrayIndex]->GetClass() == ClassToCheck
 				: ObjectArray[ArrayIndex]->IsA(ClassToCheck);
 
@@ -2114,6 +2217,12 @@ bool ContainsObjectOfClass( const TArray<T*>& ObjectArray, UClass* ClassToCheck,
 
 	return bResult;
 }
+template <typename T, UE::CSameAs<bool> ExactClassType>
+UE_EXACTCLASS_BOOL_DEPRECATED("ContainsObjectOfClass")
+UE_NODEBUG FORCEINLINE bool ContainsObjectOfClass( const TArray<T*>& ObjectArray, UClass* ClassToCheck, ExactClassType bExactClass, TArray<T*>* out_Objects=nullptr )
+{
+	return ContainsObjectOfClass(ObjectArray, ClassToCheck, bExactClass ? EFindObjectFlags::ExactClass : EFindObjectFlags::None, out_Objects );
+}
 
 /**
  * Utility struct for restoring object flags for all objects.
@@ -2188,14 +2297,20 @@ public:
 		* @param	InClass			if non-null, will only iterate on items IsA this class
 		* @param	InbExactClass	if true, will only iterate on exact matches
 		*/
-	FORCEINLINE TObjectArrayIterator( TArray<TObjectClass*>& InArray, UClass* InClassToCheck = nullptr, bool InbExactClass = false) :	
+	FORCEINLINE TObjectArrayIterator( TArray<TObjectClass*>& InArray, UClass* InClassToCheck = nullptr, EFindObjectFlags InFlags = EFindObjectFlags::None) :	
 		Array(InArray),
 		Index(-1),
 		ClassToCheck(InClassToCheck),
-		bExactClass(InbExactClass)
+		Flags(InFlags)
 	{
 		Advance();
 	}
+	template <UE::CSameAs<bool> ExactClassType>
+	UE_EXACTCLASS_BOOL_DEPRECATED("TObjectArrayIterator")
+	UE_NODEBUG FORCEINLINE TObjectArrayIterator( TArray<TObjectClass*>& InArray, UClass* InClassToCheck, ExactClassType bInExactClass)
+		: TObjectArrayIterator(InArray, InClassToCheck, bInExactClass ? EFindObjectFlags::ExactClass : EFindObjectFlags::None)
+	{
+	}
 	/**
 		* Iterator advance
 		*/
@@ -2271,7 +2386,7 @@ protected:
 			if (
 				IsValid(At) && 
 				(!ClassToCheck ||
-					(bExactClass
+					(!!(Flags & EFindObjectFlags::ExactClass)
 						? At->GetClass() == ClassToCheck
 						: At->IsA(ClassToCheck)))
 				)
@@ -2285,11 +2400,11 @@ private:
 	/** The array that we are iterating on */
 	TArray<TObjectClass*>&	Array;
 	/** Index of the current element in the object array */
-	int32						Index;
+	int32					Index;
 	/** Class using as a criteria */
 	UClass*					ClassToCheck;
-	/** Flag to require exact class matches */
-	bool					bExactClass;
+	/** Flag to control the iteration */
+	EFindObjectFlags		Flags;
 };
 
 /** Reference collecting archive created by FReferenceCollector::GetVerySlowReferenceCollectorArchive() */
diff --git a/Engine/Source/Runtime/CoreUObject/Public/UObject/UObjectHash.h b/Engine/Source/Runtime/CoreUObject/Public/UObject/UObjectHash.h
--- a/Engine/Source/Runtime/CoreUObject/Public/UObject/UObjectHash.h
+++ b/Engine/Source/Runtime/CoreUObject/Public/UObject/UObjectHash.h
@@ -6,6 +6,7 @@
 
 #pragma once
 
+#include "Concepts/SameAs.h"
 #include "Containers/Array.h"
 #include "Containers/ArrayView.h"
 #include "Containers/Map.h"
@@ -14,6 +15,7 @@
 #include "CoreTypes.h"
 #include "Stats/Stats.h"
 #include "Templates/Function.h"
+#include "UObject/FindObjectFlags.h"
 #include "UObject/NameTypes.h"
 #include "UObject/ObjectMacros.h"
 #include "UObject/UObjectArray.h"
@@ -41,7 +43,13 @@ DECLARE_DWORD_COUNTER_STAT_EXTERN(TEXT("NumObjects"), STAT_Hash_NumObjects, STAT
  * @param	ExclusiveInternalFlags	Ignores objects that contain any of the specified internal exclusive flags
  * @return	Returns a pointer to the found object or NULL if none could be found
  */
-COREUOBJECT_API UObject* StaticFindObjectFastInternal(const UClass* Class, const UObject* InOuter, FName InName, bool ExactClass = false, EObjectFlags ExclusiveFlags = RF_NoFlags, EInternalObjectFlags ExclusiveInternalFlags = EInternalObjectFlags::None);
+COREUOBJECT_API UObject* StaticFindObjectFastInternal(const UClass* Class, const UObject* InOuter, FName InName, EFindObjectFlags Flags = EFindObjectFlags::None, EObjectFlags ExclusiveFlags = RF_NoFlags, EInternalObjectFlags ExclusiveInternalFlags = EInternalObjectFlags::None);
+template <UE::CSameAs<bool> ExactClassType>
+UE_EXACTCLASS_BOOL_DEPRECATED("StaticFindObjectFastInternal")
+UE_NODEBUG FORCEINLINE UObject* StaticFindObjectFastInternal(const UClass* Class, const UObject* InOuter, FName InName, ExactClassType bExactClass, EObjectFlags ExclusiveFlags = RF_NoFlags, EInternalObjectFlags ExclusiveInternalFlags = EInternalObjectFlags::None)
+{
+	return StaticFindObjectFastInternal(Class, InOuter, InName, bExactClass ? EFindObjectFlags::ExactClass : EFindObjectFlags::None, ExclusiveFlags, ExclusiveInternalFlags);
+}
 
 /**
  * Private internal version of StaticFindAllObjectsFast that allows using 0 exclusion flags.
@@ -54,7 +62,13 @@ COREUOBJECT_API UObject* StaticFindObjectFastInternal(const UClass* Class, const
  * @param	ExclusiveInternalFlags	Ignores objects that contain any of the specified internal exclusive flags
  * @return	Returns true if any objects were found, false otherwise
  */
-COREUOBJECT_API bool StaticFindAllObjectsFastInternal(TArray<UObject*>& OutFoundObjects, const UClass* ObjectClass, FName ObjectName, bool bExactClass, EObjectFlags ExcludeFlags = RF_NoFlags, EInternalObjectFlags ExclusiveInternalFlags = EInternalObjectFlags::None);
+COREUOBJECT_API bool StaticFindAllObjectsFastInternal(TArray<UObject*>& OutFoundObjects, const UClass* ObjectClass, FName ObjectName, EFindObjectFlags Flags = EFindObjectFlags::None, EObjectFlags ExcludeFlags = RF_NoFlags, EInternalObjectFlags ExclusiveInternalFlags = EInternalObjectFlags::None);
+template <UE::CSameAs<bool> ExactClassType>
+UE_EXACTCLASS_BOOL_DEPRECATED("StaticFindAllObjectsFastInternal")
+UE_NODEBUG FORCEINLINE bool StaticFindAllObjectsFastInternal(TArray<UObject*>& OutFoundObjects, const UClass* ObjectClass, FName ObjectName, ExactClassType bExactClass, EObjectFlags ExcludeFlags = RF_NoFlags, EInternalObjectFlags ExclusiveInternalFlags = EInternalObjectFlags::None)
+{
+	return StaticFindAllObjectsFastInternal(OutFoundObjects, ObjectClass, ObjectName, bExactClass ? EFindObjectFlags::ExactClass : EFindObjectFlags::None, ExcludeFlags, ExclusiveInternalFlags);
+}
 
 /**
  * Private internal version of StaticFindFirstObject that allows using 0 exclusion flags.
@@ -66,7 +80,13 @@ COREUOBJECT_API bool StaticFindAllObjectsFastInternal(TArray<UObject*>& OutFound
  * @param	ExclusiveInternalFlags	Ignores objects that contain any of the specified internal exclusive flags
  * @return	Returns first object with the specified name and class, null if not found
  */
-COREUOBJECT_API UObject* StaticFindFirstObjectFastInternal(const UClass* ObjectClass, FName ObjectName, bool bExactClass, EObjectFlags ExcludeFlag = RF_NoFlags, EInternalObjectFlags ExclusiveInternalFlags = EInternalObjectFlags::None);
+COREUOBJECT_API UObject* StaticFindFirstObjectFastInternal(const UClass* ObjectClass, FName ObjectName, EFindObjectFlags Flags = EFindObjectFlags::None, EObjectFlags ExcludeFlag = RF_NoFlags, EInternalObjectFlags ExclusiveInternalFlags = EInternalObjectFlags::None);
+template <UE::CSameAs<bool> ExactClassType>
+UE_EXACTCLASS_BOOL_DEPRECATED("StaticFindFirstObjectFastInternal")
+UE_NODEBUG FORCEINLINE UObject* StaticFindFirstObjectFastInternal(const UClass* ObjectClass, FName ObjectName, ExactClassType bExactClass, EObjectFlags ExcludeFlag = RF_NoFlags, EInternalObjectFlags ExclusiveInternalFlags = EInternalObjectFlags::None)
+{
+	return StaticFindFirstObjectFastInternal(ObjectClass, ObjectName, bExactClass ? EFindObjectFlags::ExactClass : EFindObjectFlags::None, ExcludeFlag, ExclusiveInternalFlags);
+}
 
 /**
  * Variation of StaticFindObjectFast that uses explicit path.
@@ -78,7 +98,13 @@ COREUOBJECT_API UObject* StaticFindFirstObjectFastInternal(const UClass* ObjectC
  * @param	ExclusiveFlags	Ignores objects that contain any of the specified exclusive flags
  * @return	Returns a pointer to the found object or NULL if none could be found
  */
-UObject* StaticFindObjectFastExplicit(const UClass* ObjectClass, FName ObjectName, const FString& ObjectPathName, bool bExactClass, EObjectFlags ExcludeFlags = RF_NoFlags);
+UObject* StaticFindObjectFastExplicit(const UClass* ObjectClass, FName ObjectName, const FString& ObjectPathName, EFindObjectFlags Flags = EFindObjectFlags::None, EObjectFlags ExcludeFlags = RF_NoFlags);
+template <UE::CSameAs<bool> ExactClassType>
+UE_EXACTCLASS_BOOL_DEPRECATED("StaticFindObjectFastExplicit")
+UE_NODEBUG FORCEINLINE UObject* StaticFindObjectFastExplicit(const UClass* ObjectClass, FName ObjectName, const FString& ObjectPathName, ExactClassType bExactClass, EObjectFlags ExcludeFlags = RF_NoFlags)
+{
+	return StaticFindObjectFastExplicit(ObjectClass, ObjectName, ObjectPathName, bExactClass ? EFindObjectFlags::ExactClass : EFindObjectFlags::None, ExcludeFlags);
+}
 
 /**
  * Variation of StaticFindObjectFastInternal that uses remote id.
