Assign Verse objects to the correct packages

#rnx

diff --git a/Engine/Source/Runtime/CoreUObject/Private/UObject/SavePackage2.cpp b/Engine/Source/Runtime/CoreUObject/Private/UObject/SavePackage2.cpp
--- a/Engine/Source/Runtime/CoreUObject/Private/UObject/SavePackage2.cpp
+++ b/Engine/Source/Runtime/CoreUObject/Private/UObject/SavePackage2.cpp
@@ -1364,9 +1364,15 @@ ESavePackageResult BuildLinker(FSaveContext& SaveContext)
 			VersePathMap.Add(VersePackage, &VersePackage->GetRootPath());
 			for (int32 Index = 0; Index < VersePackage->NumDefinitions(); ++Index)
 			{
-				Verse::VUniqueString* VersePath = &VersePackage->GetDefinitionName(Index);
-				Verse::VCell* Cell = VersePackage->GetDefinition(Index).Follow().ExtractCell();
-				VersePathMap.Add(Cell, VersePath);
+				if (Verse::VCell* Cell = VersePackage->GetDefinition(Index).ExtractCell())
+				{
+					// Packages can re-export cells from other packages; only use their original Verse path here.
+					if (Context.PackageForCell(Cell) == VersePackage)
+					{
+						Verse::VUniqueString* VersePath = &VersePackage->GetDefinitionName(Index);
+						VersePathMap.Add(Cell, VersePath);
+					}
+				}
 			}
 		}
 	}
