Fixing a crash in CppStructOps caused by unloading a dll at an inopportune time resulting in a dangling pointer to a vtable.

diff --git a/Engine/Source/Runtime/CoreUObject/Private/UObject/Class.cpp b/Engine/Source/Runtime/CoreUObject/Private/UObject/Class.cpp
--- a/Engine/Source/Runtime/CoreUObject/Private/UObject/Class.cpp
+++ b/Engine/Source/Runtime/CoreUObject/Private/UObject/Class.cpp
@@ -3080,6 +3080,21 @@ void UScriptStruct::DeferCppStructOps(FTopLevelAssetPath Target, ICppStructOps*
 	DeferredStructOps.Add(Target, InCppStructOps);
 }
 
+/**
+ * On module unload, remove the entry from the stash
+ * @param Target Name of the struct
+ */
+void UScriptStruct::RemoveDeferredCppStructOps(FTopLevelAssetPath Target)
+{
+	TMap<FTopLevelAssetPath, UScriptStruct::ICppStructOps*>& DeferredStructOps = GetDeferredCppStructOps();
+
+	UScriptStruct::ICppStructOps* ExistingOps = nullptr;
+	if (DeferredStructOps.RemoveAndCopyValue(Target, ExistingOps))
+	{
+		delete ExistingOps;
+	}
+}
+
 FTopLevelAssetPath UScriptStruct::GetFlattenedStructPathName() const
 {
 	return FTopLevelAssetPath(GetPackage()->GetFName(), GetFName());
diff --git a/Engine/Source/Runtime/CoreUObject/Public/UObject/Class.h b/Engine/Source/Runtime/CoreUObject/Public/UObject/Class.h
--- a/Engine/Source/Runtime/CoreUObject/Public/UObject/Class.h
+++ b/Engine/Source/Runtime/CoreUObject/Public/UObject/Class.h
@@ -1786,9 +1786,16 @@ public:
 	struct TAutoCppStructOps
 	{
 		TAutoCppStructOps(FTopLevelAssetPath InName)
+		: Name(InName)
 		{
-			DeferCppStructOps(InName,new TCppStructOps<CPPSTRUCT>);
+			DeferCppStructOps(Name,new TCppStructOps<CPPSTRUCT>);
 		}
+
+		~TAutoCppStructOps()
+		{
+			RemoveDeferredCppStructOps(Name);
+		}
+		FTopLevelAssetPath Name;
 	};
 	#define IMPLEMENT_STRUCT(BaseName) \
 		UE_DEPRECATED_MACRO(5.1, "IMPLEMENT_STRUCT has been deprecated. Use UE_IMPLEMENT_STRUCT and provide struct package name as well as struct name") static UScriptStruct::TAutoCppStructOps<F##BaseName> BaseName##_Ops(FTopLevelAssetPath(TEXT("/Script/CoreUObject"), TEXT(#BaseName))); 
@@ -1841,6 +1848,12 @@ public:
 	 */
 	static COREUOBJECT_API void DeferCppStructOps(FTopLevelAssetPath Target, ICppStructOps* InCppStructOps);
 
+	/**
+	 * On module unload, remove the entry from the stash
+	 * @param Target Name of the struct
+	 */
+	static COREUOBJECT_API void RemoveDeferredCppStructOps(FTopLevelAssetPath Target);
+
 	template<class CPPSTRUCT>
 	static void DeferCppStructOps(FTopLevelAssetPath Target)
 	{
