Deprecated bool parameters in CoreUObject FindObject functions.

diff --git a/Engine/Source/Runtime/CoreUObject/Public/UObject/FindObjectFlags.h b/Engine/Source/Runtime/CoreUObject/Public/UObject/FindObjectFlags.h
--- a/Engine/Source/Runtime/CoreUObject/Public/UObject/FindObjectFlags.h
+++ b/Engine/Source/Runtime/CoreUObject/Public/UObject/FindObjectFlags.h
@@ -13,4 +13,4 @@ enum class EFindObjectFlags
 };
 ENUM_CLASS_FLAGS(EFindObjectFlags)
 
-#define UE_EXACTCLASS_BOOL_DEPRECATED(FunctionName) //UE_DEPRECATED(5.7, FunctionName " with a boolean ExactClass has been deprecated - please use the EFindObjectFlags enum instead")
+#define UE_EXACTCLASS_BOOL_DEPRECATED(FunctionName) UE_DEPRECATED(5.7, FunctionName " with a boolean ExactClass has been deprecated - please use the EFindObjectFlags enum instead")
