Changed FPackageName::GetShortName and FPackageName::GetShortFName to take FStringView instead of const FString&.

diff --git a/Engine/Source/Runtime/CoreUObject/Private/Misc/PackageName.cpp b/Engine/Source/Runtime/CoreUObject/Private/Misc/PackageName.cpp
--- a/Engine/Source/Runtime/CoreUObject/Private/Misc/PackageName.cpp
+++ b/Engine/Source/Runtime/CoreUObject/Private/Misc/PackageName.cpp
@@ -91,12 +91,12 @@ bool FPackageName::IsShortPackageName(FName PossiblyLongName)
 	return IsShortPackageName(FStringView(Buffer, static_cast<int32>(Len)));
 }
 
-FString FPackageName::GetShortName(const FString& LongName)
+FString FPackageName::GetShortName(FStringView LongName)
 {
 	// Get everything after the last slash
 	int32 IndexOfLastSlash = INDEX_NONE;
 	LongName.FindLastChar(TEXT('/'), IndexOfLastSlash);
-	return LongName.Mid(IndexOfLastSlash + 1);
+	return FString(LongName.Mid(IndexOfLastSlash + 1));
 }
 
 FString FPackageName::GetShortName(const UPackage* Package)
@@ -115,9 +115,12 @@ FString FPackageName::GetShortName(const TCHAR* LongName)
 	return GetShortName(FString(LongName));
 }
 
-FName FPackageName::GetShortFName(const FString& LongName)
+FName FPackageName::GetShortFName(FStringView LongName)
 {
-	return GetShortFName(*LongName);
+	// Get everything after the last slash
+	int32 IndexOfLastSlash = INDEX_NONE;
+	LongName.FindLastChar(TEXT('/'), IndexOfLastSlash);
+	return FName(LongName.Mid(IndexOfLastSlash + 1));
 }
 
 FName FPackageName::GetShortFName(const FName& LongName)
diff --git a/Engine/Source/Runtime/CoreUObject/Public/Misc/PackageName.h b/Engine/Source/Runtime/CoreUObject/Public/Misc/PackageName.h
--- a/Engine/Source/Runtime/CoreUObject/Public/Misc/PackageName.h
+++ b/Engine/Source/Runtime/CoreUObject/Public/Misc/PackageName.h
@@ -286,7 +286,7 @@ public:
 	 * @param LongName Package name to convert.
 	 * @return Short package name.
 	 */
-	static COREUOBJECT_API FString GetShortName(const FString& LongName);
+	static COREUOBJECT_API FString GetShortName(FStringView LongName);
 	static COREUOBJECT_API FString GetShortName(const FName& LongName);
 	static COREUOBJECT_API FString GetShortName(const TCHAR* LongName);
 
@@ -296,7 +296,7 @@ public:
 	 * @param LongName Package name to convert.
 	 * @return Short package name.
 	 */
-	static COREUOBJECT_API FName GetShortFName(const FString& LongName);
+	static COREUOBJECT_API FName GetShortFName(FStringView LongName);
 	static COREUOBJECT_API FName GetShortFName(const FName& LongName);
 	static COREUOBJECT_API FName GetShortFName(const TCHAR* LongName);
 
