Do not time out sync loads due to GC.

- Under certain circumstances (when GameThread is taking a GC lock during an object duplication), sync loads would be timed out, triggering an assert in on the async loading thread.

diff --git a/Engine/Source/Runtime/CoreUObject/Private/Serialization/AsyncLoading2.cpp b/Engine/Source/Runtime/CoreUObject/Private/Serialization/AsyncLoading2.cpp
--- a/Engine/Source/Runtime/CoreUObject/Private/Serialization/AsyncLoading2.cpp
+++ b/Engine/Source/Runtime/CoreUObject/Private/Serialization/AsyncLoading2.cpp
@@ -3192,8 +3192,14 @@ struct FAsyncLoadingThreadState2
 
 		if (!bTimeLimitExceeded)
 		{
-			bTimeLimitExceeded = IsGarbageCollectionWaiting();
-			UE_CLOG(bTimeLimitExceeded, LogStreaming, Verbose, TEXT("Timing out async loading due to Garbage Collection request"));
+			const bool bGCIsWaiting = IsGarbageCollectionWaiting();
+			
+			// Do not time out sync loads due to GC 
+			if (bGCIsWaiting && SyncLoadContextStack.IsEmpty())
+			{
+				bTimeLimitExceeded = true;
+				UE_LOG(LogStreaming, Verbose, TEXT("Timing out async loading (last work %s) due to Garbage Collection request"), InLastTypeOfWorkPerformed);
+			}
 		}
 
 		return bTimeLimitExceeded;
