[UObjectBase]

UObjectLoadAllCompiledInDefaultProperties now sorts classes it will call GetDefaultObject on by class depth to allow parent types to be created first. This is to avoid potential cycles that may occur when sync loading objects from within PostCDOConstruct.

Specifically, UMaterialInterface may attempt to load built-in CDOs _synchronously_ from within a PostCDOConstruct method. In doing so, if dependent CDOs haven't been loaded yet, it's possible to recurse into UMaterialInterface::InitDefaultMaterials() and issue another load for the already loading package without allowing the loader to satisfy the outstanding CDO being waited for.
- Load ADefaultPawn CDO
- ADefaultPawn triggers UMaterialInterface::InitDefaultMaterials() which loads WorldGridMaterial
- WorldGridMaterial depends on Material
- Material's PostCDOConstruct triggers UMaterialInterface::InitDefaultMaterials() which loads WorldGridMaterial
- Hang

#rnx

diff --git a/Engine/Source/Runtime/CoreUObject/Private/UObject/UObjectBase.cpp b/Engine/Source/Runtime/CoreUObject/Private/UObject/UObjectBase.cpp
--- a/Engine/Source/Runtime/CoreUObject/Private/UObject/UObjectBase.cpp
+++ b/Engine/Source/Runtime/CoreUObject/Private/UObject/UObjectBase.cpp
@@ -858,6 +858,41 @@ static void UObjectLoadAllCompiledInDefaultProperties(TArray<UClass*>& OutAllNew
 				OutAllNewClasses.Add(&Class);
 			}); 
 
+		// Sort classes by depth to allow creating CDOs by parent type first. Otherwise, complex PostCDOConstruct calls 
+		// might rely on their parent having been initialized already, which might not be possible if there is a cycle in 
+		// the GetDefaultObject -> PostCDOConstruct loading graph.
+		{
+			TMap<const UClass*, int32> DepthMemo;
+			auto SortByClassDepth = [&DepthMemo](TArray<UClass*>& Classes)
+				{
+					Algo::StableSort(Classes, [&DepthMemo](const UClass* A, const UClass* B)
+						{
+							int32& ADepth = DepthMemo.FindOrAdd(A);
+							if (ADepth == 0)
+							{
+								do
+								{
+									ADepth++;
+									A = A->GetSuperClass();
+								} while (A);
+							}
+							int32& BDepth = DepthMemo.FindOrAdd(B);
+							if (BDepth == 0)
+							{
+								do
+								{
+									BDepth++;
+									B = B->GetSuperClass();
+								} while (B);
+							}
+							return ADepth < BDepth;
+						});
+				};
+			SortByClassDepth(NewClassesInCoreUObject);
+			SortByClassDepth(NewClassesInEngine);
+			SortByClassDepth(NewClasses);
+		}
+
 		auto NotifyClassFinishedRegistrationEvents = [](TArray<UClass*>& Classes)
 		{
 			for (UClass* Class : Classes)
