Transformed OS PropertyPath tracking from just a combined FName into TArray<FName> so we are not bound to the FName max length and we are not blowing up too much the FName table

[RN] Support for longer hierarchy in the entity trees

diff --git a/Engine/Source/Runtime/CoreUObject/Private/UObject/OverriddenPropertySet.cpp b/Engine/Source/Runtime/CoreUObject/Private/UObject/OverriddenPropertySet.cpp
--- a/Engine/Source/Runtime/CoreUObject/Private/UObject/OverriddenPropertySet.cpp
+++ b/Engine/Source/Runtime/CoreUObject/Private/UObject/OverriddenPropertySet.cpp
@@ -272,41 +272,85 @@ FOverridableTextPortPropertyPathScope::~FOverridableTextPortPropertyPathScope()
 }
 
 //----------------------------------------------------------------------//
-// FOverriddenPropertyNodeID
+// FOverriddenPropertyPath
 //----------------------------------------------------------------------//
-FOverriddenPropertyNodeID::FOverriddenPropertyNodeID(const FProperty* Property)
-	: Object(nullptr)
+FOverriddenPropertyPath::FOverriddenPropertyPath(const FName Name)
+: Path{Name}
+, CachedHash(GetTypeHash(Name))
+{}
+
+FOverriddenPropertyPath& FOverriddenPropertyPath::operator=(const FName Name)
 {
-	if (Property)
+	Path = {Name};
+	CachedHash = GetTypeHash(Name);
+	return *this;
+}
+
+FOverriddenPropertyPath FOverriddenPropertyPath::operator+(const FName Name) const
+{
+	FOverriddenPropertyPath NewPath(*this);
+	NewPath += Name;
+	return NewPath;
+}
+
+FOverriddenPropertyPath FOverriddenPropertyPath::operator+(const FOverriddenPropertyPath& Other) const
+{
+	FOverriddenPropertyPath NewPath(*this);
+	NewPath += Other;
+	return NewPath;
+}
+
+void FOverriddenPropertyPath::operator+=(const FName Name)
+{
+	Path.Emplace(Name);
+	CachedHash = HashCombine(CachedHash, GetTypeHash(Name));
+}
+
+void FOverriddenPropertyPath::operator+=(const FOverriddenPropertyPath& Other)
+{
+	for(FName Name : Other.Path)
 	{
-		// append typename to the end of the property ID
-		UE::FPropertyTypeNameBuilder TypeNameBuilder;
+		this->operator+=(Name);
+	}
+}
+
+FString FOverriddenPropertyPath::ToString() const
+{
+	return FString::JoinBy(Path, TEXT("."), [](FName Name) { return Name.ToString(); });
+}
+
+//----------------------------------------------------------------------//
+// FOverriddenPropertyNodeID
+//----------------------------------------------------------------------//
+FOverriddenPropertyNodeID::FOverriddenPropertyNodeID(TNotNull<const FProperty*> Property)
+{
+	// append typename to the end of the property ID
+	UE::FPropertyTypeNameBuilder TypeNameBuilder;
 #if WITH_EDITORONLY_DATA
-		{
-			// use property impersonation for SaveTypeName so that keys don't change when classes die
-			FUObjectSerializeContext* SerializeContext = FUObjectThreadContext::Get().GetSerializeContext();
-            TGuardValue<bool> ScopedImpersonateProperties(SerializeContext->bImpersonateProperties, true);
-            Property->SaveTypeName(TypeNameBuilder);
-		}
-#endif
-		UE::FPropertyTypeName TypeName = TypeNameBuilder.Build();
-		TStringBuilder<256> StringBuilder;
-		StringBuilder << Property->GetFName();
-		StringBuilder << " - ";
-		StringBuilder << TypeName;
-		Path = FName(StringBuilder.ToView());
+	{
+		// use property impersonation for SaveTypeName so that keys don't change when classes die
+		FUObjectSerializeContext* SerializeContext = FUObjectThreadContext::Get().GetSerializeContext();
+        TGuardValue<bool> ScopedImpersonateProperties(SerializeContext->bImpersonateProperties, true);
+        Property->SaveTypeName(TypeNameBuilder);
 	}
+#endif
+	UE::FPropertyTypeName TypeName = TypeNameBuilder.Build();
+	Path = Property->GetFName();
+	Path += TypeName.GetName();
 }
 
 FOverriddenPropertyNodeID::FOverriddenPropertyNodeID(const FOverriddenPropertyNodeID& ParentNodeID, const FOverriddenPropertyNodeID& NodeID)
+	: Path(ParentNodeID.Path + NodeID.Path) // Combine the 2 node ids
+	, Object(NodeID.Object)
+{
+}
+
+FOverriddenPropertyNodeID::FOverriddenPropertyNodeID(TNotNull<const UObject*> InObject): Object(InObject)
 {
-	// Combine the 2 node ids
-	FStringBuilderBase SubPropertyKeyBuilder;
-	SubPropertyKeyBuilder = *ParentNodeID.ToString();
-	SubPropertyKeyBuilder.Append(TEXT("."));
-	SubPropertyKeyBuilder.Append(*NodeID.ToString());
-	Path = FName(SubPropertyKeyBuilder.ToString());
-	Object = NodeID.Object;
+	// Note: Using ObjectIndex by itself is not sufficient for an enduring unique identifier
+	// as re-instantiation can cause a reuse of the index for another object. Appending the serial solves this issue
+	const int32 ObjectIndex = GUObjectArray.ObjectToIndex(InObject);
+	Path = FName(FString::Printf(TEXT("ID_%d_%d"), ObjectIndex, GUObjectArray.AllocateSerialNumber(ObjectIndex)));
 }
 
 FOverriddenPropertyNodeID FOverriddenPropertyNodeID::RootNodeId()
@@ -379,7 +423,7 @@ bool FOverriddenPropertyNodeID::operator==(const FOverriddenPropertyNodeID& Othe
 		return true;
 	}
 
-	// After reinstanciation, we do not change the path we only patch the ptr.
+	// After reinstantiation, we do not change the path we only patch the ptr.
 	// Unfortunately we do not have a stable id that is stable through reinstantiation,
 	// so are only way if to compare ptr.
 	if (Object && Other.Object && Object == Other.Object)
@@ -468,7 +512,7 @@ EOverriddenPropertyOperation FOverriddenPropertySet::GetOverriddenPropertyOperat
 		const FOverriddenPropertyNode* CurrentOverriddenPropertyNode = nullptr;
 		if (OverriddenPropertyNode)
 		{
-			if (const FOverriddenPropertyNodeID* CurrentPropKey = OverriddenPropertyNode->SubPropertyNodeKeys.Find(CurrentProperty))
+			if (const FOverriddenPropertyNodeID* CurrentPropKey = OverriddenPropertyNode->SubPropertyNodeKeys.Find(FOverriddenPropertyNodeID(CurrentProperty)))
 			{
 				CurrentOverriddenPropertyNode = OverriddenPropertyNodes.Find(*CurrentPropKey);
 				checkf(CurrentOverriddenPropertyNode, TEXT("Expecting a node"));
@@ -619,7 +663,7 @@ bool FOverriddenPropertySet::ClearOverriddenProperty(FOverriddenPropertyNode& Pa
 		FOverriddenPropertyNode* CurrentOverriddenPropertyNode = nullptr;
 		if (OverriddenPropertyNode)
 		{
-			if (const FOverriddenPropertyNodeID* CurrentPropKey = OverriddenPropertyNode->SubPropertyNodeKeys.Find(CurrentProperty))
+			if (const FOverriddenPropertyNodeID* CurrentPropKey = OverriddenPropertyNode->SubPropertyNodeKeys.Find(FOverriddenPropertyNodeID(CurrentProperty)))
 			{
 				CurrentOverriddenPropertyNode = OverriddenPropertyNodes.Find(*CurrentPropKey);
 				checkf(CurrentOverriddenPropertyNode, TEXT("Expecting a node"));
@@ -847,7 +891,7 @@ void FOverriddenPropertySet::NotifyPropertyChange(FOverriddenPropertyNode* Paren
 	FOverriddenPropertyNode* SubPropertyNode = nullptr;
 	if (ParentPropertyNode)
 	{
-		FOverriddenPropertyNode& SubPropertyNodeRef = FindOrAddNode(*ParentPropertyNode, Property);
+		FOverriddenPropertyNode& SubPropertyNodeRef = FindOrAddNode(*ParentPropertyNode, FOverriddenPropertyNodeID(Property));
 		SubPropertyNode = SubPropertyNodeRef.Operation != EOverriddenPropertyOperation::Replace ? &SubPropertyNodeRef : nullptr;
 	}
 
@@ -864,7 +908,7 @@ void FOverriddenPropertySet::NotifyPropertyChange(FOverriddenPropertyNode* Paren
 			 SubPropertyNode->Operation == EOverriddenPropertyOperation::Modified))
 		{
 			FOverriddenPropertyNodeID RemovedNodeID;
-			if (ParentPropertyNode->SubPropertyNodeKeys.RemoveAndCopyValue(Property, RemovedNodeID))
+			if (ParentPropertyNode->SubPropertyNodeKeys.RemoveAndCopyValue(FOverriddenPropertyNodeID(Property), RemovedNodeID))
 			{
 				verifyf(OverriddenPropertyNodes.Remove(RemovedNodeID), TEXT("Expecting the node to be removed"));
 			}
@@ -1537,7 +1581,7 @@ EOverriddenPropertyOperation FOverriddenPropertySet::GetOverriddenPropertyOperat
 		const FProperty* CurrentProperty = (*PropertyIterator);
 		if (OverriddenPropertyNode)
 		{
-			if (const FOverriddenPropertyNodeID* CurrentPropKey = OverriddenPropertyNode->SubPropertyNodeKeys.Find(CurrentProperty))
+			if (const FOverriddenPropertyNodeID* CurrentPropKey = OverriddenPropertyNode->SubPropertyNodeKeys.Find(FOverriddenPropertyNodeID(CurrentProperty)))
 			{
 				OverriddenPropertyNode = OverriddenPropertyNodes.Find(*CurrentPropKey);
 				checkf(OverriddenPropertyNode, TEXT("Expecting a node"));
@@ -1587,7 +1631,7 @@ FOverriddenPropertyNode* FOverriddenPropertySet::SetOverriddenPropertyOperation(
 		{
 			return nullptr;
 		}
-		OverriddenPropertyNode = &FindOrAddNode(*OverriddenPropertyNode, CurrentProperty);
+		OverriddenPropertyNode = &FindOrAddNode(*OverriddenPropertyNode, FOverriddenPropertyNodeID(CurrentProperty));
 		++PropertyIterator;
 	}
 
@@ -1710,7 +1754,7 @@ const FOverriddenPropertyNode* FOverriddenPropertySet::GetOverriddenPropertyNode
 	while (PropertyIterator && OverriddenPropertyNode)
 	{
 		const FProperty* CurrentProperty = (*PropertyIterator);
-		if (const FOverriddenPropertyNodeID* CurrentPropKey = OverriddenPropertyNode->SubPropertyNodeKeys.Find(CurrentProperty))
+		if (const FOverriddenPropertyNodeID* CurrentPropKey = OverriddenPropertyNode->SubPropertyNodeKeys.Find(FOverriddenPropertyNodeID(CurrentProperty)))
 		{
 			OverriddenPropertyNode = OverriddenPropertyNodes.Find(*CurrentPropKey);
 			checkf(OverriddenPropertyNode, TEXT("Expecting a node"));
diff --git a/Engine/Source/Runtime/CoreUObject/Public/UObject/OverriddenPropertySet.h b/Engine/Source/Runtime/CoreUObject/Public/UObject/OverriddenPropertySet.h
--- a/Engine/Source/Runtime/CoreUObject/Public/UObject/OverriddenPropertySet.h
+++ b/Engine/Source/Runtime/CoreUObject/Public/UObject/OverriddenPropertySet.h
@@ -205,21 +205,63 @@ inline FString GetOverriddenOperationString(EOverriddenPropertyOperation Operati
 }
 
 USTRUCT()
-struct FOverriddenPropertyNodeID
+struct FOverriddenPropertyPath
 {
 	GENERATED_BODY()
 
-	FOverriddenPropertyNodeID(const FProperty* Property = nullptr);
+	FOverriddenPropertyPath(const FName Name);
+
+	FOverriddenPropertyPath() = default;
+	FOverriddenPropertyPath(const FOverriddenPropertyPath&) = default;
+	FOverriddenPropertyPath(FOverriddenPropertyPath&&) = default;
+	FOverriddenPropertyPath& operator=(const FOverriddenPropertyPath&) = default;
+	FOverriddenPropertyPath& operator=(FOverriddenPropertyPath&&) = default;
 
-	FOverriddenPropertyNodeID(TNotNull<const UObject*> InObject)
-		: Object(InObject)
+	FOverriddenPropertyPath& operator=(const FName Name);
+	FOverriddenPropertyPath operator+(const FName Name) const;
+	FOverriddenPropertyPath operator+(const FOverriddenPropertyPath& Other) const;
+	void operator+=(const FName Name);
+	void operator+=(const FOverriddenPropertyPath& Other);
+
+	bool operator==(const FOverriddenPropertyPath& Other) const
 	{
-		// Note: Using ObjectIndex by itself is not sufficient for an enduring unique identifier
-		// as re-instantiation can cause a reuse of the index for another object. Appending the serial solves this issue
-		const int32 ObjectIndex = GUObjectArray.ObjectToIndex(InObject);
-		Path = *(FString::Printf(TEXT("%d%d"), ObjectIndex, GUObjectArray.AllocateSerialNumber(ObjectIndex)));
+		return CachedHash == Other.CachedHash;
 	}
-	FOverriddenPropertyNodeID(const FOverriddenPropertyNodeID& ParentNodeID, const FOverriddenPropertyNodeID& SubNodeID);
+
+	bool IsEmpty() const
+	{
+		return Path.IsEmpty();
+	}
+
+	FString ToString() const;
+
+private:
+	UPROPERTY()
+	TArray<FName> Path;
+
+	UPROPERTY()
+	uint32 CachedHash = 0;
+
+	friend uint32 GetTypeHash(const FOverriddenPropertyPath& MappableString)
+	{
+		return MappableString.CachedHash;
+	}
+};
+
+USTRUCT()
+struct FOverriddenPropertyNodeID
+{
+	GENERATED_BODY()
+
+	COREUOBJECT_API FOverriddenPropertyNodeID(TNotNull<const FProperty*> Property);
+	COREUOBJECT_API FOverriddenPropertyNodeID(TNotNull<const UObject*> InObject);
+	COREUOBJECT_API FOverriddenPropertyNodeID(const FOverriddenPropertyNodeID& ParentNodeID, const FOverriddenPropertyNodeID& SubNodeID);
+
+	FOverriddenPropertyNodeID() = default;
+	FOverriddenPropertyNodeID(const FOverriddenPropertyNodeID&) = default;
+	FOverriddenPropertyNodeID(FOverriddenPropertyNodeID&&) = default;
+	FOverriddenPropertyNodeID& operator=(const FOverriddenPropertyNodeID&) = default;
+	FOverriddenPropertyNodeID& operator=(FOverriddenPropertyNodeID&&) = default;
 
 	static COREUOBJECT_API FOverriddenPropertyNodeID RootNodeId();
 
@@ -241,7 +283,7 @@ struct FOverriddenPropertyNodeID
 
 	bool IsValid() const
 	{
-		return !Path.IsNone();
+		return !Path.IsEmpty();
 	}
 
 	void HandleObjectsReInstantiated(const TMap<UObject*, UObject*>& Map);
@@ -255,7 +297,7 @@ struct FOverriddenPropertyNodeID
 
 private:
 	UPROPERTY()
-	FName Path;
+	FOverriddenPropertyPath Path;
 
 	/* NOTE: Not always valid can point to a destroyed or can be a stale ptr. Use with cautiousness
 	 * This can happen during reinstantiation. It is only there for ptr comparison in the ==.
@@ -279,9 +321,18 @@ struct FOverriddenPropertyNode
 {
 	GENERATED_BODY()
 
-	FOverriddenPropertyNode(const FOverriddenPropertyNodeID& InNodeID = FOverriddenPropertyNodeID())
+	FOverriddenPropertyNode() = default;
+	FOverriddenPropertyNode(const FOverriddenPropertyNode&) = default;
+	FOverriddenPropertyNode(FOverriddenPropertyNode&&) = default;
+	FOverriddenPropertyNode& operator=(const FOverriddenPropertyNode&) = default;
+	FOverriddenPropertyNode& operator=(FOverriddenPropertyNode&&) = default;
+
+	FOverriddenPropertyNode(const FOverriddenPropertyNodeID& InNodeID)
 		: NodeID(InNodeID)
 	{}
+	FOverriddenPropertyNode(FOverriddenPropertyNodeID&& InNodeID)
+		: NodeID(MoveTemp(InNodeID))
+	{}
 
 	UPROPERTY()
 	FOverriddenPropertyNodeID NodeID;
