SavePackage: Never build and save searchable names when cooking.
Tiny optimization/cleanup that fixes "Searchable object X is not mapped" ensures that can trigger during cook if MarkSearchableName has been called for an excluded object (e.g via NeedsLoadForClient/Server).

#rnx

diff --git a/Engine/Source/Runtime/CoreUObject/Private/UObject/SavePackage2.cpp b/Engine/Source/Runtime/CoreUObject/Private/UObject/SavePackage2.cpp
--- a/Engine/Source/Runtime/CoreUObject/Private/UObject/SavePackage2.cpp
+++ b/Engine/Source/Runtime/CoreUObject/Private/UObject/SavePackage2.cpp
@@ -1538,16 +1538,20 @@ ESavePackageResult BuildLinker(FSaveContext& SaveContext)
 
 	// Build SoftPackageReference & Searchable Name Map
 	{
+		SCOPED_SAVETIMER(UPackage_Save_BuildSoftPackagesAndSearchableNames);
 		Linker->SoftPackageReferenceList = SaveContext.GetSoftPackageReferenceList().Array();
 
-		// Convert the searchable names map from UObject to packageindex
-		for (TPair<TObjectPtr<UObject>, TArray<FName>>& SearchableNamePair : SaveContext.GetSearchableNamesObjectMap())
+		if (!SaveContext.IsFilterEditorOnly())
 		{
-			const FPackageIndex PackageIndex = Linker->MapObject(SearchableNamePair.Key);
-			// This should always be in the imports already
-			if (ensure(!PackageIndex.IsNull()))
+			// Convert the searchable names map from UObject to packageindex
+			for (TPair<TObjectPtr<UObject>, TArray<FName>>& SearchableNamePair : SaveContext.GetSearchableNamesObjectMap())
 			{
-				Linker->SearchableNamesMap.FindOrAdd(PackageIndex) = MoveTemp(SearchableNamePair.Value);
+				const FPackageIndex PackageIndex = Linker->MapObject(SearchableNamePair.Key);
+				// This should always be in the imports already
+				if (ensureMsgf(!PackageIndex.IsNull(), TEXT("%s: Searchable object %s with is not mapped"), *Linker->LinkerRoot->GetFullName(), *SearchableNamePair.Key->GetFullName()))
+				{
+					Linker->SearchableNamesMap.FindOrAdd(PackageIndex) = MoveTemp(SearchableNamePair.Value);
+				}
 			}
 		}
 		SaveContext.GetSearchableNamesObjectMap().Empty();
@@ -2304,8 +2308,15 @@ ESavePackageResult WritePackageHeader(FStructuredArchive::FRecord& StructuredArc
 			}
 
 			// Save searchable names map
-			Linker->Summary.SearchableNamesOffset = (int32)Linker->Tell();
-			Linker->SerializeSearchableNamesMap(StructuredArchiveRoot.EnterField(TEXT("SearchableNames")));
+			if (!SaveContext.IsFilterEditorOnly())
+			{
+				Linker->Summary.SearchableNamesOffset = (int32)Linker->Tell();
+				Linker->SerializeSearchableNamesMap(StructuredArchiveRoot.EnterField(TEXT("SearchableNames")));
+			}
+			else
+			{
+				Linker->Summary.SearchableNamesOffset = 0;
+			}
 		}
 	}
 	else
