Fix for using a template object with an incompatible type in ConstructDefaultObjectValueIfNecessary.

#rnx

diff --git a/Engine/Source/Runtime/CoreUObject/Private/UObject/PropertyBaseObject.cpp b/Engine/Source/Runtime/CoreUObject/Private/UObject/PropertyBaseObject.cpp
--- a/Engine/Source/Runtime/CoreUObject/Private/UObject/PropertyBaseObject.cpp
+++ b/Engine/Source/Runtime/CoreUObject/Private/UObject/PropertyBaseObject.cpp
@@ -851,6 +851,10 @@ UObject* FObjectPropertyBase::ConstructDefaultObjectValueIfNecessary(UObject* Ex
 			FName NewName = (bCanReplaceObjectWithExistingName) ? ExistingName : NAME_None;
 			EObjectFlags NewFlags = ExistingValue ? (ExistingValue->GetFlags() & FlagsToKeepMask) : RF_NoFlags;
 			UObject* Template = Defaults ? GetObjectPropertyValue(Defaults) : nullptr;
+			if (Template && !Template->IsA(PropertyClass))
+			{
+				Template = nullptr;
+			}
 
 			// Existing value did not exist or it could not be used as a template
 			// Existing value may be null in case we were serializing an array of UObjects that failed to load (missing import class). Since the array is first pre-allocated with null values
