Fix some false negative references detected in OverridableManager on a map about to be unloaded.

#rnx

diff --git a/Engine/Source/Runtime/CoreUObject/Private/UObject/OverridableManager.cpp b/Engine/Source/Runtime/CoreUObject/Private/UObject/OverridableManager.cpp
--- a/Engine/Source/Runtime/CoreUObject/Private/UObject/OverridableManager.cpp
+++ b/Engine/Source/Runtime/CoreUObject/Private/UObject/OverridableManager.cpp
@@ -406,25 +406,6 @@ void FOverridableManager::HandleObjectsReInstantiated(const TMap<UObject*, UObje
 #endif // WITH_EDITORONLY_DATA
 }
 
-void FOverridableManager::AddReferencedObjects(FReferenceCollector& Collector)
-{
-#if WITH_EDITORONLY_DATA
-	// This isn't great but this object implements this FGCObject interface just to so 
-	// that the replace references archives go through this object
-	if (!GIsGarbageCollecting)
-	{
-		UE::TRWScopeLock AnnotationMapLock(AnnotationMapCritical, SLT_ReadOnly);
-		for (const auto& Pair : GetAnnotationMap())
-		{
-			if (FOverriddenPropertySet* OverriddenProperties = Pair.Value.OverriddenProperties.Get())
-			{
-				OverriddenProperties->AddReferencedObjects(Collector);
-			}
-		}
-	}
-#endif // WITH_EDITORONLY_DATA
-}
-
 void FOverridableManager::HandleDeadObjectReferences(const TSet<UClass*>& DeadClasses, const TSet<UObject*>& ActiveInstances, const TSet<UObject*>& TemplateInstances)
 {
 #if WITH_EDITORONLY_DATA
diff --git a/Engine/Source/Runtime/CoreUObject/Private/UObject/OverriddenPropertySet.cpp b/Engine/Source/Runtime/CoreUObject/Private/UObject/OverriddenPropertySet.cpp
--- a/Engine/Source/Runtime/CoreUObject/Private/UObject/OverriddenPropertySet.cpp
+++ b/Engine/Source/Runtime/CoreUObject/Private/UObject/OverriddenPropertySet.cpp
@@ -345,7 +345,8 @@ FOverriddenPropertyNodeID::FOverriddenPropertyNodeID(const FOverriddenPropertyNo
 {
 }
 
-FOverriddenPropertyNodeID::FOverriddenPropertyNodeID(TNotNull<const UObject*> InObject): Object(InObject)
+FOverriddenPropertyNodeID::FOverriddenPropertyNodeID(TNotNull<const UObject*> InObject)
+: Object(NotNullGet(InObject))
 {
 	// Note: Using ObjectIndex by itself is not sufficient for an enduring unique identifier
 	// as re-instantiation can cause a reuse of the index for another object. Appending the serial solves this issue
@@ -426,7 +427,9 @@ bool FOverriddenPropertyNodeID::operator==(const FOverriddenPropertyNodeID& Othe
 	// After reinstantiation, we do not change the path we only patch the ptr.
 	// Unfortunately we do not have a stable id that is stable through reinstantiation,
 	// so are only way if to compare ptr.
-	if (Object && Other.Object && Object == Other.Object)
+	if (Object.IsValid(/*bEvenIfGarbage*/true) && 
+		Other.Object.IsValid(/*bEvenIfGarbage*/true) && 
+		Object == Other.Object)
 	{
 		return true;
 	}
@@ -436,32 +439,23 @@ bool FOverriddenPropertyNodeID::operator==(const FOverriddenPropertyNodeID& Othe
 
 void FOverriddenPropertyNodeID::HandleObjectsReInstantiated(const TMap<UObject*, UObject*>& Map)
 {
-	if (!Object)
+	if (const UObject* ObjectPtr = Object.Get(/*bEvenIfGarbage*/true))
 	{
-		return;
-	}
-
-	if (UObject*const* ReplacedObject = Map.Find(Object))
-	{
-		Object = *ReplacedObject;
+		if (UObject*const* ReplacedObject = Map.Find(ObjectPtr))
+		{
+			Object = *ReplacedObject;
+		}
 	}
 }
 
-void FOverriddenPropertyNodeID::AddReferencedObjects(FReferenceCollector& Collector, UObject* Owner)
-{
-	Collector.AddReferencedObject(Object, Owner);
-}
-
 void FOverriddenPropertyNodeID::HandleDeadObjectReferences(const TSet<UObject*>& ActiveInstances, const TSet<UObject*>& TemplateInstances)
 {
-	if (!Object)
-	{
-		return;
-	}
-
-	if (ActiveInstances.Find(Object) || TemplateInstances.Find(Object) )
+	if (const UObject* ObjectPtr = Object.Get(/*bEvenIfGarbage*/true))
 	{
-		Object = nullptr;
+		if (ActiveInstances.Find(ObjectPtr) || TemplateInstances.Find(ObjectPtr) )
+		{
+			Object.Reset();
+		}
 	}
 }
 
@@ -1716,19 +1710,6 @@ void FOverriddenPropertySet::HandleObjectsReInstantiated(const TMap<UObject*, UO
 	}
 }
 
-void FOverriddenPropertySet::AddReferencedObjects(FReferenceCollector& Collector)
-{
-	for (FOverriddenPropertyNode& Node : OverriddenPropertyNodes)
-	{
-		Node.NodeID.AddReferencedObjects(Collector, Owner);
-		for (auto& Pair : Node.SubPropertyNodeKeys)
-		{
-			Pair.Key.AddReferencedObjects(Collector, Owner);
-			Pair.Value.AddReferencedObjects(Collector, Owner);
-		}
-	}
-}
-
 void FOverriddenPropertySet::HandleDeadObjectReferences(const TSet<UObject*>& ActiveInstances, const TSet<UObject*>& TemplateInstances)
 {
 	for (FOverriddenPropertyNode& Node : OverriddenPropertyNodes)
diff --git a/Engine/Source/Runtime/CoreUObject/Public/UObject/OverridableManager.h b/Engine/Source/Runtime/CoreUObject/Public/UObject/OverridableManager.h
--- a/Engine/Source/Runtime/CoreUObject/Public/UObject/OverridableManager.h
+++ b/Engine/Source/Runtime/CoreUObject/Public/UObject/OverridableManager.h
@@ -51,9 +51,9 @@ struct FOverrideManagerCapabilityInterface;
  * Overridable serialization is experimental, not supported and use at your own risk *
  *************************************************************************************
  */
-class FOverridableManager : public FGCObject
+class FOverridableManager 
 #if WITH_EDITORONLY_DATA
-	, protected FUObjectAnnotationSparse<FOverriddenPropertyAnnotation, true/*bAutoRemove*/>
+	: protected FUObjectAnnotationSparse<FOverriddenPropertyAnnotation, true/*bAutoRemove*/>
 #endif // WITH_EDITORONLY_DATA
 {
 public:
@@ -295,15 +295,6 @@ public:
 	void HandleObjectsReInstantiated(const TMap<UObject*, UObject*>& OldToNewInstanceMap);
 	void HandleDeadObjectReferences(const TSet<UClass*>& DeadClasses, const TSet<UObject*>& ActiveInstances, const TSet<UObject*>& TemplateInstances);
 
-	// FSerializableObject interface
-	virtual void AddReferencedObjects(FReferenceCollector& Collector) override;
-	virtual FString GetReferencerName() const override
-	{
-		return TEXT("FOverridableManager");
-	}
-	// End of FSerializableObject interface
-
-
 protected:
 	FOverridableManager();
 
diff --git a/Engine/Source/Runtime/CoreUObject/Public/UObject/OverriddenPropertySet.h b/Engine/Source/Runtime/CoreUObject/Public/UObject/OverriddenPropertySet.h
--- a/Engine/Source/Runtime/CoreUObject/Public/UObject/OverriddenPropertySet.h
+++ b/Engine/Source/Runtime/CoreUObject/Public/UObject/OverriddenPropertySet.h
@@ -288,7 +288,6 @@ struct FOverriddenPropertyNodeID
 
 	void HandleObjectsReInstantiated(const TMap<UObject*, UObject*>& Map);
 	void HandleDeadObjectReferences(const TSet<UObject*>& ActiveInstances, const TSet<UObject*>& TemplateInstances);
-	void AddReferencedObjects(FReferenceCollector& Collector, UObject* Owner);
 
 	friend uint32 GetTypeHash(const FOverriddenPropertyNodeID& NodeID)
 	{
@@ -299,14 +298,8 @@ private:
 	UPROPERTY()
 	FOverriddenPropertyPath Path;
 
-	/* NOTE: Not always valid can point to a destroyed or can be a stale ptr. Use with cautiousness
-	 * This can happen during reinstantiation. It is only there for ptr comparison in the ==.
-	 * There is no unique id kept between old and new when an object is reinstantiated, so we are keeping the ptr here.
-	 * When it is a ptr of an archetype, there are some cases we do not patch it because we need it in the CPFUO.
-	 * Check FOverriddenPropertySet::HandleObjectsReInstantiated special exception
-	 */
 	UPROPERTY()
-	TObjectPtr<const UObject> Object;
+	TWeakObjectPtr<const UObject> Object;
 };
 
 /*
@@ -474,7 +467,6 @@ public:
 	UE_API void Reset();
 	UE_API void HandleObjectsReInstantiated(const TMap<UObject*, UObject*>& Map);
 	UE_API void HandleDeadObjectReferences(const TSet<UObject*>& ActiveInstances, const TSet<UObject*>& TemplateInstances);
-	UE_API void AddReferencedObjects(FReferenceCollector& Collector);
 
 	/* return whether this object is considered added or not */
 	bool WasAdded() const
