Fixed natvis for FUObjectItem

diff --git a/Engine/Source/Runtime/CoreUObject/Public/UObject/UObjectArray.h b/Engine/Source/Runtime/CoreUObject/Public/UObject/UObjectArray.h
--- a/Engine/Source/Runtime/CoreUObject/Public/UObject/UObjectArray.h
+++ b/Engine/Source/Runtime/CoreUObject/Public/UObject/UObjectArray.h
@@ -61,9 +61,23 @@ private:
 	int64 FlagsAndRefCount;
 
 #if !UE_ENABLE_FUOBJECT_ITEM_PACKING
-	// Pointer to the allocated object
-	class UObjectBase* Object;
+public:
+	union
+	{
+		// Pointer to the allocated object
+		UE_DEPRECATED(5.5, "Use GetObject() and SetObject() to access Object.")
+		class UObjectBase* Object = nullptr;
+		uint32 ObjectPtrLow;	// this one is used as a dummy for natvis only an will be removed once packing will be enabled by default
+	};
 #else
+	union
+	{
+		// Stores lower 32 bits of UObject pointer shifted by 3 to the left as all our allocations are at least 8 bytes aligned and lower 3 bits will always be 0
+		uint32 ObjectPtrLow = 0;
+		uint32 Object;	// this one is used as a dummy for natvis only an will be removed once packing will be enabled by default
+	};
+#endif
+private:
 	// Currently we assume UObjects are aligned by 8 bytes, that gives us 3 lower bits as zeros that we can discard.
 	// This will give us total 45 bits in a pointer that we pack into a int32 and the remaining 13 bits we pack with Flags
 	// EInternalObjectFlags_MinFlagBitIndex at the time of writing this is 14 and we have only 1 bit left in the EInternalObjectFlags for future use
@@ -74,10 +88,6 @@ private:
 	constexpr static int32 FlagsMask = 0xFFFFFFFF << int(EInternalObjectFlags_MinFlagBitIndex);
 	constexpr static int32 PtrMask = ~FlagsMask;
 
-	// Stores lower 32 bits of UObject pointer shifted by 3 to the left as all our allocations are at least 8 bytes aligned and lower 3 bits will always be 0
-	uint32 ObjectPtrLow;
-#endif
-
 public:
 	// Weak Object Pointer Serial number associated with the object
 	int32 SerialNumber;
@@ -100,11 +110,6 @@ public:
 
 	FUObjectItem()
 		: FlagsAndRefCount(0)
-#if !UE_ENABLE_FUOBJECT_ITEM_PACKING
-		, Object(nullptr)
-#else
-		, ObjectPtrLow(0)
-#endif
 		, SerialNumber(0)
 		, ClusterRootIndex(0)
 #if ENABLE_STATNAMEDEVENTS_UOBJECT
@@ -132,8 +137,8 @@ public:
 	FORCEINLINE class UObjectBase* GetObject() const
 	{
 #if UE_ENABLE_FUOBJECT_ITEM_PACKING
-		const uintptr_t Object = (uintptr_t((FlagsAndRefCount >> 32) & PtrMask) << (32 + UObjectPtrTrailingZeroes)) | (uintptr_t(ObjectPtrLow) << UObjectPtrTrailingZeroes);
-		return (class UObjectBase*)Object;
+		const uintptr_t Obj = (uintptr_t((FlagsAndRefCount >> 32) & PtrMask) << (32 + UObjectPtrTrailingZeroes)) | (uintptr_t(ObjectPtrLow) << UObjectPtrTrailingZeroes);
+		return (class UObjectBase*)Obj;
 #else
 		return Object;
 #endif
@@ -680,6 +685,8 @@ class FChunkedFixedUObjectArray
 	/** Number of chunks we currently have **/
 	TSAN_ATOMIC(int32) NumChunks;
 
+	static constexpr bool bFUObjectItemIsPacked = UE_ENABLE_FUOBJECT_ITEM_PACKING;
+
 
 	/**
 	* Allocates new chunk for the array
