Suppressed deprecated Object calls in FUobjectItem

diff --git a/Engine/Source/Runtime/CoreUObject/Public/UObject/UObjectArray.h b/Engine/Source/Runtime/CoreUObject/Public/UObject/UObjectArray.h
--- a/Engine/Source/Runtime/CoreUObject/Public/UObject/UObjectArray.h
+++ b/Engine/Source/Runtime/CoreUObject/Public/UObject/UObjectArray.h
@@ -140,7 +140,9 @@ public:
 		const uintptr_t Obj = (uintptr_t((FlagsAndRefCount >> 32) & PtrMask) << (32 + UObjectPtrTrailingZeroes)) | (uintptr_t(ObjectPtrLow) << UObjectPtrTrailingZeroes);
 		return (class UObjectBase*)Obj;
 #else
+		PRAGMA_DISABLE_DEPRECATION_WARNINGS
 		return Object;
+		PRAGMA_ENABLE_DEPRECATION_WARNINGS
 #endif
 	}
 
@@ -152,7 +154,9 @@ public:
 		FlagsAndRefCount |= ((uint64)ObjectPtrHi) << 32;
 		ObjectPtrLow = ((uintptr_t)InObject >> UObjectPtrTrailingZeroes) & 0xFFFFFFFF;
 #else
+		PRAGMA_DISABLE_DEPRECATION_WARNINGS
 		Object = InObject;
+		PRAGMA_ENABLE_DEPRECATION_WARNINGS
 #endif
 	}
 
