Removing UInterface from NoExportTypes

diff --git a/Engine/Source/Runtime/CoreUObject/Private/UObject/Interface.cpp b/Engine/Source/Runtime/CoreUObject/Private/UObject/Interface.cpp
--- a/Engine/Source/Runtime/CoreUObject/Private/UObject/Interface.cpp
+++ b/Engine/Source/Runtime/CoreUObject/Private/UObject/Interface.cpp
@@ -3,21 +3,3 @@
 #include "UObject/Interface.h"
 #include "UObject/Class.h"
 
-#if WITH_EDITOR
-IMPLEMENT_CORE_INTRINSIC_CLASS(UInterface, UObject,
-	{
-		Class->CppClassStaticFunctions = UOBJECT_CPPCLASS_STATICFUNCTIONS_FORCLASS(UInterface);
-		Class->SetMetaData(TEXT("IsBlueprintBase"), TEXT("true"));
-		Class->SetMetaData(TEXT("CannotImplementInterfaceInBlueprint"), TEXT(""));
-	}
-);
-
-#else
-
-IMPLEMENT_CORE_INTRINSIC_CLASS(UInterface, UObject,
-{
-	Class->CppClassStaticFunctions = UOBJECT_CPPCLASS_STATICFUNCTIONS_FORCLASS(UInterface);
-}
-);
-
-#endif
diff --git a/Engine/Source/Runtime/CoreUObject/Public/UObject/Interface.h b/Engine/Source/Runtime/CoreUObject/Public/UObject/Interface.h
--- a/Engine/Source/Runtime/CoreUObject/Public/UObject/Interface.h
+++ b/Engine/Source/Runtime/CoreUObject/Public/UObject/Interface.h
@@ -7,23 +7,20 @@
 #include "UObject/ObjectMacros.h"
 #include "UObject/UObjectGlobals.h"
 
+#include "Interface.generated.h"
+
 /**
  * Base class for all interfaces
  *
  */
-
+UINTERFACE(MinimalAPI, meta=(IsBlueprintBase="true", CannotImplementInterfaceInBlueprint))
 class UInterface : public UObject
 {
-	DECLARE_CASTED_CLASS_INTRINSIC_WITH_API(UInterface, UObject, CLASS_Interface | CLASS_Abstract, TEXT("/Script/CoreUObject"), CASTCLASS_None, COREUOBJECT_API)
+	GENERATED_BODY()
+public:
 };
 
 class IInterface
 {
-protected:
-
-	virtual ~IInterface() {}
-
-public:
-
-	typedef UInterface UClassType;
+	GENERATED_BODY()
 };
diff --git a/Engine/Source/Runtime/CoreUObject/Public/UObject/NoExportTypes.h b/Engine/Source/Runtime/CoreUObject/Public/UObject/NoExportTypes.h
--- a/Engine/Source/Runtime/CoreUObject/Public/UObject/NoExportTypes.h
+++ b/Engine/Source/Runtime/CoreUObject/Public/UObject/NoExportTypes.h
@@ -3045,24 +3045,6 @@ public:
 	UClass(const FObjectInitializer& ObjectInitializer = FObjectInitializer::Get());
 };
 
-/**
- * @note The full C++ class is located here: Engine\Source\Runtime\CoreUObject\Public\UObject\Interface.h
- */
-
-UCLASS(abstract, noexport, intrinsic, interface, Config = Engine)
-class UInterface : public UObject
-{
-	GENERATED_BODY()
-public:
-	UInterface(const FObjectInitializer& ObjectInitializer = FObjectInitializer::Get());
-};
-
-class IInterface
-{
-	GENERATED_BODY()
-};
-
-
 /**
  * @note The full C++ class is located here: Engine\Source\Runtime\CoreUObject\Public\UObject\Package.h
  */
