[UObjectBase]

- Reserves the max size upfront for DepthMemo. This resolves ptr invalidation if the second FindOrAdd invalidates the reference returned from the first FindOrAdd.

#rnx

diff --git a/Engine/Source/Runtime/CoreUObject/Private/UObject/UObjectBase.cpp b/Engine/Source/Runtime/CoreUObject/Private/UObject/UObjectBase.cpp
--- a/Engine/Source/Runtime/CoreUObject/Private/UObject/UObjectBase.cpp
+++ b/Engine/Source/Runtime/CoreUObject/Private/UObject/UObjectBase.cpp
@@ -862,7 +862,11 @@ static void UObjectLoadAllCompiledInDefaultProperties(TArray<UClass*>& OutAllNew
 		// might rely on their parent having been initialized already, which might not be possible if there is a cycle in 
 		// the GetDefaultObject -> PostCDOConstruct loading graph.
 		{
+			// Memo to avoid re-calculating depth needlessly as we sort. It's important to note that since
+			// two FindOrAdds will occur back to back that we reserve our memory upfront to avoid re-allocing
+			// in-between FindOrAdd calls which will invalidate the first reference returned.
 			TMap<const UClass*, int32> DepthMemo;
+			DepthMemo.Reserve(OutAllNewClasses.Num());
 			auto SortByClassDepth = [&DepthMemo](TArray<UClass*>& Classes)
 				{
 					Algo::StableSort(Classes, [&DepthMemo](const UClass* A, const UClass* B)
@@ -891,6 +895,7 @@ static void UObjectLoadAllCompiledInDefaultProperties(TArray<UClass*>& OutAllNew
 			SortByClassDepth(NewClassesInCoreUObject);
 			SortByClassDepth(NewClassesInEngine);
 			SortByClassDepth(NewClasses);
+			checkf(DepthMemo.Num() <= OutAllNewClasses.Num(), TEXT("If we've added more than we reserved initially, we resized while iterating which may indicate a use-after-free problem."));
 		}
 
 		auto NotifyClassFinishedRegistrationEvents = [](TArray<UClass*>& Classes)
