Fixed a wrong deprecation version in FUObjectItem from 5.5 to current 5.6

diff --git a/Engine/Source/Runtime/CoreUObject/Public/UObject/UObjectArray.h b/Engine/Source/Runtime/CoreUObject/Public/UObject/UObjectArray.h
--- a/Engine/Source/Runtime/CoreUObject/Public/UObject/UObjectArray.h
+++ b/Engine/Source/Runtime/CoreUObject/Public/UObject/UObjectArray.h
@@ -65,7 +65,7 @@ public:
 	union
 	{
 		// Pointer to the allocated object
-		UE_DEPRECATED(5.5, "Use GetObject() and SetObject() to access Object.")
+		UE_DEPRECATED(5.6, "Use GetObject() and SetObject() to access Object.")
 		class UObjectBase* Object = nullptr;
 		uint32 ObjectPtrLow;	// this one is used as a dummy for natvis only an will be removed once packing will be enabled by default
 	};
