Add comments with full paths to header containing mirrored struct

diff --git a/Engine/Source/Runtime/CoreUObject/Public/UObject/NoExportTypes.h b/Engine/Source/Runtime/CoreUObject/Public/UObject/NoExportTypes.h
--- a/Engine/Source/Runtime/CoreUObject/Public/UObject/NoExportTypes.h
+++ b/Engine/Source/Runtime/CoreUObject/Public/UObject/NoExportTypes.h
@@ -118,7 +118,7 @@ namespace ELogTimes
 	};
 }
 
-/** Generic axis enum (mirrored for native use in Axis.h). */
+/** Generic axis enum (mirrored for native use in Engine/Source/Runtime/Core/Public/Axis.h). */
 UENUM(BlueprintType, meta=(ScriptName="AxisType"))
 namespace EAxis
 {
@@ -131,7 +131,7 @@ namespace EAxis
 	};
 }
 
-/** Generic axis list enum (mirrored for native use in Axis.h). */
+/** Generic axis list enum (mirrored for native use in Engine/Source/Runtime/Core/Public/Axis.h). */
 UENUM()
 namespace EAxisList
 {
@@ -166,7 +166,7 @@ namespace EAxisList
 	};
 }
 
-/** Describes shape of an interpolation curve (mirrored from InterpCurvePoint.h). */
+/** Describes shape of an interpolation curve (mirrored from Engine/Source/Runtime/Core/Public/Math/InterpCurvePoint.h). */
 UENUM()
 enum EInterpCurveMode : int
 {
@@ -193,6 +193,7 @@ enum EInterpCurveMode : int
 
 /**
  * Describes the format of a each pixel in a graphics buffer.
+ * Mirrored from Engine/Source/Runtime/Core/Public/PixelFormat.h
  * @warning: When you update this, you must add an entry to GPixelFormats(see RenderUtils.cpp)
  * @warning: When you update this, you must add an entries to PixelFormat.h, usually just copy the generated section on the header into EPixelFormat
  * @warning: The *Tools DLLs will also need to be recompiled if the ordering is changed, but should not need code changes.
@@ -317,7 +318,7 @@ enum EPixelFormat : int
 	PF_MAX,
 };
 
-/** Mouse cursor types (mirrored from ICursor.h) */
+/** Mouse cursor types (mirrored from Engine/Source/Runtime/ApplicationCore/Public/GenericPlatform/ICursor.h) */
 UENUM()
 namespace EMouseCursor
 {
@@ -370,7 +371,7 @@ namespace EMouseCursor
 	};
 }
 
-/** A set of numerical unit types supported by the engine. Mirrored from UnitConversion.h */
+/** A set of numerical unit types supported by the engine. Mirrored from Engine/Source/Runtime/Core/Public/Math/UnitConversion.h */
 UENUM(BlueprintType)
 enum class EUnit : uint8
 {
@@ -452,7 +453,7 @@ enum class EUnit : uint8
 
 /**
  * Enum controlling when to emit property change notifications when setting a property value.
- * @note Mirrored from PropertyAccessUtil.h
+ * @note Mirrored from Engine/Source/Runtime/CoreUObject/Public/UObject/PropertyAccessUtil.h
  */
 UENUM(BlueprintType)
 enum class EPropertyAccessChangeNotifyMode : uint8
@@ -467,7 +468,7 @@ enum class EPropertyAccessChangeNotifyMode : uint8
 
 /**
  * Enumerates supported message dialog category types.
- * @note Mirrored from GenericPlatformMisc.h
+ * @note Mirrored from Engine/Source/Runtime/Core/Public/GenericPlatform/GenericPlatformMisc.h
  */
 UENUM(BlueprintType)
 enum class EAppMsgCategory : uint8
@@ -480,7 +481,7 @@ enum class EAppMsgCategory : uint8
 
 /**
 * Enum denoting message dialog return types.
-* @note Mirrored from GenericPlatformMisc.h
+* @note Mirrored from Engine/Source/Runtime/Core/Public/GenericPlatform/GenericPlatformMisc.h
 */
 UENUM(BlueprintType)
 namespace EAppReturnType
@@ -500,7 +501,7 @@ namespace EAppReturnType
 
 /**
 * Enum denoting message dialog button choices. Used in combination with EAppReturnType.
-* @note Mirrored from GenericPlatformMisc.h
+* @note Mirrored from Engine/Source/Runtime/Core/Public/GenericPlatform/GenericPlatformMisc.h
 */
 UENUM(BlueprintType)
 namespace EAppMsgType
@@ -523,13 +524,14 @@ namespace EAppMsgType
 
 /**
  * A struct used as stub for deleted ones. 
+ * Mirrored from Engine/Source/Runtime/Core/Public/Misc/FallbackStruct.h
  */
 USTRUCT(noexport, IsAlwaysAccessible, HasDefaults)
 struct FFallbackStruct
 {
 };
 
-/** A globally unique identifier (mirrored from Guid.h) */
+/** A globally unique identifier (mirrored from Engine/Source/Runtime/Core/Public/Misc/Guid.h) */
 USTRUCT(immutable, noexport, BlueprintType, IsAlwaysAccessible, HasDefaults)
 struct FGuid
 {
@@ -703,7 +705,7 @@ struct FVector2D
 	FLargeWorldCoordinatesReal Y;
 };
 
-/** A pair of 3D vectors (mirrored from TwoVectors.h). */
+/** A pair of 3D vectors (mirrored from Engine/Source/Runtime/Core/Public/Math/TwoVectors.h). */
 USTRUCT(immutable, BlueprintType, noexport, IsAlwaysAccessible, HasDefaults)
 struct FTwoVectors
 {
@@ -1025,6 +1027,9 @@ struct FIntPoint
 	int32 Y;
 };
 
+/**
+ * @note The full C++ class is located here: Engine\Source\Runtime\Core\Public\Math\IntPoint.h
+ */
 USTRUCT(immutable, noexport, IsAlwaysAccessible, HasDefaults, HasNoOpConstructor, IsCoreType)
 struct FInt32Point
 {
@@ -1035,6 +1040,9 @@ struct FInt32Point
 	int32 Y;
 };
 
+/**
+ * @note The full C++ class is located here: Engine\Source\Runtime\Core\Public\Math\IntPoint.h
+ */
 USTRUCT(immutable, noexport, IsAlwaysAccessible, HasDefaults, HasNoOpConstructor, IsCoreType)
 struct FInt64Point
 {
@@ -1059,6 +1067,9 @@ struct FUintPoint
 	int32 Y;
 };
 
+/**
+ * @note The full C++ class is located here: Engine\Source\Runtime\Core\Public\Math\IntPoint.h
+ */
 USTRUCT(immutable, noexport, IsAlwaysAccessible, HasDefaults, HasNoOpConstructor, IsCoreType)
 struct FUint32Point
 {
@@ -1069,6 +1080,9 @@ struct FUint32Point
 	int32 Y;
 };
 
+/**
+ * @note The full C++ class is located here: Engine\Source\Runtime\Core\Public\Math\IntPoint.h
+ */
 USTRUCT(immutable, noexport, IsAlwaysAccessible, HasDefaults, HasNoOpConstructor, IsCoreType)
 struct FUint64Point
 {
@@ -1093,6 +1107,10 @@ struct FIntRect
 	FIntPoint Max;
 };
 
+/**
+ * An integer rectangle in 2D space.
+ * @note The full C++ class is located here: Engine\Source\Runtime\Core\Public\Math\IntRect.h
+ */
 USTRUCT(immutable, noexport, IsAlwaysAccessible, HasDefaults, HasNoOpConstructor, IsCoreType)
 struct FInt32Rect
 {
@@ -1103,6 +1121,10 @@ struct FInt32Rect
 	FInt32Point Max;	
 };
 
+/**
+ * An integer rectangle in 2D space.
+ * @note The full C++ class is located here: Engine\Source\Runtime\Core\Public\Math\IntRect.h
+ */
 USTRUCT(immutable, noexport, IsAlwaysAccessible, HasDefaults, HasNoOpConstructor, IsCoreType)
 struct FInt64Rect
 {
@@ -1127,6 +1149,10 @@ struct FUintRect
 	FUintPoint Max;
 };
 
+/**
+ * An integer rectangle in 2D space.
+ * @note The full C++ class is located here: Engine\Source\Runtime\Core\Public\Math\IntRect.h
+ */
 USTRUCT(immutable, noexport, IsAlwaysAccessible, HasDefaults, HasNoOpConstructor, IsCoreType)
 struct FUint32Rect
 {
@@ -1137,6 +1163,10 @@ struct FUint32Rect
 	FUint32Point Max;
 };
 
+/**
+ * An integer rectangle in 2D space.
+ * @note The full C++ class is located here: Engine\Source\Runtime\Core\Public\Math\IntRect.h
+ */
 USTRUCT(immutable, noexport, IsAlwaysAccessible, HasDefaults, HasNoOpConstructor, IsCoreType)
 struct FUint64Rect
 {
@@ -1764,7 +1794,7 @@ struct FMatrix
 
 /** 
  * Describes one specific point on an interpolation curve.
- * @note This is a mirror of TInterpCurvePoint<float>, defined in InterpCurvePoint.h
+ * @note This is a mirror of TInterpCurvePoint<float>, defined in Engine/Source/Runtime/Core/Public/Math/InterpCurvePoint.h
  */
 USTRUCT(noexport, BlueprintType, IsAlwaysAccessible, HasDefaults, HasNoOpConstructor)
 struct FInterpCurvePointFloat
@@ -1793,7 +1823,7 @@ struct FInterpCurvePointFloat
 
 /** 
  * Describes an entire curve that is used to compute a float output value from a float input.
- * @note This is a mirror of TInterpCurve<float>, defined in InterpCurve.h
+ * @note This is a mirror of TInterpCurve<float>, defined in Engine/Source/Runtime/Core/Public/Math/InterpCurve.h
  */
 USTRUCT(noexport, BlueprintType)
 struct FInterpCurveFloat
@@ -1813,7 +1843,7 @@ struct FInterpCurveFloat
 
 /** 
  * Describes one specific point on an interpolation curve.
- * @note This is a mirror of TInterpCurvePoint<FVector2D>, defined in InterpCurvePoint.h
+ * @note This is a mirror of TInterpCurvePoint<FVector2D>, defined in Engine/Source/Runtime/Core/Public/Math/InterpCurvePoint.h
  */
 USTRUCT(noexport, BlueprintType, IsAlwaysAccessible, HasDefaults, HasNoOpConstructor)
 struct FInterpCurvePointVector2D
@@ -1841,7 +1871,7 @@ struct FInterpCurvePointVector2D
 
 /** 
  * Describes an entire curve that is used to compute a 2D vector output value from a float input.
- * @note This is a mirror of TInterpCurve<FVector2D>, defined in InterpCurve.h
+ * @note This is a mirror of TInterpCurve<FVector2D>, defined in Engine/Source/Runtime/Core/Public/Math/InterpCurve.h
  */
 USTRUCT(noexport, BlueprintType)
 struct FInterpCurveVector2D
@@ -1861,7 +1891,7 @@ struct FInterpCurveVector2D
 
 /** 
  * Describes one specific point on an interpolation curve.
- * @note This is a mirror of TInterpCurvePoint<FVector>, defined in InterpCurvePoint.h
+ * @note This is a mirror of TInterpCurvePoint<FVector>, defined in Engine/Source/Runtime/Core/Public/Math/InterpCurvePoint.h
  */
 USTRUCT(noexport, BlueprintType, IsAlwaysAccessible, HasDefaults, HasNoOpConstructor)
 struct FInterpCurvePointVector
@@ -1889,7 +1919,7 @@ struct FInterpCurvePointVector
 
 /** 
  * Describes an entire curve that is used to compute a 3D vector output value from a float input.
- * @note This is a mirror of TInterpCurve<FVector>, defined in InterpCurve.h
+ * @note This is a mirror of TInterpCurve<FVector>, defined in Engine/Source/Runtime/Core/Public/Math/InterpCurve.h
  */
 USTRUCT(noexport, BlueprintType)
 struct FInterpCurveVector
@@ -1909,7 +1939,7 @@ struct FInterpCurveVector
 
 /**
  * Describes one specific point on an interpolation curve.
- * @note This is a mirror of TInterpCurvePoint<FQuat>, defined in InterpCurvePoint.h
+ * @note This is a mirror of TInterpCurvePoint<FQuat>, defined in Engine/Source/Runtime/Core/Public/Math/InterpCurvePoint.h
  */
 USTRUCT(noexport, BlueprintType, IsAlwaysAccessible, HasDefaults, HasNoOpConstructor)
 struct FInterpCurvePointQuat
@@ -1937,7 +1967,7 @@ struct FInterpCurvePointQuat
 
 /**
  * Describes an entire curve that is used to compute a quaternion output value from a float input.
- * @note This is a mirror of TInterpCurve<FQuat>, defined in InterpCurve.h
+ * @note This is a mirror of TInterpCurve<FQuat>, defined in Engine/Source/Runtime/Core/Public/Math/InterpCurve.h
  */
 USTRUCT(noexport, BlueprintType)
 struct FInterpCurveQuat
@@ -1957,7 +1987,7 @@ struct FInterpCurveQuat
 
 /**
  * Describes one specific point on an interpolation curve.
- * @note This is a mirror of TInterpCurvePoint<FTwoVectors>, defined in InterpCurvePoint.h
+ * @note This is a mirror of TInterpCurvePoint<FTwoVectors>, defined in Engine/Source/Runtime/Core/Public/Math/InterpCurvePoint.h
  */
 USTRUCT(noexport, BlueprintType, IsAlwaysAccessible, HasDefaults, HasNoOpConstructor)
 struct FInterpCurvePointTwoVectors
@@ -1985,7 +2015,7 @@ struct FInterpCurvePointTwoVectors
 
 /**
  * Describes an entire curve that is used to compute two 3D vector values from a float input.
- * @note This is a mirror of TInterpCurve<FTwoVectors>, defined in InterpCurve.h
+ * @note This is a mirror of TInterpCurve<FTwoVectors>, defined in Engine/Source/Runtime/Core/Public/Math/InterpCurve.h
  */
 USTRUCT(noexport, BlueprintType)
 struct FInterpCurveTwoVectors
@@ -2005,7 +2035,7 @@ struct FInterpCurveTwoVectors
 
 /**
  * Describes one specific point on an interpolation curve.
- * @note This is a mirror of TInterpCurvePoint<FLinearColor>, defined in InterpCurvePoint.h
+ * @note This is a mirror of TInterpCurvePoint<FLinearColor>, defined in Engine/Source/Runtime/Core/Public/Math/InterpCurvePoint.h
  */
 USTRUCT(noexport, BlueprintType, IsAlwaysAccessible, HasDefaults, HasNoOpConstructor)
 struct FInterpCurvePointLinearColor
@@ -2033,7 +2063,7 @@ struct FInterpCurvePointLinearColor
 
 /**
  * Describes an entire curve that is used to compute a color output value from a float input.
- * @note This is a mirror of TInterpCurve<FLinearColor>, defined in InterpCurve.h
+ * @note This is a mirror of TInterpCurve<FLinearColor>, defined in Engine/Source/Runtime/Core/Public/Math/InterpCurve.h
  */
 USTRUCT(noexport, BlueprintType)
 struct FInterpCurveLinearColor
@@ -2053,7 +2083,7 @@ struct FInterpCurveLinearColor
 
 /**
  * Transform composed of Quat/Translation/Scale.
- * @note This is implemented in either TransformVectorized.h or TransformNonVectorized.h depending on the platform.
+ * @note This is implemented in either Engine/Source/Runtime/Core/Public/Math/TransformVectorized.h or TransformNonVectorized.h depending on the platform.
  */
 USTRUCT(immutable, noexport, BlueprintType, BlueprintInternalUseOnly, IsAlwaysAccessible, HasDefaults, IsCoreType, meta = (ScriptDefaultMake, ScriptDefaultBreak, HasNativeMake = "/Script/Engine.KismetMathLibrary.MakeTransform", HasNativeBreak = "/Script/Engine.KismetMathLibrary.BreakTransform"))
 struct FTransform3f
@@ -2073,7 +2103,7 @@ struct FTransform3f
 
 /**
  * Transform composed of Quat/Translation/Scale.
- * @note This is implemented in either TransformVectorized.h or TransformNonVectorized.h depending on the platform.
+ * @note This is implemented in either Engine/Source/Runtime/Core/Public/Math/TransformVectorized.h or TransformNonVectorized.h depending on the platform.
  */
 USTRUCT(immutable, noexport, IsAlwaysAccessible, HasDefaults, IsCoreType)
 struct FTransform3d
@@ -2093,7 +2123,7 @@ struct FTransform3d
 
 /**
  * Transform composed of Quat/Translation/Scale.
- * @note This is implemented in either TransformVectorized.h or TransformNonVectorized.h depending on the platform.
+ * @note This is implemented in either Engine/Source/Runtime/Core/Public/Math/TransformVectorized.h or TransformNonVectorized.h depending on the platform.
  */
 USTRUCT(noexport, BlueprintType, IsAlwaysAccessible, HasDefaults, IsCoreType, meta=(HasNativeMake="/Script/Engine.KismetMathLibrary.MakeTransform", HasNativeBreak="/Script/Engine.KismetMathLibrary.BreakTransform"))
 struct FTransform
@@ -2329,7 +2359,7 @@ struct FPrimaryAssetId
 	FName PrimaryAssetName;
 };
 
-/** Enumerates the valid types of range bounds (mirrored from RangeBound.h) */
+/** Enumerates the valid types of range bounds (mirrored from Engine/Source/Runtime/Core/Public/Math/RangeBound.h) */
 UENUM(BlueprintType)
 namespace ERangeBoundTypes
 {
@@ -2354,7 +2384,7 @@ namespace ERangeBoundTypes
 
 /**
  * Defines a single bound for a range of values.
- * @note This is a mirror of TRangeBound<float>, defined in RangeBound.h
+ * @note This is a mirror of TRangeBound<float>, defined in Engine/Source/Runtime/Core/Public/Math/RangeBound.h
  * @note Fields are private to match the C++ declaration in the header above.
  */
 USTRUCT(noexport, BlueprintType)
@@ -2372,7 +2402,7 @@ private:
 
 /**
  * A contiguous set of floats described by lower and upper bound values.
- * @note This is a mirror of TRange<float>, defined in Range.h
+ * @note This is a mirror of TRange<float>, defined in Engine/Source/Runtime/Core/Public/Math/Range.h
  * @note Fields are private to match the C++ declaration in the header above.
  */
 USTRUCT(noexport, BlueprintType)
@@ -2391,7 +2421,7 @@ private:
 
 /**
  * Defines a single bound for a range of values.
- * @note This is a mirror of TRangeBound<double>, defined in RangeBound.h
+ * @note This is a mirror of TRangeBound<double>, defined in Engine/Source/Runtime/Core/Public/Math/RangeBound.h
  * @note Fields are private to match the C++ declaration in the header above.
  */
 USTRUCT(noexport, BlueprintType)
@@ -2409,7 +2439,7 @@ private:
 
 /**
  * A contiguous set of doubles described by lower and upper bound values.
- * @note This is a mirror of TRange<double>, defined in Range.h
+ * @note This is a mirror of TRange<double>, defined in Engine/Source/Runtime/Core/Public/Math/Range.h
  * @note Fields are private to match the C++ declaration in the header above.
  */
 USTRUCT(noexport, BlueprintType)
@@ -2427,7 +2457,7 @@ private:
 
 /**
  * Defines a single bound for a range of values.
- * @note This is a mirror of TRangeBound<int32>, defined in RangeBound.h
+ * @note This is a mirror of TRangeBound<int32>, defined in Engine/Source/Runtime/Core/Public/Math/RangeBound.h
  * @note Fields are private to match the C++ declaration in the header above.
  */
 USTRUCT(noexport, BlueprintType)
@@ -2445,7 +2475,7 @@ private:
 
 /**
  * A contiguous set of floats described by lower and upper bound values.
- * @note This is a mirror of TRange<int32>, defined in Range.h
+ * @note This is a mirror of TRange<int32>, defined in Engine/Source/Runtime/Core/Public/Math/Range.h
  * @note Fields are private to match the C++ declaration in the header above.
  */
 USTRUCT(noexport, BlueprintType)
@@ -2463,7 +2493,7 @@ private:
 
 /**
  * Defines a single bound for a range of frame numbers.
- * @note This is a mirror of TRangeBound<FFrameNumber>, defined in RangeBound.h
+ * @note This is a mirror of TRangeBound<FFrameNumber>, defined in Engine/Source/Runtime/Core/Public/Math/RangeBound.h
  * @note Fields are private to match the C++ declaration in the header above.
  */
 USTRUCT(noexport, BlueprintType)
@@ -2479,7 +2509,7 @@ private:
 
 /**
  * A contiguous set of frame numbers described by lower and upper bound values.
- * @note This is a mirror of TRange<FFrameNumber>, defined in Range.h
+ * @note This is a mirror of TRange<FFrameNumber>, defined in Engine/Source/Runtime/Core/Public/Math/Range.h
  * @note Fields are private to match the C++ declaration in the header above.
  */
 USTRUCT(noexport, BlueprintType)
@@ -2495,7 +2525,7 @@ private:
 
 /**
  * An interval of floats, defined by inclusive min and max values
- * @note This is a mirror of TInterval<float>, defined in Interval.h
+ * @note This is a mirror of TInterval<float>, defined in Engine/Source/Runtime/Core/Public/Math/Interval.h
  */
 USTRUCT(noexport, BlueprintType)
 struct FFloatInterval
@@ -2511,7 +2541,7 @@ struct FFloatInterval
 
 /**
  * An interval of integers, defined by inclusive min and max values
- * @note This is a mirror of TInterval<int32>, defined in Interval.h
+ * @note This is a mirror of TInterval<int32>, defined in Engine/Source/Runtime/Core/Public/Math/Interval.h
  */
 USTRUCT(noexport, BlueprintType)
 struct FInt32Interval
@@ -2525,7 +2555,7 @@ struct FInt32Interval
 	int32 Max;
 };
 
-/** Categories of localized text (mirrored in LocalizedTextSourceTypes.h */
+/** Categories of localized text (mirrored in Engine/Source/Runtime/Core/Public/Internationalization/LocalizedTextSourceTypes.h */
 UENUM(BlueprintType)
 enum class ELocalizedTextSourceCategory : uint8
 {
@@ -2536,7 +2566,7 @@ enum class ELocalizedTextSourceCategory : uint8
 
 /**
  * Polyglot data that may be registered to the text localization manager at runtime.
- * @note This struct is mirrored in PolyglotTextData.h
+ * @note This struct is mirrored in Engine/Source/Runtime/Core/Public/Internationalization/Internationalization.h
  */
 USTRUCT(noexport, BlueprintType)
 struct FPolyglotTextData
@@ -2597,7 +2627,7 @@ struct FPolyglotTextData
 	FText CachedText;
 };
 
-/** Report level of automation events (mirrored in AutomationEvent.h). */
+/** Report level of automation events (mirrored in Engine/Source/Runtime/Core/Public/Misc/AutomationEvent.h). */
 UENUM()
 enum class EAutomationEventType : uint8
 {
@@ -2606,7 +2636,7 @@ enum class EAutomationEventType : uint8
 	Error
 };
 
-/** Event emitted by automation system (mirrored in AutomationEvent.h). */
+/** Event emitted by automation system (mirrored in Engine/Source/Runtime/Core/Public/Misc/AutomationEvent.h). */
 USTRUCT(noexport)
 struct FAutomationEvent
 {
@@ -2623,7 +2653,7 @@ struct FAutomationEvent
 	FGuid Artifact;
 };
 
-/** Information about the execution of an automation task (mirrored in AutomationEvent.h). */
+/** Information about the execution of an automation task (mirrored in Engine/Source/Runtime/Core/Public/Misc/AutomationEvent.h). */
 USTRUCT(noexport)
 struct FAutomationExecutionEntry
 {
@@ -2646,7 +2676,7 @@ struct FAutomationExecutionEntry
  *
  * Has a globally unique identifier.
  * 
- * Opaque struct for the FInputDeviceId struct defined in CoreMiscDefines.h
+ * Opaque struct for the FInputDeviceId struct defined in Engine/Source/Runtime/Core/Public/Misc/CoreMiscDefines.h
  */
 USTRUCT(noexport, BlueprintType)
 struct FInputDeviceId
@@ -2663,7 +2693,7 @@ private:
  * This used to be just a typedef int32 that was used interchangeably as ControllerId and LocalUserIndex.
  * Moving forward these will be allocated by the platform application layer.
  *
- * Opaque struct for the FPlatformUserId struct defined in CoreMiscDefines.h
+ * Opaque struct for the FPlatformUserId struct defined in Engine/Source/Runtime/Core/Public/Misc/CoreMiscDefines.h
  */
 USTRUCT(noexport, BlueprintType)
 struct FPlatformUserId
@@ -2677,6 +2707,7 @@ private:
 
 /**
  * Represents the connection status of a given FInputDeviceId
+ * @note Mirrored from Engine/Source/Runtime/Core/Public/Misc/CoreMiscDefines.h
  */
 UENUM(BlueprintType)
 enum class EInputDeviceConnectionState : uint8
@@ -2701,6 +2732,8 @@ enum class EInputDeviceConnectionState : uint8
  *
  * This policy is only in affect on platforms which do NOT have a managed user login
  * or some kind of dedicated input mapping OS functionality (consoles). 
+ *
+ * @note Mirrored from Engine/Source/Runtime/ApplicationCore/Public/GenericPlatform/InputDeviceMappingPolicy.h
  */
 UENUM()
 enum class EInputDeviceMappingPolicy : int32
@@ -2758,7 +2791,7 @@ enum class EInputDeviceMappingPolicy : int32
 /**
  * Represents input device triggers that are available
  *
- * NOTE: Make sure to keep this type in sync with the reflected version in IInputInterface.h!
+ * NOTE: Make sure to keep this type in sync with the reflected version in Engine/Source/Runtime/ApplicationCore/Public/GenericPlatform/IInputInterface.h!
  */
 UENUM(BlueprintType)
 enum class EInputDeviceTriggerMask : uint8
@@ -2772,7 +2805,7 @@ enum class EInputDeviceTriggerMask : uint8
 /**
  * Represents input device analog sticks that are available
  *
- * NOTE: Make sure to keep this type in sync with the reflected version in IInputInterface.h!
+ * NOTE: Make sure to keep this type in sync with the reflected version in Engine/Source/Runtime/ApplicationCore/Public/GenericPlatform/IInputInterface.h!
  */
 UENUM(BlueprintType)
 enum class EInputDeviceAnalogStickMask : uint8
@@ -2784,6 +2817,7 @@ enum class EInputDeviceAnalogStickMask : uint8
 
 /**
  * Data about an input device's current state
+ * @note Mirrored from Engine/Source/Runtime/ApplicationCore/Public/GenericPlatform/InputDeviceMappingPolicy.h
  */
 USTRUCT(noexport, BlueprintType)
 struct FPlatformInputDeviceState
@@ -2799,7 +2833,10 @@ struct FPlatformInputDeviceState
 	EInputDeviceConnectionState ConnectionState = EInputDeviceConnectionState::Invalid;
 };
 
-/** Enum used by DataValidation plugin to see if an asset has been validated for correctness (mirrored in UObjectGlobals.h)*/
+/** 
+ * Enum used by DataValidation plugin to see if an asset has been validated for correctness 
+ * @note Mirrored from Engine/Source/Runtime/CoreUObject/Public/UObject/InputDeviceMappingPolicy.h
+ */
 UENUM(BlueprintType)
 enum class EDataValidationResult : uint8
 {
@@ -2812,7 +2849,7 @@ enum class EDataValidationResult : uint8
 };
 
 /**
- * A struct to serve as a filter for Asset Registry queries. (mirrored in ARFilter.h)
+ * A struct to serve as a filter for Asset Registry queries. (mirrored in Engine/Source/Runtime/CoreUObject/Public/AssetRegistry/ARFilter.h)
  */
 USTRUCT(noexport, BlueprintType, IsAlwaysAccessible, HasDefaults, meta = (HasNativeMake = "/Script/Engine.KismetSystemLibrary.MakeARFilter", HasNativeBreak = "/Script/Engine.KismetSystemLibrary.BreakARFilter"))
 struct FARFilter
@@ -2873,6 +2910,9 @@ struct FARFilter
 	uint32 WithPackageFlags = 0;
 };
 
+/**
+ * @note Mirrored from Engine/Source/Runtime/CoreUObject/Public/AssetRegistry/AssetBundleData.h
+ */
 USTRUCT(noexport)
 struct FAssetBundleEntry
 {
@@ -2890,7 +2930,10 @@ struct FAssetBundleEntry
 	TArray<FTopLevelAssetPath> AssetPaths;
 };
 
-/** A struct with a list of asset bundle entries. If one of these is inside a UObject it will get automatically exported as the asset registry tag AssetBundleData */
+/** 
+ * A struct with a list of asset bundle entries. If one of these is inside a UObject it will get automatically exported as the asset registry tag AssetBundleData 
+ * @note Mirrored from Engine/Source/Runtime/CoreUObject/Public/AssetRegistry/AssetBundleData.h
+ */
 USTRUCT(noexport, IsAlwaysAccessible, HasDefaults)
 struct FAssetBundleData
 {
@@ -2902,6 +2945,7 @@ struct FAssetBundleData
 /**
  * A struct to hold important information about an assets found by the Asset Registry
  * This struct is transient and should never be serialized
+ * @note Mirrored from Engine/Source/Runtime/CoreUObject/Public/AssetRegistry/AssetBundleData.h
  */
 USTRUCT(noexport, BlueprintType, IsAlwaysAccessible, HasDefaults)
 struct FAssetData
@@ -2941,6 +2985,9 @@ struct FAssetData
 	TArray<int32, TInlineAllocator<2>> ChunkIDs;
 };
 
+/** 
+ * @note Mirrored from Engine/Source/Runtime/CoreUObject/Public/UObject/Class.h
+ */
 USTRUCT(noexport, IsAlwaysAccessible, HasDefaults)
 struct FTestUninitializedScriptStructMembersTest
 {
@@ -2954,6 +3001,9 @@ struct FTestUninitializedScriptStructMembersTest
 	float UnusedValue;
 };
 
+/** 
+ * @note Mirrored from Engine/Source/Runtime/CoreUObject/Public/Serialization/TestUndeclaredScriptStructObjectReferences.h
+ */
 USTRUCT(noexport, IsAlwaysAccessible, HasDefaults)
 struct FTestUndeclaredScriptStructObjectReferencesTest
 {
@@ -3335,11 +3385,10 @@ public:
 	UMapProperty(const FObjectInitializer& ObjectInitializer = FObjectInitializer::Get());
 };
 
-/**
+/** 
+ * Structure for file paths that are displayed in the editor with a picker UI. 
  * @note The full C++ class is located here: Engine\Source\Runtime\CoreUObject\Public\UObject\SoftObjectPath.h
  */
-
-/** Structure for file paths that are displayed in the editor with a picker UI. */
 USTRUCT(noexport, BlueprintType)
 struct FFilePath
 {
@@ -3352,7 +3401,10 @@ struct FFilePath
 	FString FilePath;
 };
 
-/** Structure for directory paths that are displayed in the editor with a picker UI. */
+/** 
+ * Structure for directory paths that are displayed in the editor with a picker UI. 
+ * @note The full C++ class is located here: Engine\Source\Runtime\CoreUObject\Public\UObject\SoftObjectPath.h
+ */
 USTRUCT(noexport, BlueprintType)
 struct FDirectoryPath
 {
@@ -3365,7 +3417,10 @@ struct FDirectoryPath
 	FString Path;
 };
 
-/** Structure for templated strings that are displayed in the editor with a allowed args. */
+/** 
+ * Structure for templated strings that are displayed in the editor with a allowed args. 
+ * @note The full C++ class is located here: Engine\Source\Runtime\CoreUObject\Public\UObject\TemplateString.h
+ */
 USTRUCT(noexport, BlueprintType)
 struct FTemplateString
 {
@@ -3384,6 +3439,9 @@ struct FTemplateString
 	FText Resolved;
 };
 
+/** 
+ * @note The full C++ class is located here: Engine\Source\Runtime\CoreUObject\Public\UObject\RemoteObjectTypes.h
+ */
 USTRUCT(noexport, BlueprintType)
 struct FRemoteObjectId
 {
@@ -3393,6 +3451,9 @@ struct FRemoteObjectId
 	uint64 Id;
 };
 
+/** 
+ * @note The full C++ class is located here: Engine\Source\Runtime\CoreUObject\Public\UObject\RemoteObjectPathName.h
+ */
 USTRUCT(noexport, BlueprintType)
 struct FPackedRemoteObjectPathName
 {
@@ -3405,6 +3466,9 @@ struct FPackedRemoteObjectPathName
 	TArray<uint16> Names;
 };
 
+/** 
+ * @note The full C++ class is located here: Engine\Source\Runtime\CoreUObject\Public\UObject\RemoteObjectTransfer.h
+ */
 USTRUCT(noexport, BlueprintType)
 struct FRemoteObjectBytes
 {
@@ -3412,6 +3476,9 @@ struct FRemoteObjectBytes
 	TArray<uint8> Bytes;
 };
 
+/** 
+ * @note The full C++ class is located here: Engine\Source\Runtime\CoreUObject\Public\UObject\RemoteObjectPathName.h
+ */
 USTRUCT(noexport, BlueprintType)
 struct FRemoteObjectTables
 {
@@ -3422,12 +3489,18 @@ struct FRemoteObjectTables
 	TArray<FRemoteObjectId> RemoteIds;
 };
 
+/** 
+ * @note The full C++ class is located here: Engine\Source\Runtime\CoreUObject\Public\UObject\RemoteObjectPathName.h
+ */
 USTRUCT(noexport)
 struct FRemoteObjectPathName : public FRemoteObjectTables
 {
 	GENERATED_BODY()
 };
 
+/** 
+ * @note The full C++ class is located here: Engine\Source\Runtime\CoreUObject\Public\UObject\RemoteObjectTransfer.h
+ */
 USTRUCT(noexport, BlueprintType)
 struct FRemoteObjectData
 {
@@ -3443,6 +3516,9 @@ struct FRemoteObjectData
 	TArray<FRemoteObjectBytes> Bytes;
 };
 
+/** 
+ * @note The full C++ class is located here: Engine\Source\Runtime\CoreUObject\Public\UObject\RemoteObjectTypes.h
+ */
 USTRUCT(noexport, BlueprintType)
 struct FRemoteServerId
 {
@@ -3452,6 +3528,9 @@ struct FRemoteServerId
 	uint32 Id;
 };
 
+/** 
+ * @note The full C++ class is located here: Engine\Source\Runtime\CoreUObject\Public\UObject\RemoteExecutor.h
+ */
 USTRUCT(noexport, BlueprintType)
 struct FRemoteWorkPriority
 {
@@ -3461,6 +3540,9 @@ struct FRemoteWorkPriority
 	uint64 PackedData;
 };
 
+/** 
+ * @note The full C++ class is located here: Engine\Source\Runtime\CoreUObject\Public\UObject\RemoteObjectTransfer.h
+ */
 USTRUCT(noexport, BlueprintType)
 struct FRemoteObjectReference
 {
@@ -3473,6 +3555,9 @@ struct FRemoteObjectReference
 	FRemoteServerId ServerId;
 };
 
+/** 
+ * @note The full C++ class is located here: Engine\Source\Runtime\CoreUObject\Public\UObject\RemoteExecutor.h
+ */
 USTRUCT(noexport, BlueprintType)
 struct FRemoteTransactionId
 {
