[BugFix] Fixing several bugs that are blocking impersonation from being turned on
- Fixed: as a bandaid fix, StaticDuplicateObject was duplicating instance data objects as well. this is no longer needed and had a naieve implementation so removing it fixes issues with overrides
- Feature: Idos are back to outering to one another (todo: this likely regresses GC bugs)
- Fixed: when creating IDOs, the ArMergeOverrides flag was causing their OverriddenPropertySet to be missing a root
- Feature: Enabled impersonation cvar by default
- Feature: IDO classes now have a proper inherritance. Before they always inherrited from UObject  but now if the template class has a super, it will set the ido classes super to another ido class generated by the template super with no loose properties
- Fixed: when serializing with impersonation, the IDO's CDO was always used as the diff object. This breaks subobjects with archetypes like the one's being used for overridable serialization. To solve this we're now finding or constructing IDOs for the diff object and using that
- Fixed: subobject nodes in the override tree were being keyed by both the instance subobject and the ido. This was causing key misses when checking if a subobject was overridden
- Fixed: OverriddenPropertySet needed to redirect subobject pointers to IDOs so that override events sent for subobjects could properly apply

diff --git a/Engine/Source/Runtime/CoreUObject/Internal/UObject/InstanceDataObjectUtils.h b/Engine/Source/Runtime/CoreUObject/Internal/UObject/InstanceDataObjectUtils.h
--- a/Engine/Source/Runtime/CoreUObject/Internal/UObject/InstanceDataObjectUtils.h
+++ b/Engine/Source/Runtime/CoreUObject/Internal/UObject/InstanceDataObjectUtils.h
@@ -37,6 +37,8 @@ UE_API bool IsInstanceDataObjectImpersonationEnabledOnSave();
 UE_API bool IsInstanceDataObjectArchetypeChainEnabled();
 /** Query if InstanceDataObjects should impose placeholder objects for missing types. */
 UE_API bool IsInstanceDataObjectPlaceholderObjectSupportEnabled();
+/** Query if InstanceDataObjects of subobjects should be outered to the parent object's IDO */
+UE_API bool IsInstanceDataObjectOuterChainEnabled();
 
 UE_API bool StructContainsLooseProperties(const UStruct* Struct);
 UE_API bool StructIsInstanceDataObjectStruct(const UStruct* Struct);
@@ -49,6 +51,9 @@ bool CanCreatePropertyBagPlaceholderTypeForImportClass(const UClass* ImportType)
 /** Helper to check if a class is an IDO class. */
 UE_API bool IsClassOfInstanceDataObjectClass(const UStruct* Class);
 
+/** When Archetype chains are enabled, find or construct an object of a matching InstanceDataObjectClass with the data in InstanceArchetype */
+UE_API UObject* GetTemplateForInstanceDataObject(const UObject* Instance, const UObject* InstanceArchetype, const UClass* InstanceDataObjectClass);
+
 /** Generate a UClass that contains the union of the properties of PropertyTree and OwnerClass. */
 UClass* CreateInstanceDataObjectClass(const FPropertyPathNameTree* PropertyTree, const FUnknownEnumNames* EnumNames, UClass* OwnerClass, UObject* Outer);
 
diff --git a/Engine/Source/Runtime/CoreUObject/Internal/UObject/PropertyBagRepository.h b/Engine/Source/Runtime/CoreUObject/Internal/UObject/PropertyBagRepository.h
--- a/Engine/Source/Runtime/CoreUObject/Internal/UObject/PropertyBagRepository.h
+++ b/Engine/Source/Runtime/CoreUObject/Internal/UObject/PropertyBagRepository.h
@@ -83,7 +83,6 @@ public:
 	 * @return				The InstanceDataObject object, UClass derived from associated unknown property tree.
 	 */
 	UObject* CreateInstanceDataObject(UObject* Owner, FArchive& Archive, int64 StartOffset, int64 EndOffset, bool bIsArchetype = false);
-	UObject* DuplicateInstanceDataObject(UObject* SourceOwner, UObject* DestOwner);
 
 	// called at the end of postload to copy data from Owner to its IDO
 	COREUOBJECT_API void PostLoadInstanceDataObject(const UObject* Owner);
diff --git a/Engine/Source/Runtime/CoreUObject/Private/UObject/Class.cpp b/Engine/Source/Runtime/CoreUObject/Private/UObject/Class.cpp
--- a/Engine/Source/Runtime/CoreUObject/Private/UObject/Class.cpp
+++ b/Engine/Source/Runtime/CoreUObject/Private/UObject/Class.cpp
@@ -1598,19 +1598,7 @@ void UStruct::SerializeVersionedTaggedProperties(FStructuredArchive::FSlot Slot,
 				// Overridden values are saved independently in transaction, so do no need to restore them here.
 				if (!UnderlyingArchive.IsTransacting())
 				{
-					ControlContext.OverriddenProperties = nullptr;
-					if (UnderlyingArchive.ArMergeOverrides)
-					{
-						ControlContext.OverriddenProperties = FOverridableManager::Get().GetOverriddenProperties((UObject*)Data);
-					}
-					if (!ControlContext.OverriddenProperties)
-					{
-						ControlContext.OverriddenProperties = FOverridableManager::Get().RestoreOverrideOperation((UObject*)Data, Operation, /*bNeedsSubobjectTemplateInstantiation*/true);
-					}
-					else
-					{
-						ControlContext.OverriddenProperties->bNeedsSubobjectTemplateInstantiation = true;
-					}
+					ControlContext.OverriddenProperties = FOverridableManager::Get().RestoreOverrideOperation((UObject*)Data, Operation, /*bNeedsSubobjectTemplateInstantiation*/true, UnderlyingArchive.ArMergeOverrides);
 				}
 			}
 		}
diff --git a/Engine/Source/Runtime/CoreUObject/Private/UObject/InstanceDataObjectUtils.cpp b/Engine/Source/Runtime/CoreUObject/Private/UObject/InstanceDataObjectUtils.cpp
--- a/Engine/Source/Runtime/CoreUObject/Private/UObject/InstanceDataObjectUtils.cpp
+++ b/Engine/Source/Runtime/CoreUObject/Private/UObject/InstanceDataObjectUtils.cpp
@@ -357,6 +357,20 @@ namespace UE
 		ECVF_Default
 	);
 
+	bool bEnableFullIDOInherritance = true;
+	FAutoConsoleVariableRef EnableFullIDOInherritanceCVar(
+		TEXT("IDO.EnableFullInherritance"),
+		bEnableFullIDOInherritance,
+		TEXT("When enabled, IDO classes will inherrit from an IDO class of the template's super")
+	);
+
+	bool bEnableIDOOuterChain = true;
+	FAutoConsoleVariableRef EnableIDOOuterChainCVar(
+		TEXT("IDO.EnableOuterChain"),
+		bEnableIDOOuterChain,
+		TEXT("When enabled, IDOs of subobjects will be outered to the parent object's IDO")
+	);
+
 	FAutoConsoleCommand EnableIDOUnknownsCommand(
 		TEXT("IDO.Unknowns.Enable"),
 		TEXT("Use this command to toggle IDO.Unknowns.* on or off, or to report their current state."),
@@ -387,6 +401,8 @@ namespace UE
 		EnableIDOsForBlueprintInstancesCVar->Set(bEnabled);
 		EnableIDOArchetypeChainCVar->Set(bEnabled);
 		EnableIDOPlaceholderObjectsCVar->Set(bEnabled);
+		EnableFullIDOInherritanceCVar->Set(bEnabled);
+		EnableIDOOuterChainCVar->Set(bEnabled);
 	}
 
 	static bool AreAllIDOFeaturesEnabled()
@@ -401,7 +417,9 @@ namespace UE
 			bEnableIDOsForBlueprintArchetypes &&
 			bEnableIDOsForBlueprintInstances &&
 			bEnableIDOArchetypeChain &&
-			bEnableIDOPlaceholderObjects;
+			bEnableIDOPlaceholderObjects &&
+			bEnableFullIDOInherritance &&
+			bEnableIDOOuterChain;
 	}
 
 	FAutoConsoleCommand EnableAllIDOFeaturesCommand(
@@ -553,6 +571,11 @@ namespace UE
 		return bEnableIDOPlaceholderObjects;
 	}
 
+	bool IsInstanceDataObjectOuterChainEnabled()
+	{
+		return bEnableIDOOuterChain;
+	}
+
 	static const UObject* GetBlueprintGeneratedObject(const UObject* InObject)
 	{
 		static const FName NAME_BlueprintGeneratedClass(TEXT("BlueprintGeneratedClass"));
@@ -1032,8 +1055,21 @@ namespace UE
 	{
 		TSet<FPropertyPathName> SuperPropertyPathsFromTree;
 
-		// UClass is required to inherit from UObject
-		UStruct* Super = StructClass->IsChildOf<UClass>() ? UObject::StaticClass() : nullptr;
+		UStruct* Super = nullptr;
+		if (StructClass->IsChildOf<UClass>())
+		{
+			// classes should inherrit from an ido created out of the template super. This way they can be used as archetypes during serialization
+			UClass* TemplateSuper = Cast<UClass>(TemplateStruct->GetSuperStruct()); 
+			if (bEnableFullIDOInherritance && TemplateSuper && TemplateSuper != UObject::StaticClass())
+			{
+				Super = CreateInstanceDataObjectClass(nullptr, nullptr, TemplateSuper, Outer);
+			}
+			else
+			{
+				Super = UObject::StaticClass();
+			}
+		}
+
 
 		if (TemplateStruct)
 		{
@@ -1048,7 +1084,8 @@ namespace UE
 
 			// Gather properties for Super Struct
 			TArray<FProperty*> SuperProperties;
-			for (const FProperty* TemplateProperty : TFieldRange<FProperty>(TemplateStruct))
+			EFieldIterationFlags ExcludeSuperFlags = (EFieldIterationFlags)((uint8)EFieldIterationFlags::Default & ~(uint8)EFieldIterationFlags::IncludeSuper);
+			for (const FProperty* TemplateProperty : TFieldRange<FProperty>(TemplateStruct, ExcludeSuperFlags))
 			{
 				FProperty* SuperProperty = CastFieldChecked<FProperty>(FField::Duplicate(TemplateProperty, Super));
 				SuperProperties.Add(SuperProperty);
@@ -1259,6 +1296,38 @@ namespace UE
 		LoadTaggedProperties(SaveTaggedProperties(Source), Dest);
 	}
 
+	UObject* GetTemplateForInstanceDataObject(const UObject* Instance, const UObject* InstanceArchetype, const UClass* InstanceDataObjectClass)
+	{
+		if (IsInstanceDataObjectArchetypeChainEnabled() && Instance && InstanceArchetype && InstanceDataObjectClass)
+		{
+			if (InstanceArchetype == Instance->GetClass()->GetDefaultObject(false))
+			{
+				// if the archetype is a CDO, we can simply use the IDO's CDO because it's data should match.
+				return InstanceDataObjectClass->GetDefaultObject(false);
+			}
+			else
+			{
+				// attempt to find/create an IDO for the archetype. Unfortunately if there's unkown properties involved, the types may not match.
+				UObject* ArchetypeIdo = UE::CreateInstanceDataObject(const_cast<UObject*>(InstanceArchetype));
+				if (ArchetypeIdo && ArchetypeIdo->GetClass() && ArchetypeIdo->IsA(InstanceDataObjectClass))
+				{
+					return ArchetypeIdo;
+				}
+				else if (InstanceArchetype->GetArchetype())
+				{
+					// we need a archetype that is the same type as the IDO but has the same data as the owner's archetype
+					FStaticConstructObjectParameters TemplateParams(InstanceDataObjectClass);
+					TemplateParams.SetFlags |= EObjectFlags::RF_ArchetypeObject;
+					TemplateParams.Outer = GetTransientPackage();
+					UObject* Result = StaticConstructObject_Internal(TemplateParams);
+					CopyTaggedProperties(InstanceArchetype, Result);
+					return Result;
+				}
+			}
+		}
+		return nullptr;
+	}
+
 	static void SetClassFlags(UClass* IDOClass, const UClass* OwnerClass)
 	{
 		// always set
diff --git a/Engine/Source/Runtime/CoreUObject/Private/UObject/Obj.cpp b/Engine/Source/Runtime/CoreUObject/Private/UObject/Obj.cpp
--- a/Engine/Source/Runtime/CoreUObject/Private/UObject/Obj.cpp
+++ b/Engine/Source/Runtime/CoreUObject/Private/UObject/Obj.cpp
@@ -1939,7 +1939,7 @@ void UObject::SerializeScriptProperties( FStructuredArchive::FSlot Slot ) const
 
 		// When migrating remote objects the only instance where we serialize a CDO is when resetting an object to its archetype state in which case
 		// we want to serialize against the actual object class to get the right delta
-		UStruct* DiffClass = (HasAnyFlags(RF_ClassDefaultObject) && !UnderlyingArchive.HasAnyPortFlags(PPF_AvoidRemoteObjectMigration)) ? ObjClass->GetSuperClass() : ObjClass;
+		UClass* DiffClass = (HasAnyFlags(RF_ClassDefaultObject) && !UnderlyingArchive.HasAnyPortFlags(PPF_AvoidRemoteObjectMigration)) ? ObjClass->GetSuperClass() : ObjClass;
 
 		// Query if this object data is being impersonated 
 		const UObject* ThisObject = this;
@@ -1949,14 +1949,15 @@ void UObject::SerializeScriptProperties( FStructuredArchive::FSlot Slot ) const
 			ThisObject = Impersonator;
 			ObjClass = ThisObject->GetClass();
 
-			//@todo FH: Support for Default Object needs to be added when using impersonators
-			const UObject* DiffImpersonator = nullptr; //UE::Private::GetDataImpersonator(DiffObject);
-			DiffObject = DiffImpersonator ? const_cast<UObject*>(DiffImpersonator) : ObjClass->GetDefaultObject(false);
-			if (DiffObject)
+			// Find (or create) an impersonator for the archetype that matches the same class so it's serializable
+			DiffObject = UE::GetTemplateForInstanceDataObject(this, DiffObject, Impersonator->GetClass());
+			if (!DiffObject || !ThisObject->IsA(DiffObject->GetClass()))
 			{
-				DiffClass = DiffObject->GetClass();
+				DiffObject = ObjClass->GetDefaultObject();
 			}
-			ensureAlwaysMsgf(DiffClass == ObjClass, TEXT("Impersonation of '%s' using a different default class not appropriately supported at the moment. Class: '%s', DefaultClass: '%s'")
+			DiffClass = DiffObject ? DiffObject->GetClass() : nullptr;
+
+			ensureAlwaysMsgf(ThisObject->IsA(DiffClass), TEXT("Impersonation of '%s' using a different default class not appropriately supported at the moment. Class: '%s', DefaultClass: '%s'")
 				, *ThisObject->GetPathName(), *ObjClass->GetPathName(), *DiffClass->GetPathName());
 		}
 #endif
diff --git a/Engine/Source/Runtime/CoreUObject/Private/UObject/OverridableManager.cpp b/Engine/Source/Runtime/CoreUObject/Private/UObject/OverridableManager.cpp
--- a/Engine/Source/Runtime/CoreUObject/Private/UObject/OverridableManager.cpp
+++ b/Engine/Source/Runtime/CoreUObject/Private/UObject/OverridableManager.cpp
@@ -111,7 +111,7 @@ const FOverriddenPropertySet* FOverridableManager::GetOverriddenPropertiesIntern
 }
 
 
-FOverriddenPropertySet* FOverridableManager::RestoreOverrideOperation(TNotNull<UObject*> Object, EOverriddenPropertyOperation Operation, const bool bNeedsSubobjectTemplateInstantiation)
+FOverriddenPropertySet* FOverridableManager::RestoreOverrideOperation(TNotNull<UObject*> Object, EOverriddenPropertyOperation Operation, const bool bNeedsSubobjectTemplateInstantiation, const bool bDontClearOverrides)
 {
 	// When sub property overrides are going to be restored, it will turn this object back to modified, so let's not restore that operation.
 	// This allows cleanup of the modified state if we end up not overriding any sub properties
@@ -122,7 +122,10 @@ FOverriddenPropertySet* FOverridableManager::RestoreOverrideOperation(TNotNull<U
 
 #if WITH_EDITORONLY_DATA
 	FOverriddenPropertySet& ObjectOverriddenProperties = FindOrAdd(Object);
-	ObjectOverriddenProperties.Reset();
+	if (!bDontClearOverrides)
+	{
+		ObjectOverriddenProperties.Reset();
+	}
 	ObjectOverriddenProperties.SetOverriddenPropertyOperation(Operation, /*CurrentPropertyChain*/nullptr, /*Property*/nullptr);
 	ObjectOverriddenProperties.bNeedsSubobjectTemplateInstantiation = bNeedsSubobjectTemplateInstantiation;
 	return &ObjectOverriddenProperties;
diff --git a/Engine/Source/Runtime/CoreUObject/Private/UObject/OverriddenPropertySet.cpp b/Engine/Source/Runtime/CoreUObject/Private/UObject/OverriddenPropertySet.cpp
--- a/Engine/Source/Runtime/CoreUObject/Private/UObject/OverriddenPropertySet.cpp
+++ b/Engine/Source/Runtime/CoreUObject/Private/UObject/OverriddenPropertySet.cpp
@@ -8,6 +8,7 @@
 #include "UObject/PropertyOptional.h"
 #include "Misc/ScopeExit.h"
 #include "UObject/InstanceDataObjectUtils.h"
+#include "UObject/PropertyBagRepository.h"
 #include "UObject/UObjectArchetypeHelper.h"
 #include "UObject/UObjectThreadContext.h"
 
@@ -348,9 +349,18 @@ FOverriddenPropertyNodeID::FOverriddenPropertyNodeID(const FOverriddenPropertyNo
 FOverriddenPropertyNodeID::FOverriddenPropertyNodeID(TNotNull<const UObject*> InObject)
 : Object(NotNullGet(InObject))
 {
+	const UObject* KeyObject = InObject;
+#if WITH_EDITORONLY_DATA
+	// subobject pointers inside IDOs never point to another IDO so to keep that consistent, we'll always redirect IDO pointers here to the instance counterpart
+	if (const UObject* Instance = UE::FPropertyBagRepository::Get().FindInstanceForDataObject(InObject))
+	{
+		KeyObject = Instance;
+	}
+#endif
+
 	// Note: Using ObjectIndex by itself is not sufficient for an enduring unique identifier
 	// as re-instantiation can cause a reuse of the index for another object. Appending the serial solves this issue
-	const int32 ObjectIndex = GUObjectArray.ObjectToIndex(InObject);
+	const int32 ObjectIndex = GUObjectArray.ObjectToIndex(KeyObject);
 	Path = FName(FString::Printf(TEXT("ID_%d_%d"), ObjectIndex, GUObjectArray.AllocateSerialNumber(ObjectIndex)));
 }
 
@@ -520,7 +530,7 @@ EOverriddenPropertyOperation FOverriddenPropertySet::GetOverriddenPropertyOperat
 			if (NextPropertyIterator)
 			{
 				// Forward any sub queries to the subobject
-				if (UObject* SubObject = ObjectProperty->GetObjectPropertyValue(SubValuePtr))
+				if (UObject* SubObject = TryGetInstancedSubObjectValue(ObjectProperty, SubValuePtr))
 				{
 					// This should not be needed in the property grid, as it should already been called on the subobject.
 					return OverridableManager.GetOverriddenPropertyOperation(SubObject, NextPropertyIterator, bOutInheritedOperation);
@@ -669,7 +679,7 @@ bool FOverriddenPropertySet::ClearOverriddenProperty(FOverriddenPropertyNode& Pa
 		FPropertyVisitorPath::Iterator NextPropertyIterator = PropertyIterator+1;
 		if (const FObjectPropertyBase* ObjectProperty = CastField<FObjectPropertyBase>(CurrentProperty))
 		{
-			if (UObject* SubObject = ObjectProperty->GetObjectPropertyValue(SubValuePtr))
+			if (UObject* SubObject = TryGetInstancedSubObjectValue(ObjectProperty, SubValuePtr))
 			{
 				if (NextPropertyIterator)
 				{
@@ -846,7 +856,7 @@ void FOverriddenPropertySet::NotifyPropertyChange(FOverriddenPropertyNode* Paren
 {
 	checkf(IsValid(Owner), TEXT("Expecting a valid overridable owner"));
 
-	if (ChangeType == EPropertyChangeType::ResetToDefault)
+	if (ChangeType & EPropertyChangeType::ResetToDefault)
 	{
 		if (ParentPropertyNode && Notification == EPropertyNotificationType::PostEdit)
 		{
@@ -979,7 +989,7 @@ void FOverriddenPropertySet::NotifyPropertyChange(FOverriddenPropertyNode* Paren
 					PreEditArrayHelper.EmptyAndAddValues(ArrayHelper.Num());
 					for (int32 i = 0; i < ArrayHelper.Num(); i++)
 					{
-						InnerObjectProperty->SetObjectPropertyValue(PreEditArrayHelper.GetElementPtr(i), InnerObjectProperty->GetObjectPropertyValue(ArrayHelper.GetElementPtr(i)));
+						InnerObjectProperty->SetObjectPropertyValue(PreEditArrayHelper.GetElementPtr(i), TryGetInstancedSubObjectValue(InnerObjectProperty, ArrayHelper.GetElementPtr(i)));
 					}
 					return;
 				}
@@ -1015,7 +1025,7 @@ void FOverriddenPropertySet::NotifyPropertyChange(FOverriddenPropertyNode* Paren
 				auto ArrayAddImpl = [&]()
 				{
 					checkf(ArrayHelper.IsValidIndex(ArrayIndex), TEXT("ArrayAdd change type expected to have an valid index"));
-					if (UObject* AddedSubObject = InnerObjectProperty->GetObjectPropertyValue(ArrayHelper.GetElementPtr(ArrayIndex)))
+					if (UObject* AddedSubObject = TryGetInstancedSubObjectValue(InnerObjectProperty, ArrayHelper.GetElementPtr(ArrayIndex)))
 					{
 						if(SubPropertyNode)
 						{
@@ -1035,7 +1045,7 @@ void FOverriddenPropertySet::NotifyPropertyChange(FOverriddenPropertyNode* Paren
 				auto ArrayRemoveImpl = [&]()
 				{
 					checkf(PreEditArrayHelper.IsValidIndex(ArrayIndex), TEXT("ArrayRemove change type expected to have an valid index"));
-					if (UObject* RemovedSubObject = InnerObjectProperty->GetObjectPropertyValue(PreEditArrayHelper.GetElementPtr(ArrayIndex)))
+					if (UObject* RemovedSubObject = TryGetInstancedSubObjectValue(InnerObjectProperty, PreEditArrayHelper.GetElementPtr(ArrayIndex)))
 					{
 						if(SubPropertyNode)
 						{
@@ -1164,7 +1174,7 @@ void FOverriddenPropertySet::NotifyPropertyChange(FOverriddenPropertyNode* Paren
 			// Can only forward to subobject if we have a valid index
 			else if (ArrayHelper.IsValidIndex(ArrayIndex))
 			{
-				if (UObject* SubObject = InnerObjectProperty->GetObjectPropertyValue(ArrayHelper.GetElementPtr(ArrayIndex)))
+				if (UObject* SubObject = TryGetInstancedSubObjectValue(InnerObjectProperty, ArrayHelper.GetElementPtr(ArrayIndex)))
 				{
 					// This should not be needed in the property grid, as it should already been called on the subobject itself.
 					OverridableManager.NotifyPropertyChange(Notification, SubObject, NextPropertyIterator, ChangeType);
@@ -1409,7 +1419,7 @@ void FOverriddenPropertySet::NotifyPropertyChange(FOverriddenPropertyNode* Paren
 				SubPropertyNode->Operation = EOverriddenPropertyOperation::Replace;
 			}
 		}
-		else if (UObject* SubObject = ObjectProperty->GetObjectPropertyValue(SubValuePtr))
+		else if (UObject* SubObject = TryGetInstancedSubObjectValue(ObjectProperty, SubValuePtr))
 		{
 			// This should not be needed in the property grid, as it should already been called on the subobject.
 			OverridableManager.NotifyPropertyChange(Notification, SubObject, NextPropertyIterator, ChangeType);
@@ -1456,7 +1466,7 @@ void FOverriddenPropertySet::RemoveOverriddenSubProperties(FOverriddenPropertyNo
 	PropertyNode.SubPropertyNodeKeys.Empty();
 }
 
-UObject* FOverriddenPropertySet::TryGetInstancedSubObjectValue(const FObjectPropertyBase* FromProperty, void* ValuePtr) const
+UObject* FOverriddenPropertySet::TryGetInstancedSubObjectValue(const FObjectPropertyBase* FromProperty, const void* ValuePtr) const
 {
 	// Property can be NULL - in that case there is no value.
 	if (!FromProperty)
@@ -1551,6 +1561,11 @@ const FOverriddenPropertyNode* FOverriddenPropertySet::GetOverriddenPropertyNode
 	return nullptr;
 }
 
+const FOverriddenPropertyNode* FOverriddenPropertySet::GetOverriddenPropertyNode(FOverriddenPropertyNodeID NodeID) const
+{
+	return OverriddenPropertyNodes.Find(NodeID);
+}
+
 EOverriddenPropertyOperation FOverriddenPropertySet::GetOverriddenPropertyOperation(const FOverriddenPropertyNode* ParentPropertyNode, const FArchiveSerializedPropertyChain* CurrentPropertyChain, FProperty* Property) const
 {
 	// No need to look further
@@ -1602,7 +1617,7 @@ FOverriddenPropertyNode* FOverriddenPropertySet::SetOverriddenPropertyOperation(
 	// No need to look further
 	// if it is the entire property is replaced or
 	// if it is the FOverriddenPropertySet struct which is always Overridden
-	if (ParentPropertyNode.Operation == EOverriddenPropertyOperation::Replace)
+	if (ParentPropertyNode.Operation == EOverriddenPropertyOperation::Replace && Property)
 	{
 		return nullptr;
 	}
diff --git a/Engine/Source/Runtime/CoreUObject/Private/UObject/PropertyBagRepository.cpp b/Engine/Source/Runtime/CoreUObject/Private/UObject/PropertyBagRepository.cpp
--- a/Engine/Source/Runtime/CoreUObject/Private/UObject/PropertyBagRepository.cpp
+++ b/Engine/Source/Runtime/CoreUObject/Private/UObject/PropertyBagRepository.cpp
@@ -639,33 +639,6 @@ void FPropertyBagRepository::PostEditChangeChainProperty(const UObject* Object,
 #endif
 }
 
-UObject* FPropertyBagRepository::DuplicateInstanceDataObject(UObject* SourceOwner, UObject* DestOwner)
-{
-	FPropertyBagRepositoryLock LockRepo(this);
-	if (FPropertyBagAssociationData* SourceData = AssociatedData.Find(SourceOwner))
-	{
-		check(SourceData->InstanceDataObject);
-		FPropertyBagAssociationData& DestData = AssociatedData.FindOrAdd(DestOwner);
-		ensure(!DestData.InstanceDataObject);
-
-		UObject* Outer = GetTransientPackage();
-
-		// construct InstanceDataObject
-		FStaticConstructObjectParameters Params(SourceData->InstanceDataObject->GetClass());
-		Params.SetFlags |= EObjectFlags::RF_Transactional;
-		Params.Name = MakeUniqueObjectName(Outer, Params.Class, DestOwner->GetFName(), EUniqueObjectNameOptions::UniversallyUnique);
-		Params.Outer = Outer;
-		DestData.InstanceDataObject = StaticConstructObject_Internal(Params);
-		InstanceDataObjectToOwner.Add(DestData.InstanceDataObject, DestOwner);
-
-		CopyTaggedProperties(SourceData->InstanceDataObject, DestData.InstanceDataObject);
-
-		DestData.bNeedsFixup = SourceData->bNeedsFixup;
-		return DestData.InstanceDataObject;
-	}
-	return nullptr;
-}
-
 void FPropertyBagRepository::PostLoadInstanceDataObject(const UObject* Owner)
 {
 	// fixups may have been applied to the instance during PostLoad and they need to be copied to its IDO
@@ -824,46 +797,35 @@ UObject* FPropertyBagRepository::CreateInstanceDataObject(UObject* Owner, FArchi
 		return BagData->InstanceDataObject;
 	}
 
-	// TODO: Should the IDO or its class be in another package?
-	UObject* Outer = GetTransientPackage();
+	// use the IDO of the outer as the new outer if there is one, otherwise use the transient package
+	UObject* Outer = nullptr;
+	if (IsInstanceDataObjectOuterChainEnabled())
+	{
+		if (const FPropertyBagAssociationData* OuterData = AssociatedData.Find(Owner->GetOuter()))
+		{
+			Outer = OuterData->InstanceDataObject;
+		}
+	}
+	if (!Outer)
+	{
+		Outer = GetTransientPackage();
+	}
 
 	// Construct the class for the IDO.
 	TSharedPtr<FPropertyPathNameTree> PropertyTree = FUnknownPropertyTree(Owner).Find();
 	const FUnknownEnumNames EnumNames(Owner);
 	const FUnknownEnumNames* EnumNamesPtr = !EnumNames.IsEmpty() ? &EnumNames : nullptr;
-	const UClass* InstanceDataObjectClass = CreateInstanceDataObjectClass(PropertyTree.Get(), EnumNamesPtr, Owner->GetClass(), Outer);
+	const UClass* InstanceDataObjectClass = CreateInstanceDataObjectClass(PropertyTree.Get(), EnumNamesPtr, Owner->GetClass(), GetTransientPackage());
 
 	// Find the template for the IDO.
-	UObject* Template = nullptr;
-	if (IsInstanceDataObjectArchetypeChainEnabled())
+
+	//@todoio GetArchetype is pathological for blueprint classes and the event driven loader; the EDL already knows what the archetype is; just calling this->GetArchetype() tries to load some other stuff.
+	UObject* OwnerArchetype = Archive.GetArchetypeFromLoader(Owner);
+	if (!OwnerArchetype)
 	{
-		if (UObject* OwnerArchetype = Owner->GetArchetype())
-		{
-			if (OwnerArchetype == Owner->GetClass()->GetDefaultObject(false))
-			{
-				// if the archetype is a CDO, we can simply use the IDO's CDO because it's data should match.
-				Template = InstanceDataObjectClass->GetDefaultObject(false);
-			}
-			else
-			{
-				// attempt to find/create an IDO for the archetype. Unfortunately if there's unkown properties involved, the types may not match.
-				UObject* ArchetypeIdo = UE::CreateInstanceDataObject(OwnerArchetype);
-				if (ArchetypeIdo && ArchetypeIdo->GetClass() && ArchetypeIdo->IsA(InstanceDataObjectClass))
-				{
-					Template = ArchetypeIdo;
-				}
-				else if (OwnerArchetype->GetArchetype())
-				{
-					// we need a archetype that is the same type as the IDO but has the same data as the owner's archetype
-					FStaticConstructObjectParameters TemplateParams(InstanceDataObjectClass);
-					TemplateParams.SetFlags |= EObjectFlags::RF_ArchetypeObject;
-					TemplateParams.Outer = GetTransientPackage();
-					Template = StaticConstructObject_Internal(TemplateParams);
-					CopyTaggedProperties(OwnerArchetype, Template);
-				}
-			}
-		}
+		OwnerArchetype = Owner->GetArchetype();
 	}
+	UObject* Template = GetTemplateForInstanceDataObject(Owner, OwnerArchetype, InstanceDataObjectClass);
 
 	// Construct the IDO.
 	FStaticConstructObjectParameters Params(InstanceDataObjectClass);
diff --git a/Engine/Source/Runtime/CoreUObject/Private/UObject/UObjectGlobals.cpp b/Engine/Source/Runtime/CoreUObject/Private/UObject/UObjectGlobals.cpp
--- a/Engine/Source/Runtime/CoreUObject/Private/UObject/UObjectGlobals.cpp
+++ b/Engine/Source/Runtime/CoreUObject/Private/UObject/UObjectGlobals.cpp
@@ -35,7 +35,6 @@
 #include "UObject/Class.h"
 #include "UObject/CoreRedirects.h"
 #include "UObject/FastReferenceCollector.h"
-#include "UObject/InstanceDataObjectUtils.h"
 #include "UObject/OverridableManager.h"
 #include "UObject/PropertyStateTracking.h"
 #include "UObject/UObjectIterator.h"
@@ -3331,15 +3330,6 @@ UObject* StaticDuplicateObjectEx( FObjectDuplicationParameters& Parameters )
 			}
 		}
 	}
-
-#if WITH_EDITORONLY_DATA
-	// if Source has an IDO, make one for Dest and copy it
-	if (UE::IsInstanceDataObjectSupportEnabled(Parameters.SourceObject) && UE::IsInstanceDataObjectImpersonationEnabledOnSave())
-	{
-		UE::FPropertyBagRepository& PropertyBagRepository = UE::FPropertyBagRepository::Get();
-		PropertyBagRepository.DuplicateInstanceDataObject(Parameters.SourceObject, DupRootObject);
-	}
-#endif
 	
 	return DupRootObject;
 }
diff --git a/Engine/Source/Runtime/CoreUObject/Public/UObject/OverridableManager.h b/Engine/Source/Runtime/CoreUObject/Public/UObject/OverridableManager.h
--- a/Engine/Source/Runtime/CoreUObject/Public/UObject/OverridableManager.h
+++ b/Engine/Source/Runtime/CoreUObject/Public/UObject/OverridableManager.h
@@ -139,8 +139,9 @@ public:
 	 * @param Object to set the override operation on
 	 * @param Operation the override operation to set on the object
 	 * @param bNeedsSubobjectTemplateInstantiation set to true if it will need a sub object template instantiation during postload, false will not do it.
+	 * @param bDontClearOverrides set to true if you don't want to stomp any overrides that currently exist
 	 * @return the overridden properties of the object */
-	COREUOBJECT_API FOverriddenPropertySet* RestoreOverrideOperation(TNotNull<UObject*> Object, EOverriddenPropertyOperation Operation, const bool bNeedsSubobjectTemplateInstantiation);
+	COREUOBJECT_API FOverriddenPropertySet* RestoreOverrideOperation(TNotNull<UObject*> Object, EOverriddenPropertyOperation Operation, const bool bNeedsSubobjectTemplateInstantiation, const bool bDontClearOverrides = false);
 
 	/**
 	 * Restore some of the overridden state that is not necessarily restored by the CPFUO
diff --git a/Engine/Source/Runtime/CoreUObject/Public/UObject/OverriddenPropertySet.h b/Engine/Source/Runtime/CoreUObject/Public/UObject/OverriddenPropertySet.h
--- a/Engine/Source/Runtime/CoreUObject/Public/UObject/OverriddenPropertySet.h
+++ b/Engine/Source/Runtime/CoreUObject/Public/UObject/OverriddenPropertySet.h
@@ -432,6 +432,7 @@ public:
 	 * @param CurrentPropertyChain leading to the property being serialized if any, null or empty it will return the node of the object itself
 	 * @return the node containing the information of the overridden property */
 	UE_API const FOverriddenPropertyNode* GetOverriddenPropertyNode(const FArchiveSerializedPropertyChain* CurrentPropertyChain) const;
+	UE_API const FOverriddenPropertyNode* GetOverriddenPropertyNode(FOverriddenPropertyNodeID NodeID) const;
 
 	/**
 	 * Retrieve the overridable operation given the property key
@@ -499,7 +500,7 @@ protected:
 
 	UE_API void RemoveOverriddenSubProperties(FOverriddenPropertyNode& PropertyIterator);
 
-	UE_API UObject* TryGetInstancedSubObjectValue(const FObjectPropertyBase* FromProperty, void* ValuePtr) const;
+	UE_API UObject* TryGetInstancedSubObjectValue(const FObjectPropertyBase* FromProperty, const void* ValuePtr) const;
 
 private:
 	UPROPERTY()
