[CIS Fix] Used potentially null varialbe.

diff --git a/Engine/Source/Runtime/CoreUObject/Private/UObject/InstanceDataObjectUtils.cpp b/Engine/Source/Runtime/CoreUObject/Private/UObject/InstanceDataObjectUtils.cpp
--- a/Engine/Source/Runtime/CoreUObject/Private/UObject/InstanceDataObjectUtils.cpp
+++ b/Engine/Source/Runtime/CoreUObject/Private/UObject/InstanceDataObjectUtils.cpp
@@ -1059,7 +1059,7 @@ namespace UE
 		if (StructClass->IsChildOf<UClass>())
 		{
 			// classes should inherrit from an ido created out of the template super. This way they can be used as archetypes during serialization
-			UClass* TemplateSuper = Cast<UClass>(TemplateStruct->GetSuperStruct()); 
+			UClass* TemplateSuper = TemplateStruct ? Cast<UClass>(TemplateStruct->GetSuperStruct()) : nullptr;
 			if (bEnableFullIDOInherritance && TemplateSuper && TemplateSuper != UObject::StaticClass())
 			{
 				Super = CreateInstanceDataObjectClass(nullptr, nullptr, TemplateSuper, Outer);
