Make it so that `TProperty_WithEqualityAndSerializer::Identical` won't call the copy constructor.

diff --git a/Engine/Source/Runtime/CoreUObject/Public/UObject/UnrealType.h b/Engine/Source/Runtime/CoreUObject/Public/UObject/UnrealType.h
--- a/Engine/Source/Runtime/CoreUObject/Public/UObject/UnrealType.h
+++ b/Engine/Source/Runtime/CoreUObject/Public/UObject/UnrealType.h
@@ -1659,9 +1659,21 @@ public:
 	// FProperty interface.
 	virtual bool Identical( const void* A, const void* B, uint32 PortFlags=0 ) const override
 	{
-		// RHS is the same as TTypeFundamentals::GetOptionalPropertyValue(B) but avoids an unnecessary copy of B
-		return TTypeFundamentals::GetPropertyValue(A) == (B ? TTypeFundamentals::GetPropertyValue(B) : TTypeFundamentals::GetDefaultPropertyValue());
+		// `B` can be null, but we want to avoid the cost of copying a non-null `B` if we were to just use
+		// `TTypeFundamentals::GetOptionalPropertyValue(B)`. It is important that we branch on `B` here
+		// rather than using a ternary between the two variants of B (its own value or the default), because
+		// `TTypeFundamentals::GetDefaultPropertyValue()` returns a non-ref, and thus doing a ternary would
+		// also result in a copy of `B` which we are trying to avoid!
+		if (B)
+		{
+			return TTypeFundamentals::GetPropertyValue(A) == TTypeFundamentals::GetPropertyValue(B);
+		}
+		else
+		{
+			return TTypeFundamentals::GetPropertyValue(A) == TTypeFundamentals::GetDefaultPropertyValue();
+		}
 	}
+
 	virtual void SerializeItem(FStructuredArchive::FSlot Slot, void* Value, void const* Defaults) const override
 	{
 		Slot << *TTypeFundamentals::GetPropertyValuePtr(Value);
