ZenLoader: Improve error printing of missing script imports when in debug/test configs.
Always print full path names and hashes instead of nothing or none.
Note that the existing "Import index X is null" messages are only enabled at log level 'Log'.
Compiled out in shipping.
Can be disabled with s.EnableScriptImportsDebugging=false.

Server cost:
Adds debug data for ~46 000 script objects at a cost of ~1.6 MiB, plus ~38 000 FNames at a cost of ~1.1 MiB.
Adds ~50 ms startup time (~40 ms to load the debug data, ~10 ms to trim it for a lower memory footprint).

Client cost:
Adds debug data for ~25 000 script objects at a cost of ~840 KiB, plus ~19 000 FNames at a cost of ~640 KiB
Adds ~50 ms startup time (~40 ms to load the debug data, ~10 ms to trim it for a lower memory footprint).

#rnx
#tests client, server, editor

diff --git a/Engine/Source/Runtime/CoreUObject/Private/Serialization/AsyncLoading2.cpp b/Engine/Source/Runtime/CoreUObject/Private/Serialization/AsyncLoading2.cpp
--- a/Engine/Source/Runtime/CoreUObject/Private/Serialization/AsyncLoading2.cpp
+++ b/Engine/Source/Runtime/CoreUObject/Private/Serialization/AsyncLoading2.cpp
@@ -322,6 +322,20 @@ void FindAllRuntimeScriptPackages(FRuntimeScriptPackages& OutPackages)
 #define ALT2_ENABLE_PACKAGE_DEPENDENCY_DEBUGGING 0
 #endif
 
+#ifndef ALT2_ENABLE_SCRIPT_IMPORTS_DEBUGGING
+#define ALT2_ENABLE_SCRIPT_IMPORTS_DEBUGGING !UE_SHIPPING
+#endif
+#if ALT2_ENABLE_SCRIPT_IMPORTS_DEBUGGING
+static bool GEnableScriptImportsDebugging = true;
+static FAutoConsoleVariableRef CEnableScriptImportsDebugging(
+	TEXT("s.EnableScriptImportsDebugging"),
+	GEnableScriptImportsDebugging,
+	TEXT("Loads all native script objects (classes, enums, structs, CDOs) that were available in the editor during cook.")
+	TEXT("Provides more informative log messages at the cost of up to ~1.5 MiB memory and ~1 MiB FName memory depending on project size."),
+	ECVF_Default
+	);
+#endif
+
 #define checkObject(Object, expr) checkf(expr, TEXT("Object='%s' (%p), Flags=%s, InternalFlags=0x%08X") \
 									, (Object ? *Object->GetFullName() : TEXT("null")) \
 									, Object \
@@ -1571,6 +1585,24 @@ public:
 	}
 };
 
+#if ALT2_ENABLE_SCRIPT_IMPORTS_DEBUGGING
+class FCookedScriptObjectsDebug
+{
+private:
+	TMap<FPackageObjectIndex, TPair<FMinimalName, FPackageObjectIndex>> Objects;
+
+	void LoadDebugData();
+	void GetPathNameInternal(const TPair<FMinimalName, FPackageObjectIndex>& Entry, FStringBuilderBase& PathName, int32& NumParts) const;
+public:
+	FCookedScriptObjectsDebug();
+	~FCookedScriptObjectsDebug();
+	bool IsEnabled() const { return Objects.Num() > 0; }
+	void TrimDebugData(const TMap<FPackageObjectIndex, UObject*>& RuntimeScriptObjects);
+	void OnCompiledInUObjectsRegistered(FName, ECompiledInUObjectsRegisteredStatus Status);
+	FString GetPathName(FPackageObjectIndex Index) const;
+};
+#endif
+
 class FGlobalImportStore
 {
 	friend class FAsyncLoadingVerseRoot;
@@ -1593,6 +1625,11 @@ private:
 	mutable FRWLock ScriptObjectsLock;
 	// All currently loaded public export objects from any loaded package
 	TMap<int32, FPublicExportKey> ObjectIndexToPublicExport;
+#if ALT2_ENABLE_SCRIPT_IMPORTS_DEBUGGING
+public:
+	FCookedScriptObjectsDebug CookedScriptObjectsDebug;
+private:
+#endif
 
 	// Process all the deferred deletion for packages that are finished loading
 	void FlushDeferredDeletePackagesQueue();
@@ -2222,23 +2259,15 @@ struct FPackageImportStore
 	{
 	}
 
-	inline bool IsValidLocalImportIndex(const TArrayView<const FPackageObjectIndex>& ImportMap, FPackageIndex LocalIndex)
-	{
-		check(ImportMap.Num() > 0);
-		return LocalIndex.IsImport() && LocalIndex.ToImport() < ImportMap.Num();
-	}
-
-	inline UObject* FindOrGetImportObjectFromLocalIndex(const FAsyncPackageHeaderData& Header, FPackageIndex LocalIndex)
+	inline FPackageObjectIndex GetGlobalImportIndex(const FAsyncPackageHeaderData& Header, FPackageIndex LocalIndex)
 	{
 		check(LocalIndex.IsImport());
 		check(Header.ImportMap.Num() > 0);
 		const int32 LocalImportIndex = LocalIndex.ToImport();
-		check(LocalImportIndex < Header.ImportMap.Num());
-		const FPackageObjectIndex GlobalIndex = Header.ImportMap[LocalIndex.ToImport()];
-		return FindOrGetImportObject(Header, GlobalIndex);
+		return LocalImportIndex < Header.ImportMap.Num() ? Header.ImportMap[LocalImportIndex] : FPackageObjectIndex();
 	}
 
-	inline UObject* FindOrGetImportObject(const FAsyncPackageHeaderData& Header, FPackageObjectIndex GlobalIndex)
+	inline UObject* GetImportObject(const FAsyncPackageHeaderData& Header, FPackageObjectIndex GlobalIndex)
 	{
 		check(GlobalIndex.IsImport());
 		UObject* Object = nullptr;
@@ -2263,6 +2292,22 @@ struct FPackageImportStore
 		return Object;
 	}
 
+	FString GetImportName(const FAsyncPackageHeaderData& Header, FPackageObjectIndex GlobalIndex)
+	{
+		const UObject* Import = GetImportObject(Header, GlobalIndex);
+		if (Import)
+		{
+			return Import->GetFName().ToString();
+		}
+#if ALT2_ENABLE_SCRIPT_IMPORTS_DEBUGGING
+		else if (GlobalIndex.IsScriptImport())
+		{
+			return GlobalImportStore.CookedScriptObjectsDebug.GetPathName(GlobalIndex);
+		}
+#endif
+		return FString();
+	}
+
 	void GetUnresolvedCDOs(const FAsyncPackageHeaderData& Header, TArray<UClass*, TInlineAllocator<8>>& Classes)
 	{
 		for (const FPackageObjectIndex& Index : Header.ImportMap)
@@ -2577,13 +2622,14 @@ public:
 		}
 		else
 		{
-			if (ImportStore->IsValidLocalImportIndex(HeaderData->ImportMap, Index))
+			FPackageObjectIndex GlobalIndex = ImportStore->GetGlobalImportIndex(*HeaderData, Index);
+			if (!GlobalIndex.IsNull())
 			{
-				Object = ImportStore->FindOrGetImportObjectFromLocalIndex(*HeaderData, Index);
+				Object = ImportStore->GetImportObject(*HeaderData, GlobalIndex);
 
-				UE_ASYNC_PACKAGE_CLOG_VERBOSE(!Object, Log, *PackageDesc,
-					TEXT("FExportArchive: Object"), TEXT("Import index %d is null"),
-					Index.ToImport());
+				UE_ASYNC_PACKAGE_CLOG(!Object, Log, *PackageDesc,
+					TEXT("FExportArchive: Object"), TEXT("Import index %d is null (0x%llX - '%s')"),
+					Index.ToImport(), GlobalIndex.Value(), *ImportStore->GetImportName(*HeaderData, GlobalIndex));
 			}
 			else
 			{
@@ -3925,7 +3971,7 @@ private:
 
 	void ProcessExportDependencies(const FAsyncPackageHeaderData& Header, int32 LocalExportIndex, FExportBundleEntry::EExportCommandType CommandType);
 	int32 GetPublicExportIndex(uint64 ExportHash, FAsyncPackageHeaderData*& OutHeader);
-	FName GetFNameFromPackageObjectIndex(FPackageObjectIndex Index);
+	FString GetNameFromPackageObjectIndex(FPackageObjectIndex Index);
 	UObject* ConditionalCreateExport(const FAsyncPackageHeaderData& Header, int32 LocalExportIndex);
 	UObject* ConditionalSerializeExport(const FAsyncPackageHeaderData& Header, int32 LocalExportIndex);
 	UObject* ConditionalCreateImport(const FAsyncPackageHeaderData& Header, int32 LocalImportIndex);
@@ -6077,6 +6123,161 @@ void FGlobalImportStore::FindAllScriptObjects(bool bVerifyOnly)
 	FindAllScriptObjectsInPackages(ScriptPackages.VerseVNI, true);
 }
 
+#if ALT2_ENABLE_SCRIPT_IMPORTS_DEBUGGING
+
+FCookedScriptObjectsDebug::FCookedScriptObjectsDebug()
+{
+	if (FPlatformProperties::RequiresCookedData())
+	{
+		// register to sync load all the debug data just before we start loading the CDOs for the first module
+		// this is too early to check for GEnableScriptImportsDebugging if it has been set from the command line
+		FCoreUObjectDelegates::CompiledInUObjectsRegisteredDelegate.AddRaw(this, &FCookedScriptObjectsDebug::OnCompiledInUObjectsRegistered);
+	}
+}
+
+FCookedScriptObjectsDebug::~FCookedScriptObjectsDebug()
+{
+	FCoreUObjectDelegates::CompiledInUObjectsRegisteredDelegate.RemoveAll(this);
+}
+
+void FCookedScriptObjectsDebug::OnCompiledInUObjectsRegistered(FName, ECompiledInUObjectsRegisteredStatus Status)
+{
+	if (Status == ECompiledInUObjectsRegisteredStatus::PreCDO)
+	{
+		if (GEnableScriptImportsDebugging)
+		{
+			LoadDebugData();
+		}
+		FCoreUObjectDelegates::CompiledInUObjectsRegisteredDelegate.RemoveAll(this);
+	}
+}
+
+void FCookedScriptObjectsDebug::LoadDebugData()
+{
+	TRACE_CPUPROFILER_EVENT_SCOPE(FCookedScriptObjectsDebug::LoadDebugData);
+
+	FIoDispatcher& IoDispatcher = FIoDispatcher::Get();
+	FIoBatch IoBatch = IoDispatcher.NewBatch();
+	FIoRequest Request = IoBatch.Read(CreateIoChunkId(0, 0, EIoChunkType::ScriptObjects), FIoReadOptions(), IoDispatcherPriority_Max);
+	FEventRef Event;
+	IoBatch.IssueAndTriggerEvent(Event.Get());
+	Event->Wait();
+
+	if (const FIoBuffer* Buffer = Request.GetResult())
+	{
+		FLargeMemoryReader Ar(Buffer->Data(), Buffer->DataSize());
+		TArray<FDisplayNameEntryId> NameMap = LoadNameBatch(Ar);
+		int32 NumObjects = 0;
+		Ar << NumObjects;
+		Objects.Reserve(NumObjects);
+		TConstArrayView<FScriptObjectEntry> Entries =
+			MakeConstArrayView(reinterpret_cast<const FScriptObjectEntry*>(Buffer->Data() + Ar.Tell()), NumObjects);
+		for (const FScriptObjectEntry& Entry : Entries)
+		{
+			FName Name = NameMap[Entry.Mapped.GetIndex()].ToName(Entry.Mapped.GetNumber());
+			Objects.Emplace(Entry.GlobalIndex, {FMinimalName(Name), Entry.OuterIndex});
+		}
+		UE_LOG(LogStreaming, Log, TEXT("AsyncLoading2 - LoadDebugData: Loaded %d cooked script object entries (%.2f KB)"),
+			Objects.Num(), (float)Objects.GetAllocatedSize()/1024.f);
+	}
+}
+
+void FCookedScriptObjectsDebug::TrimDebugData(const TMap<FPackageObjectIndex, UObject*>& RuntimeScriptObjects)
+{
+	if (!IsEnabled())
+	{
+		return;
+	}
+	TRACE_CPUPROFILER_EVENT_SCOPE(FCookedScriptObjectsDebug::TrimDebugData);
+	// Remove debug data for all cooked script objects that are available in the runtime,
+	// just keep the missing ones and their outers.
+	const int32 Num = Objects.Num();
+	TBitArray<> Remove(false, Num);
+	TBitArray<> Visited(false, Num);
+	TFunction<void(int32)> KeepOuters = [this, &KeepOuters, &Remove, &Visited](int32 I)
+	{
+		if (!Visited[I])
+		{
+			Visited[I] = true;
+			Remove[I] = false;
+			const TPair<FMinimalName, FPackageObjectIndex>& Entry = Objects.Get(FSetElementId::FromInteger(I)).Value;
+			if (!Entry.Value.IsNull())
+			{
+				FSetElementId OuterId = Objects.FindId(Entry.Value);
+				KeepOuters(OuterId.AsInteger());
+			}
+		}
+	};
+	// 1. Mark all runtime objects to be removed
+	for (const TPair<FPackageObjectIndex, UObject*>& Runtime : RuntimeScriptObjects)
+	{
+		FSetElementId Id = Objects.FindId(Runtime.Key);
+		if (Id.IsValidId())
+		{
+			Remove[Id.AsInteger()] = true;
+		}
+	}
+	// 2. Keep the outers for all objects that are kept
+	for (int I = 0; I < Num; ++I)
+	{
+		if (!Remove[I])
+		{
+			KeepOuters(I);
+		}
+	}
+	// 3. Remove the remaining runtime objects
+	for (int I = 0; I < Num; ++I)
+	{
+		if (Remove[I])
+		{
+			Objects.Remove(FSetElementId::FromInteger(I));
+		}
+	}
+	Objects.Compact();
+	UE_LOG(LogStreaming, Log, TEXT("AsyncLoading2 - TrimDebugData: Kept %d cooked script object entries (%.2f KB)"),
+		Objects.Num(), (float)Objects.GetAllocatedSize()/1024.f);
+}
+
+void FCookedScriptObjectsDebug::GetPathNameInternal(const TPair<FMinimalName, FPackageObjectIndex>& Entry, FStringBuilderBase& PathName, int32& NumParts) const
+{
+	const TPair<FMinimalName, FPackageObjectIndex>* Outer = Objects.Find(Entry.Value);
+	if (Outer)
+	{
+		GetPathNameInternal(*Outer, PathName, NumParts);
+		// Use ? as divider for longer paths where we don't know if ':' or '.' should be used.
+		if (NumParts == 1)
+		{
+			PathName.AppendChar('.');
+		}
+		else if (NumParts == 2)
+		{
+			PathName.AppendChar(SUBOBJECT_DELIMITER_CHAR);
+		}
+		else if (NumParts > 2)
+		{
+			// we don't know if this should be a dot or a colon separator
+			PathName.AppendChar('?');
+		}
+	}
+	FName(Entry.Key).AppendString(PathName);
+	++NumParts;
+}
+
+FString FCookedScriptObjectsDebug::GetPathName(FPackageObjectIndex Index) const
+{
+	const TPair<FMinimalName,FPackageObjectIndex>* Entry = Objects.Find(Index);
+	if (Entry)
+	{
+		TStringBuilder<256> PathName;
+		int32 NumParts = 0;
+		GetPathNameInternal(*Entry, PathName, NumParts);
+		return PathName.ToString();
+	}
+	return FString();
+}
+
+#endif // ALT2_ENABLE_SCRIPT_IMPORTS_DEBUGGING 
+
 void FGlobalImportStore::RegistrationComplete()
 {
 #if WITH_EDITOR
@@ -6092,6 +6293,9 @@ void FGlobalImportStore::RegistrationComplete()
 #endif
 	FWriteScopeLock ScopeLock(ScriptObjectsLock);
 	ScriptObjects.Shrink();
+#if ALT2_ENABLE_SCRIPT_IMPORTS_DEBUGGING
+	CookedScriptObjectsDebug.TrimDebugData(ScriptObjects);
+#endif
 }
 
 void FAsyncPackage2::ImportPackagesRecursiveInner(FAsyncLoadingThreadState2& ThreadState, FIoBatch& IoBatch, FPackageStore& PackageStore, FAsyncPackageHeaderData& Header)
@@ -6950,7 +7154,7 @@ bool FAsyncPackage2::ResolveLinkerLoadImports(FAsyncLoadingThreadState2& ThreadS
 		const FPackageObjectIndex& GlobalImportIndex = HeaderData.ImportMap[ImportIndex];
 		if (!GlobalImportIndex.IsNull())
 		{
-			UObject* FromImportStore = ImportStore.FindOrGetImportObject(HeaderData, GlobalImportIndex);
+			UObject* FromImportStore = ImportStore.GetImportObject(HeaderData, GlobalImportIndex);
 #if ALT2_VERIFY_LINKERLOAD_MATCHES_IMPORTSTORE
 			/*if (Desc.UPackageId.ValueForDebugging() == 0xF37353A71BF5C938 && ImportIndex == 0x0000005d)
 			{
@@ -7665,10 +7869,10 @@ UObject* FAsyncPackage2::EventDrivenIndexToObject(const FAsyncPackageHeaderData&
 	}
 	else if (Index.IsImport())
 	{
-		Result = ImportStore.FindOrGetImportObject(Header, Index);
+		Result = ImportStore.GetImportObject(Header, Index);
 		UE_CLOG(!Result, LogStreaming, Warning, TEXT("Missing Dependency, missing %s import %s(0x%llX) for package %s"),
 			Index.IsScriptImport() ? TEXT("script") : TEXT("package"),
-			*GetFNameFromPackageObjectIndex(Index).ToString(),
+			*ImportStore.GetImportName(Header, Index),
 			Index.Value(),
 			*Desc.PackagePathToLoad.GetPackageFName().ToString());
 	}
@@ -7830,7 +8034,7 @@ UObject* FAsyncPackage2::ConditionalCreateImport(const FAsyncPackageHeaderData&
 {
 	const FPackageObjectIndex& ObjectIndex = Header.ImportMap[LocalImportIndex];
 	check(ObjectIndex.IsPackageImport());
-	if (UObject* FromImportStore = ImportStore.FindOrGetImportObject(Header, ObjectIndex))
+	if (UObject* FromImportStore = ImportStore.GetImportObject(Header, ObjectIndex))
 	{
 		return FromImportStore;
 	}
@@ -7856,7 +8060,7 @@ UObject* FAsyncPackage2::ConditionalSerializeImport(const FAsyncPackageHeaderDat
 	const FPackageObjectIndex& ObjectIndex = Header.ImportMap[LocalImportIndex];
 	check(ObjectIndex.IsPackageImport());
 
-	if (UObject* FromImportStore = ImportStore.FindOrGetImportObject(Header, ObjectIndex))
+	if (UObject* FromImportStore = ImportStore.GetImportObject(Header, ObjectIndex))
 	{
 		if (!FromImportStore->HasAllFlags(RF_NeedLoad))
 		{
@@ -7961,30 +8165,27 @@ Verse::VCell* FAsyncPackage2::ConditionalSerializeCellImport(const FAsyncPackage
 	return SourcePackage->ConditionalSerializeCellExport(*SourcePackageHeader, CellExportIndex);
 }
 
-FName FAsyncPackage2::GetFNameFromPackageObjectIndex(FPackageObjectIndex Index)
+FString FAsyncPackage2::GetNameFromPackageObjectIndex(FPackageObjectIndex Index)
 {
-	FName Out(NAME_None);
+	FString Out;
 	if (Index.IsExport())
 	{
 #if ALT2_ENABLE_LINKERLOAD_SUPPORT
 		if (LinkerLoadState.IsSet())
 		{
-			Out = LinkerLoadState->Linker->ExportMap[Index.ToExport()].ObjectName;
+			Out = LinkerLoadState->Linker->ExportMap[Index.ToExport()].ObjectName.ToString();
 		}
 		else
 #endif
 		{
-			HeaderData.NameMap.TryGetName(HeaderData.ExportMap[Index.ToExport()].ObjectName, Out);
+			FName Name(NAME_None);
+			HeaderData.NameMap.TryGetName(HeaderData.ExportMap[Index.ToExport()].ObjectName, Name);
+			Out = Name.ToString();
 		}
 	}
-	else
+	else if (Index.IsImport())
 	{
-		// No need to check if the Index is invalid, that is handled internally
-		const UObject* Import = ImportStore.FindOrGetImportObject(HeaderData, Index);
-		if (Import)
-		{
-			Out = Import->GetFName();
-		}
+		Out = ImportStore.GetImportName(HeaderData, Index);
 	}
 
 	return Out;
@@ -8029,13 +8230,13 @@ void FAsyncPackage2::EventDrivenCreateExport(const FAsyncPackageHeaderData& Head
 
 	if (!LoadClass)
 	{
-		UE_ASYNC_PACKAGE_LOG(Error, Desc, TEXT("CreateExport"), TEXT("Could not find class object (0x%llX - '%s') for %s"), Export.ClassIndex.Value(), *GetFNameFromPackageObjectIndex(Export.ClassIndex).ToString(), *ObjectName.ToString());
+		UE_ASYNC_PACKAGE_LOG(Error, Desc, TEXT("CreateExport"), TEXT("Could not find class object (0x%llX - '%s') for %s"), Export.ClassIndex.Value(), *GetNameFromPackageObjectIndex(Export.ClassIndex), *ObjectName.ToString());
 		ExportObject.bExportLoadFailed = true;
 		return;
 	}
 	if (!ThisParent)
 	{
-		UE_ASYNC_PACKAGE_LOG(Error, Desc, TEXT("CreateExport"), TEXT("Could not find outer (0x%llX - '%s') object for %s"), Export.OuterIndex.Value(), *GetFNameFromPackageObjectIndex(Export.OuterIndex).ToString(), *ObjectName.ToString());
+		UE_ASYNC_PACKAGE_LOG(Error, Desc, TEXT("CreateExport"), TEXT("Could not find outer (0x%llX - '%s') object for %s"), Export.OuterIndex.Value(), *GetNameFromPackageObjectIndex(Export.OuterIndex), *ObjectName.ToString());
 		ExportObject.bExportLoadFailed = true;
 		return;
 	}
@@ -8045,7 +8246,7 @@ void FAsyncPackage2::EventDrivenCreateExport(const FAsyncPackageHeaderData& Head
 		ExportObject.SuperObject = EventDrivenIndexToObject(Header, Export.SuperIndex, false);
 		if (!ExportObject.SuperObject)
 		{
-			UE_ASYNC_PACKAGE_LOG(Error, Desc, TEXT("CreateExport"), TEXT("Could not find SuperStruct object (0x%llX - '%s') for %s"), Export.SuperIndex.Value(), *GetFNameFromPackageObjectIndex(Export.SuperIndex).ToString(), *ObjectName.ToString());
+			UE_ASYNC_PACKAGE_LOG(Error, Desc, TEXT("CreateExport"), TEXT("Could not find SuperStruct object (0x%llX - '%s') for %s"), Export.SuperIndex.Value(), *GetNameFromPackageObjectIndex(Export.SuperIndex), *ObjectName.ToString());
 			ExportObject.bExportLoadFailed = true;
 			return;
 		}
@@ -8055,7 +8256,7 @@ void FAsyncPackage2::EventDrivenCreateExport(const FAsyncPackageHeaderData& Head
 	ExportObject.TemplateObject = EventDrivenIndexToObject(Header, Export.TemplateIndex, true);
 	if (!ExportObject.TemplateObject)
 	{
-		UE_ASYNC_PACKAGE_LOG(Error, Desc, TEXT("CreateExport"), TEXT("Could not find template object for %s"), *ObjectName.ToString());
+		UE_ASYNC_PACKAGE_LOG(Error, Desc, TEXT("CreateExport"), TEXT("Could not find template object (0x%llX - '%s') for %s"), Export.TemplateIndex.Value(), *GetNameFromPackageObjectIndex(Export.TemplateIndex), *ObjectName.ToString());
 		ExportObject.bExportLoadFailed = true;
 		return;
 	}
@@ -10617,7 +10818,7 @@ void FAsyncPackage2::ClearConstructedObjects()
 		if (Export.bWasFoundInMemory)
 		{
 			checkf(Export.Object, TEXT("Export '%s' in package '%s' found in memory is missing:\n\t\tSuper: '%s', Template: '%s', LoadFailed: %s, Filtered: %s (%s), FoundInMemory: %s"),
-				*GetFNameFromPackageObjectIndex(FPackageObjectIndex::FromExportIndex(ExportIndex)).ToString(),
+				*GetNameFromPackageObjectIndex(FPackageObjectIndex::FromExportIndex(ExportIndex)),
 				*Desc.PackagePathToLoad.GetPackageFName().ToString(),
 				(Export.SuperObject ? *Export.SuperObject->GetPathName() : TEXT("null")),
 				(Export.TemplateObject ? *Export.TemplateObject->GetPathName() : TEXT("null")),
