Natvis for basic remote object related structures

diff --git a/Engine/Source/Runtime/CoreUObject/Public/UObject/RemoteObjectTypes.h b/Engine/Source/Runtime/CoreUObject/Public/UObject/RemoteObjectTypes.h
--- a/Engine/Source/Runtime/CoreUObject/Public/UObject/RemoteObjectTypes.h
+++ b/Engine/Source/Runtime/CoreUObject/Public/UObject/RemoteObjectTypes.h
@@ -127,6 +127,11 @@ private:
 
 public:
 
+	enum
+	{
+		RemoteObjectSupportCompiledIn = UE_WITH_REMOTE_OBJECT_HANDLE
+	};
+
 	FRemoteObjectId() = default;
 
 	static FRemoteObjectId CreateFromInt(uint64 InRawId)
diff --git a/Engine/Source/Runtime/CoreUObject/Public/UObject/UObjectArray.h b/Engine/Source/Runtime/CoreUObject/Public/UObject/UObjectArray.h
--- a/Engine/Source/Runtime/CoreUObject/Public/UObject/UObjectArray.h
+++ b/Engine/Source/Runtime/CoreUObject/Public/UObject/UObjectArray.h
@@ -58,7 +58,14 @@ private:
 	// while preserving atomicity of the whole thing so RootFlags+RefCount can be evaluated in a lock-less way.
 	// If we want to add more Flags, we can reduce the size of RefCount to 24 bits to give us more bits for Flags but will require that we convert
 	// EInternalObjectFlags to a 64-bit.
-	int64 FlagsAndRefCount;
+	union
+	{
+		int64 FlagsAndRefCount;
+#if !UE_WITH_REMOTE_OBJECT_HANDLE
+		// Dummy variable for natvis
+		uint8 RemoteId;
+#endif
+	};
 
 #if !UE_ENABLE_FUOBJECT_ITEM_PACKING
 public:
