[UObject]
- Fix benign race on SerialNumber when resolving weakobjectptr on old expired objects
- Fix benign race on ObjAvailableList.Num() by placing the count into a TSAN friendly value count
- Make sure this won't prevent AutoRTFM from working when TSAN is inactive

diff --git a/Engine/Source/Runtime/CoreUObject/Private/UObject/UObjectArray.cpp b/Engine/Source/Runtime/CoreUObject/Private/UObject/UObjectArray.cpp
--- a/Engine/Source/Runtime/CoreUObject/Private/UObject/UObjectArray.cpp
+++ b/Engine/Source/Runtime/CoreUObject/Private/UObject/UObjectArray.cpp
@@ -277,6 +277,7 @@ void FUObjectArray::AllocateUObjectIndex(UObjectBase* Object, EInternalObjectFla
 			Index = ObjAvailableList.Pop();
 			const int32 AvailableCount = ObjAvailableList.Num();
 			checkSlow(AvailableCount >= 0);
+			ObjAvailableListEstimateCount = AvailableCount;
 		}
 		else
 		{
@@ -299,7 +300,17 @@ void FUObjectArray::AllocateUObjectIndex(UObjectBase* Object, EInternalObjectFla
 	}
 	ObjectItem->SetObject(Object);
 	ObjectItem->ClusterRootIndex = 0;
+
+	// AutoRTFM doesn't like atomics even when relaxed, so we need to differentiate the code here
+#if WITH_INSTRUMENTATION || USING_THREAD_SANITISER
+	// This can race with weakptr trying to resolve an old object in this slot.
+	// Avoid TSAN warning here since this is safe, the ObjectItem can't possibly match as its been
+	// cleaned up during GC.
+	FPlatformAtomics::AtomicStore_Relaxed(&ObjectItem->SerialNumber, SerialNumber);
+#else
 	ObjectItem->SerialNumber = SerialNumber;
+#endif
+
 #if UE_WITH_REMOTE_OBJECT_HANDLE
 	ObjectItem->RemoteId = RemoteId;
 #endif // UE_WITH_REMOTE_OBJECT_HANDLE
@@ -409,6 +420,7 @@ void FUObjectArray::FreeUObjectIndex(UObjectBase* Object)
 	if (Index > ObjLastNonGCIndex && !GExitPurge && bShouldRecycleObjectIndices)
 	{
 		ObjAvailableList.Add(Index);
+		ObjAvailableListEstimateCount = ObjAvailableList.Num();
 	}
 }
 
diff --git a/Engine/Source/Runtime/CoreUObject/Public/UObject/UObjectArray.h b/Engine/Source/Runtime/CoreUObject/Public/UObject/UObjectArray.h
--- a/Engine/Source/Runtime/CoreUObject/Public/UObject/UObjectArray.h
+++ b/Engine/Source/Runtime/CoreUObject/Public/UObject/UObjectArray.h
@@ -192,7 +192,14 @@ public:
 
 	FORCEINLINE int32 GetSerialNumber() const
 	{
+		// AutoRTFM doesn't like atomics even when relaxed, so we need to differentiate the code here
+#if WITH_INSTRUMENTATION || USING_THREAD_SANITISER
+		// Resolving an old WeakPtr can race with new object construction.
+		// Use a relaxed atomic here so that TSAN knows this is considered safe.
+		return FPlatformAtomics::AtomicRead_Relaxed(&SerialNumber);
+#else
 		return SerialNumber;
+#endif
 	}
 
 #if UE_WITH_REMOTE_OBJECT_HANDLE
@@ -1226,7 +1233,7 @@ public:
 	 */
 	int32 GetObjectArrayNumMinusAvailable() const
 	{
-		return ObjObjects.Num() - ObjAvailableList.Num();
+		return ObjObjects.Num() - ObjAvailableListEstimateCount;
 	}
 
 	/**
@@ -1439,6 +1446,8 @@ private:
 	mutable FTransactionallySafeCriticalSection ObjObjectsCritical;
 	/** Available object indices.											*/
 	TArray<int32> ObjAvailableList;
+	/** We need the estimate without taking a lock on the array to fetch it from ObjAvailableList. */
+	TSAN_ATOMIC(int32) ObjAvailableListEstimateCount = 0;
 
 	/**
 	 * Array of things to notify when a UObjectBase is created
