Replaced const TCHAR* parameters with FStringView parameters in CoreUObject functions.

diff --git a/Engine/Source/Runtime/CoreUObject/Private/UObject/Class.cpp b/Engine/Source/Runtime/CoreUObject/Private/UObject/Class.cpp
--- a/Engine/Source/Runtime/CoreUObject/Private/UObject/Class.cpp
+++ b/Engine/Source/Runtime/CoreUObject/Private/UObject/Class.cpp
@@ -7064,7 +7064,7 @@ const FString UClass::GetConfigName() const
 	}
 }
 
-UField* UClass::TryFindTypeSlow(UClass* TypeClass, const FString& InPathNameOrShortName, EFindFirstObjectOptions InOptions)
+UField* UClass::TryFindTypeSlow(UClass* TypeClass, FStringView InPathNameOrShortName, EFindFirstObjectOptions InOptions)
 {
 	checkf(TypeClass && TypeClass->IsChildOf(UField::StaticClass()), TEXT("TryFindType requires a valid TypeClass parameter which is a subclass of UField. \"%s\" provided."), *GetFullNameSafe(TypeClass));
 
@@ -7074,13 +7074,13 @@ UField* UClass::TryFindTypeSlow(UClass* TypeClass, const FString& InPathNameOrSh
 	{
 		if (!FPackageName::IsShortPackageName(InPathNameOrShortName))
 		{
-			FoundType = (UField*)StaticFindObject(TypeClass, nullptr, *InPathNameOrShortName, !!(InOptions & EFindFirstObjectOptions::ExactClass) ? EFindObjectFlags::ExactClass : EFindObjectFlags::None);
+			FoundType = (UField*)StaticFindObject(TypeClass, nullptr, InPathNameOrShortName, !!(InOptions & EFindFirstObjectOptions::ExactClass) ? EFindObjectFlags::ExactClass : EFindObjectFlags::None);
 		}
 		else
 		{
 			// RobM: I can't decide if this should be an ensure. Maybe in the future?
 			TStringBuilder<1024> Callstack;
-			ANSICHAR Buffer[1024];			
+			ANSICHAR Buffer[1024];
 			uint64 StackFrames[10];
 			uint32 NumStackFrames = FPlatformStackWalk::CaptureStackBackTrace(StackFrames, UE_ARRAY_COUNT(StackFrames));
 			const uint32 IgnoreStackCount = 1; // Ignore the call to CaptureStackBackTrace itself
@@ -7103,16 +7103,16 @@ UField* UClass::TryFindTypeSlow(UClass* TypeClass, const FString& InPathNameOrSh
 				Callstack.Append(TrimmedBuffer);
 				Callstack.Append(TEXT("\r\n"));
 			}
-			
-			FoundType = (UField*)StaticFindFirstObject(TypeClass, *InPathNameOrShortName, InOptions | EFindFirstObjectOptions::EnsureIfAmbiguous | EFindFirstObjectOptions::NativeFirst, ELogVerbosity::Error, TEXT("TryFindType"));
 
-			UE_LOG(LogClass, Warning, TEXT("Short type name \"%s\" provided for TryFindType. Please convert it to a path name (suggested: \"%s\"). Callstack:\r\n\r\n%s"), *InPathNameOrShortName, *GetPathNameSafe(FoundType), Callstack.ToString());
+			FoundType = (UField*)StaticFindFirstObject(TypeClass, InPathNameOrShortName, InOptions | EFindFirstObjectOptions::EnsureIfAmbiguous | EFindFirstObjectOptions::NativeFirst, ELogVerbosity::Error, TEXT("TryFindType"));
+
+			UE_LOG(LogClass, Warning, TEXT("Short type name \"%.*s\" provided for TryFindType. Please convert it to a path name (suggested: \"%s\"). Callstack:\r\n\r\n%s"), InPathNameOrShortName.Len(), InPathNameOrShortName.GetData(), *GetPathNameSafe(FoundType), Callstack.ToString());
 		}
 	}
 	return FoundType;
 }
 
-UField* UClass::TryFindTypeSlowSafe(UClass* TypeClass, const FString& InPathNameOrShortName, EFindFirstObjectOptions InOptions)
+UField* UClass::TryFindTypeSlowSafe(UClass* TypeClass, FStringView InPathNameOrShortName, EFindFirstObjectOptions InOptions)
 {
 	if (!UE::IsSavingPackage(nullptr) && !IsGarbageCollectingAndLockingUObjectHashTables())
 	{
@@ -7121,7 +7121,7 @@ UField* UClass::TryFindTypeSlowSafe(UClass* TypeClass, const FString& InPathName
 	return nullptr;
 }
 
-FTopLevelAssetPath UClass::TryConvertShortTypeNameToPathName(UClass* TypeClass, const FString& InShortTypeName, ELogVerbosity::Type AmbiguousMessageVerbosity /*= ELogVerbosity::NoLogging*/, const TCHAR* AmbiguousClassMessage /*= nullptr*/)
+FTopLevelAssetPath UClass::TryConvertShortTypeNameToPathName(UClass* TypeClass, FStringView InShortTypeName, ELogVerbosity::Type AmbiguousMessageVerbosity /*= ELogVerbosity::NoLogging*/, const TCHAR* AmbiguousClassMessage /*= nullptr*/)
 {
 	checkf(TypeClass && TypeClass->IsChildOf(UField::StaticClass()), TEXT("TryConvertShortTypeNameToPathName requires a valid TypeClass parameter which is a subclass of UField. \"%s\" provided."), *GetFullNameSafe(TypeClass));
 
@@ -7141,7 +7141,7 @@ FTopLevelAssetPath UClass::TryConvertShortTypeNameToPathName(UClass* TypeClass,
 		{
 			Options |= EFindFirstObjectOptions::EnsureIfAmbiguous;
 		}
-		UField* FoundType = (UField*)StaticFindFirstObject(TypeClass, *InShortTypeName, Options, AmbiguousMessageVerbosity, AmbiguousClassMessage);
+		UField* FoundType = (UField*)StaticFindFirstObject(TypeClass, InShortTypeName, Options, AmbiguousMessageVerbosity, AmbiguousClassMessage);
 		if (FoundType)
 		{
 			// UField does not define a GetFieldPathName so do what GetStructPathName does (this is the only place that needs it)
diff --git a/Engine/Source/Runtime/CoreUObject/Private/UObject/UObjectGlobals.cpp b/Engine/Source/Runtime/CoreUObject/Private/UObject/UObjectGlobals.cpp
--- a/Engine/Source/Runtime/CoreUObject/Private/UObject/UObjectGlobals.cpp
+++ b/Engine/Source/Runtime/CoreUObject/Private/UObject/UObjectGlobals.cpp
@@ -484,16 +484,15 @@ bool ResolveName2(UObject*& InPackage, FStringBuilderBase& InOutName, bool Creat
 
 
 #if WITH_EDITOR
-static UObject* LoadObjectWhenImportingT3D(UClass* ObjectClass, const TCHAR* OrigInName)
+static UObject* LoadObjectWhenImportingT3D(UClass* ObjectClass, FStringView OrigInName)
 {
 	UObject* MatchingObject = nullptr;
 	static bool s_bCurrentlyLoading = false;
 	if (s_bCurrentlyLoading == false)
 	{
-		FString NameCheck = OrigInName;
-		if (NameCheck.Contains(TEXT("."), ESearchCase::CaseSensitive) &&
-			!NameCheck.Contains(TEXT("'"), ESearchCase::CaseSensitive) &&
-			!NameCheck.Contains(TEXT(":"), ESearchCase::CaseSensitive))
+		if (OrigInName.Contains(TEXT("."), ESearchCase::CaseSensitive) &&
+			!OrigInName.Contains(TEXT("'"), ESearchCase::CaseSensitive) &&
+			!OrigInName.Contains(TEXT(":"), ESearchCase::CaseSensitive))
 		{
 			s_bCurrentlyLoading = true;
 			MatchingObject = StaticLoadObject(ObjectClass, nullptr, OrigInName, nullptr, LOAD_NoWarn, nullptr);
@@ -507,7 +506,7 @@ static UObject* LoadObjectWhenImportingT3D(UClass* ObjectClass, const TCHAR* Ori
 //
 // Find an optional object.
 //
-UObject* StaticFindObject( UClass* ObjectClass, UObject* InObjectPackage, const TCHAR* OrigInName, EFindObjectFlags Flags )
+UObject* StaticFindObject( UClass* ObjectClass, UObject* InObjectPackage, FStringView OrigInName, EFindObjectFlags Flags )
 {
 	INC_DWORD_STAT(STAT_FindObject);
 
@@ -576,13 +575,13 @@ UObject* StaticFindObjectSafe(UClass* Class, FTopLevelAssetPath ObjectPath, EFin
 //
 // Find an object; can't fail.
 //
-UObject* StaticFindObjectChecked( UClass* ObjectClass, UObject* ObjectParent, const TCHAR* InName, EFindObjectFlags Flags )
+UObject* StaticFindObjectChecked( UClass* ObjectClass, UObject* ObjectParent, FStringView InName, EFindObjectFlags Flags )
 {
 	UObject* Result = StaticFindObject( ObjectClass, ObjectParent, InName, Flags );
 #if !(UE_BUILD_SHIPPING || UE_BUILD_TEST)
 	if( !Result )
 	{
-		UE_LOG(LogUObjectGlobals, Fatal, TEXT("Failed to find object '%s %s.%s'"), *ObjectClass->GetName(), ObjectParent ? *ObjectParent->GetName() : TEXT("None"), InName);
+		UE_LOG(LogUObjectGlobals, Fatal, TEXT("Failed to find object '%s %s.%.*s'"), *ObjectClass->GetName(), ObjectParent ? *ObjectParent->GetName() : TEXT("None"), InName.Len(), InName.GetData());
 	}
 #endif
 	return Result;
@@ -591,7 +590,7 @@ UObject* StaticFindObjectChecked( UClass* ObjectClass, UObject* ObjectParent, co
 //
 // Find an object; won't assert on UE::IsSavingPackage() or IsGarbageCollectingAndLockingUObjectHashTables()
 //
-UObject* StaticFindObjectSafe( UClass* ObjectClass, UObject* ObjectParent, const TCHAR* InName, EFindObjectFlags Flags )
+UObject* StaticFindObjectSafe( UClass* ObjectClass, UObject* ObjectParent, FStringView InName, EFindObjectFlags Flags )
 {
 	if (!UE::IsSavingPackage(nullptr) && !IsGarbageCollectingAndLockingUObjectHashTables())
 	{
@@ -664,11 +663,11 @@ bool StaticFindAllObjectsSafe(TArray<UObject*>& OutFoundObjects, UClass* ObjectC
 	return bFoundObjects;
 }
 
-UObject* StaticFindFirstObject(UClass* Class, const TCHAR* Name, EFindFirstObjectOptions Options /*= EFindFirstObjectOptions::None*/, ELogVerbosity::Type AmbiguousMessageVerbosity /*= ELogVerbosity::Warning*/, const TCHAR* InCurrentOperation /*= nullptr*/)
+UObject* StaticFindFirstObject(UClass* Class, FStringView Name, EFindFirstObjectOptions Options /*= EFindFirstObjectOptions::None*/, ELogVerbosity::Type AmbiguousMessageVerbosity /*= ELogVerbosity::Warning*/, const TCHAR* InCurrentOperation /*= nullptr*/)
 {
 	UObject* Result = nullptr;
 	FName ObjectName;
-	if (!FCString::Strchr(Name, TCHAR('\'')))
+	if (UE::String::FindFirstChar(Name, TCHAR('\'')) == INDEX_NONE)
 	{
 		// Skip unnecessary allocations in StripObjectClass
 		ObjectName = FName(Name, FNAME_Add);
@@ -774,7 +773,7 @@ UObject* StaticFindFirstObject(UClass* Class, const TCHAR* Name, EFindFirstObjec
 	return Result;
 }
 
-UObject* StaticFindFirstObjectSafe(UClass* Class, const TCHAR* Name, EFindFirstObjectOptions Options /*= EFindFirstObjectOptions::None*/, ELogVerbosity::Type AmbiguousMessageVerbosity /*= ELogVerbosity::Warning*/, const TCHAR* InCurrentOperation /*= nullptr*/)
+UObject* StaticFindFirstObjectSafe(UClass* Class, FStringView Name, EFindFirstObjectOptions Options /*= EFindFirstObjectOptions::None*/, ELogVerbosity::Type AmbiguousMessageVerbosity /*= ELogVerbosity::Warning*/, const TCHAR* InCurrentOperation /*= nullptr*/)
 {
 	UObject* Result = nullptr;
 	if (!UE::IsSavingPackage(nullptr) && !IsGarbageCollectingAndLockingUObjectHashTables())
@@ -1320,20 +1319,19 @@ UE_TRACE_EVENT_BEGIN(Cpu, LoadObject, NoSync)
 	UE_TRACE_EVENT_FIELD(UE::Trace::WideString, AssetPath)
 UE_TRACE_EVENT_END()
 
-UObject* StaticLoadObjectInternal(UClass* ObjectClass, UObject* InOuter, const TCHAR* InName, const TCHAR* Filename, uint32 LoadFlags, UPackageMap* Sandbox, bool bAllowObjectReconciliation, const FLinkerInstancingContext* InstancingContext)
+UObject* StaticLoadObjectInternal(UClass* ObjectClass, UObject* InOuter, FStringView InName, FStringView Filename, uint32 LoadFlags, UPackageMap* Sandbox, bool bAllowObjectReconciliation, const FLinkerInstancingContext* InstancingContext)
 {
 #if CPUPROFILERTRACE_ENABLED
 	UE_TRACE_LOG_SCOPED_T(Cpu, LoadObject, CpuChannel)
-		<< LoadObject.AssetPath(InName);
+		<< LoadObject.AssetPath(InName.GetData(), InName.Len());
 #endif // CPUPROFILERTRACE_ENABLED
 	SCOPED_NAMED_EVENT(StaticLoadObjectInternal, FColor::Red);
-	check(InName);
 
-	FScopedLoadingState ScopedLoadingState(InName);
+	FScopedLoadingState ScopedLoadingState = FScopedLoadingState(FString(InName));
 	TStringBuilder<512> StrName;
 	StrName = InName;
 	UObject* Result = nullptr;
-	const bool bContainsObjectName = !!FCString::Strstr(InName, TEXT("."));
+	const bool bContainsObjectName = UE::String::FindFirstChar(InName, TEXT('.')) != INDEX_NONE;
 
 	// break up the name into packages, returning the innermost name and its outer
 	ResolveName2(InOuter, StrName, true, true, LoadFlags & (LOAD_EditorOnly | LOAD_NoVerify | LOAD_Quiet | LOAD_NoWarn | LOAD_DeferDependencyLoads), InstancingContext);
@@ -1393,7 +1391,7 @@ UObject* StaticLoadObjectInternal(UClass* ObjectClass, UObject* InOuter, const T
 		StrName = InName;
 		StrName += TEXT(".");
 		StrName += FPackageName::GetShortName(InName);
-		Result = StaticLoadObjectInternal(ObjectClass, InOuter, *StrName, Filename, LoadFlags, Sandbox, bAllowObjectReconciliation, InstancingContext);
+		Result = StaticLoadObjectInternal(ObjectClass, InOuter, StrName, Filename, LoadFlags, Sandbox, bAllowObjectReconciliation, InstancingContext);
 	}
 
 	if (Result && UE::GC::GIsIncrementalReachabilityPending)
@@ -1403,7 +1401,7 @@ UObject* StaticLoadObjectInternal(UClass* ObjectClass, UObject* InOuter, const T
 	return Result;
 }
 
-UObject* StaticLoadObject(UClass* ObjectClass, UObject* InOuter, const TCHAR* InName, const TCHAR* Filename, uint32 LoadFlags, UPackageMap* Sandbox, bool bAllowObjectReconciliation, const FLinkerInstancingContext* InstancingContext)
+UObject* StaticLoadObject(UClass* ObjectClass, UObject* InOuter, FStringView InName, FStringView Filename, uint32 LoadFlags, UPackageMap* Sandbox, bool bAllowObjectReconciliation, const FLinkerInstancingContext* InstancingContext)
 {
 	UObject* Result = StaticLoadObjectInternal(ObjectClass, InOuter, InName, Filename, LoadFlags, Sandbox, bAllowObjectReconciliation, InstancingContext);
 	if (!Result)
@@ -1435,7 +1433,7 @@ UObject* StaticLoadObject(UClass* ObjectClass, UObject* InOuter, const TCHAR* In
 //
 // Load a class.
 //
-UClass* StaticLoadClass( UClass* BaseClass, UObject* InOuter, const TCHAR* InName, const TCHAR* Filename, uint32 LoadFlags, UPackageMap* Sandbox )
+UClass* StaticLoadClass( UClass* BaseClass, UObject* InOuter, FStringView InName, FStringView Filename, uint32 LoadFlags, UPackageMap* Sandbox )
 {
 	check(BaseClass);
 
diff --git a/Engine/Source/Runtime/CoreUObject/Public/UObject/Class.h b/Engine/Source/Runtime/CoreUObject/Public/UObject/Class.h
--- a/Engine/Source/Runtime/CoreUObject/Public/UObject/Class.h
+++ b/Engine/Source/Runtime/CoreUObject/Public/UObject/Class.h
@@ -3393,7 +3393,7 @@ public:
 	 * @param InShortNameOrPathName Class name
 	 * @return Pointer to a type object if successful, nullptr otherwise
 	 */
-	static COREUOBJECT_API UField* TryFindTypeSlow(UClass* TypeClass, const FString& InPathNameOrShortName, EFindFirstObjectOptions InOptions = EFindFirstObjectOptions::None);
+	static COREUOBJECT_API UField* TryFindTypeSlow(UClass* TypeClass, FStringView InPathNameOrShortName, EFindFirstObjectOptions InOptions = EFindFirstObjectOptions::None);
 
 	/**
 	 * Utility function that tries to find a type (class/struct/enum) given a path name or a short name however it will throw a warning (with a callstack) if it's the latter.
@@ -3403,7 +3403,7 @@ public:
 	 * @return Pointer to a type object if successful, false otherwise
 	 */
 	template <typename T>
-	static T* TryFindTypeSlow(const FString& InShortNameOrPathName, EFindFirstObjectOptions InOptions = EFindFirstObjectOptions::None)
+	static T* TryFindTypeSlow(FStringView InShortNameOrPathName, EFindFirstObjectOptions InOptions = EFindFirstObjectOptions::None)
 	{
 		return (T*)TryFindTypeSlow(T::StaticClass(), InShortNameOrPathName, InOptions);
 	}
@@ -3417,7 +3417,7 @@ public:
 	 * @param InShortNameOrPathName Class name
 	 * @return Pointer to a type object if successful, nullptr otherwise
 	 */
-	static COREUOBJECT_API UField* TryFindTypeSlowSafe(UClass* TypeClass, const FString& InPathNameOrShortName, EFindFirstObjectOptions InOptions = EFindFirstObjectOptions::None);
+	static COREUOBJECT_API UField* TryFindTypeSlowSafe(UClass* TypeClass, FStringView InPathNameOrShortName, EFindFirstObjectOptions InOptions = EFindFirstObjectOptions::None);
 
 	/**
 	 * Utility function that tries to find a type (class/struct/enum) given a path name or a short name however it will throw a warning (with a callstack) if it's the latter.
@@ -3428,7 +3428,7 @@ public:
 	 * @return Pointer to a type object if successful, nullptr otherwise
 	 */
 	template <typename T>
-	static T* TryFindTypeSlowSafe(const FString& InShortNameOrPathName, EFindFirstObjectOptions InOptions = EFindFirstObjectOptions::None)
+	static T* TryFindTypeSlowSafe(FStringView InShortNameOrPathName, EFindFirstObjectOptions InOptions = EFindFirstObjectOptions::None)
 	{
 		return (T*)TryFindTypeSlowSafe(T::StaticClass(), InShortNameOrPathName, InOptions);
 	}
@@ -3442,7 +3442,7 @@ public:
 	 * @param AmbiguousClassMessage Additional message to log when class name is ambiguous (e.g. current operation) 
 	 * @return Class path name if successful. Empty FTopLevelAssetPath otherwise.
 	 */
-	static COREUOBJECT_API FTopLevelAssetPath TryConvertShortTypeNameToPathName(UClass* TypeClass, const FString& InShortTypeName, ELogVerbosity::Type AmbiguousMessageVerbosity = ELogVerbosity::NoLogging, const TCHAR* AmbiguousClassMessage = nullptr);
+	static COREUOBJECT_API FTopLevelAssetPath TryConvertShortTypeNameToPathName(UClass* TypeClass, FStringView InShortTypeName, ELogVerbosity::Type AmbiguousMessageVerbosity = ELogVerbosity::NoLogging, const TCHAR* AmbiguousClassMessage = nullptr);
 
 	/**
 	 * Tries to convert short class name to class path name. This will only work if the class exists in memory
@@ -3452,7 +3452,7 @@ public:
 	 * @return Class path name if successful. Empty FTopLevelAssetPath otherwise.
 	 */
 	template <typename T>
-	static FTopLevelAssetPath TryConvertShortTypeNameToPathName(const FString& InShortTypeName, ELogVerbosity::Type AmbiguousMessageVerbosity = ELogVerbosity::NoLogging, const TCHAR* AmbiguousClassMessage = nullptr)
+	static FTopLevelAssetPath TryConvertShortTypeNameToPathName(FStringView InShortTypeName, ELogVerbosity::Type AmbiguousMessageVerbosity = ELogVerbosity::NoLogging, const TCHAR* AmbiguousClassMessage = nullptr)
 	{
 		return TryConvertShortTypeNameToPathName(T::StaticClass(), InShortTypeName, AmbiguousMessageVerbosity, AmbiguousClassMessage);
 	}
diff --git a/Engine/Source/Runtime/CoreUObject/Public/UObject/UObjectGlobals.h b/Engine/Source/Runtime/CoreUObject/Public/UObject/UObjectGlobals.h
--- a/Engine/Source/Runtime/CoreUObject/Public/UObject/UObjectGlobals.h
+++ b/Engine/Source/Runtime/CoreUObject/Public/UObject/UObjectGlobals.h
@@ -349,10 +349,10 @@ UE_NODEBUG FORCEINLINE UObject* StaticFindObjectFastSafe(UClass* Class, UObject*
  *
  * @return	Returns a pointer to the found object or nullptr if none could be found
  */
-COREUOBJECT_API UObject* StaticFindObject( UClass* Class, UObject* InOuter, const TCHAR* Name, EFindObjectFlags Flags = EFindObjectFlags::None );
+COREUOBJECT_API UObject* StaticFindObject( UClass* Class, UObject* InOuter, FStringView Name, EFindObjectFlags Flags = EFindObjectFlags::None );
 template <UE::CSameAs<bool> ExactClassType>
 UE_EXACTCLASS_BOOL_DEPRECATED("StaticFindObject")
-UE_NODEBUG FORCEINLINE UObject* StaticFindObject( UClass* Class, UObject* InOuter, const TCHAR* Name, ExactClassType bExactClass )
+UE_NODEBUG FORCEINLINE UObject* StaticFindObject( UClass* Class, UObject* InOuter, FStringView Name, ExactClassType bExactClass )
 {
 	return StaticFindObject(Class, InOuter, Name, bExactClass ? EFindObjectFlags::ExactClass : EFindObjectFlags::None);
 }
@@ -375,19 +375,19 @@ UE_NODEBUG FORCEINLINE UObject* StaticFindObject(UClass* Class, FTopLevelAssetPa
 }
 
 /** Version of StaticFindObject() that will assert if the object is not found */
-COREUOBJECT_API UObject* StaticFindObjectChecked( UClass* Class, UObject* InOuter, const TCHAR* Name, EFindObjectFlags Flags = EFindObjectFlags::None );
+COREUOBJECT_API UObject* StaticFindObjectChecked( UClass* Class, UObject* InOuter, FStringView Name, EFindObjectFlags Flags = EFindObjectFlags::None );
 template <UE::CSameAs<bool> ExactClassType>
 UE_EXACTCLASS_BOOL_DEPRECATED("StaticFindObjectChecked")
-UE_NODEBUG FORCEINLINE UObject* StaticFindObjectChecked( UClass* Class, UObject* InOuter, const TCHAR* Name, ExactClassType bExactClass)
+UE_NODEBUG FORCEINLINE UObject* StaticFindObjectChecked( UClass* Class, UObject* InOuter, FStringView Name, ExactClassType bExactClass)
 {
 	return StaticFindObjectChecked(Class, InOuter, Name, bExactClass ? EFindObjectFlags::ExactClass : EFindObjectFlags::None);
 }
 
 /** Internal version of StaticFindObject that will not assert on GIsSavingPackage or IsGarbageCollectingAndLockingUObjectHashTables() */
-COREUOBJECT_API UObject* StaticFindObjectSafe( UClass* Class, UObject* InOuter, const TCHAR* Name, EFindObjectFlags Flags = EFindObjectFlags::None );
+COREUOBJECT_API UObject* StaticFindObjectSafe( UClass* Class, UObject* InOuter, FStringView Name, EFindObjectFlags Flags = EFindObjectFlags::None );
 template <UE::CSameAs<bool> ExactClassType>
 UE_EXACTCLASS_BOOL_DEPRECATED("StaticFindObjectSafe")
-UE_NODEBUG FORCEINLINE UObject* StaticFindObjectSafe( UClass* Class, UObject* InOuter, const TCHAR* Name, ExactClassType bExactClass )
+UE_NODEBUG FORCEINLINE UObject* StaticFindObjectSafe( UClass* Class, UObject* InOuter, FStringView Name, ExactClassType bExactClass )
 {
 	return StaticFindObjectSafe(Class, InOuter, Name, bExactClass ? EFindObjectFlags::ExactClass : EFindObjectFlags::None);
 }
@@ -513,7 +513,7 @@ ENUM_CLASS_FLAGS(EFindFirstObjectOptions);
  *
  * @return	Returns a pointer to an object if found, null otherwise
  */
-COREUOBJECT_API UObject* StaticFindFirstObject(UClass* Class, const TCHAR* Name, EFindFirstObjectOptions Options = EFindFirstObjectOptions::None, ELogVerbosity::Type AmbiguousMessageVerbosity = ELogVerbosity::NoLogging, const TCHAR* InCurrentOperation = nullptr);
+COREUOBJECT_API UObject* StaticFindFirstObject(UClass* Class, FStringView Name, EFindFirstObjectOptions Options = EFindFirstObjectOptions::None, ELogVerbosity::Type AmbiguousMessageVerbosity = ELogVerbosity::NoLogging, const TCHAR* InCurrentOperation = nullptr);
 
 /**
  * Tries to find the first objects matching the search paramters in memory. This will handle fully qualified paths of the form /path/packagename.object:subobject and resolve references for you.
@@ -529,7 +529,7 @@ COREUOBJECT_API UObject* StaticFindFirstObject(UClass* Class, const TCHAR* Name,
  *
  * @return	Returns a pointer to an object if found, null otherwise
  */
-COREUOBJECT_API UObject* StaticFindFirstObjectSafe(UClass* Class, const TCHAR* Name, EFindFirstObjectOptions Options = EFindFirstObjectOptions::None, ELogVerbosity::Type AmbiguousMessageVerbosity = ELogVerbosity::NoLogging, const TCHAR* InCurrentOperation = nullptr);
+COREUOBJECT_API UObject* StaticFindFirstObjectSafe(UClass* Class, FStringView Name, EFindFirstObjectOptions Options = EFindFirstObjectOptions::None, ELogVerbosity::Type AmbiguousMessageVerbosity = ELogVerbosity::NoLogging, const TCHAR* InCurrentOperation = nullptr);
 
 /** Loading policy to use with ParseObject */
 enum class EParseObjectLoadingPolicy : uint8
@@ -574,10 +574,10 @@ inline bool ParseObject( const TCHAR* Stream, const TCHAR* Match, UClass* Class,
  *
  * @return The object that was loaded or found. nullptr for a failure.
  */
-COREUOBJECT_API UObject* StaticLoadObject( UClass* Class, UObject* InOuter, const TCHAR* Name, const TCHAR* Filename = nullptr, uint32 LoadFlags = LOAD_None, UPackageMap* Sandbox = nullptr, bool bAllowObjectReconciliation = true, const FLinkerInstancingContext* InstancingContext = nullptr);
+COREUOBJECT_API UObject* StaticLoadObject( UClass* Class, UObject* InOuter, FStringView Name, FStringView Filename = {}, uint32 LoadFlags = LOAD_None, UPackageMap* Sandbox = nullptr, bool bAllowObjectReconciliation = true, const FLinkerInstancingContext* InstancingContext = nullptr);
 
 /** Version of StaticLoadObject() that will load classes */
-COREUOBJECT_API UClass* StaticLoadClass(UClass* BaseClass, UObject* InOuter, const TCHAR* Name, const TCHAR* Filename = nullptr, uint32 LoadFlags = LOAD_None, UPackageMap* Sandbox = nullptr);
+COREUOBJECT_API UClass* StaticLoadClass(UClass* BaseClass, UObject* InOuter, FStringView Name, FStringView Filename = {}, uint32 LoadFlags = LOAD_None, UPackageMap* Sandbox = nullptr);
 
 /**
  * Find or load an object that is one of the top level assets in a package.
@@ -2024,13 +2024,13 @@ UE_NODEBUG FORCEINLINE T* FindObjectFast(UObject* Outer, FName Name, ExactClassT
  * @see StaticFindObject()
  */
 template< class T > 
-inline T* FindObject( UObject* Outer, const TCHAR* Name, EFindObjectFlags Flags = EFindObjectFlags::None )
+inline T* FindObject( UObject* Outer, FStringView Name, EFindObjectFlags Flags = EFindObjectFlags::None )
 {
 	return (T*)StaticFindObject( T::StaticClass(), Outer, Name, Flags );
 }
 template <typename T, UE::CSameAs<bool> ExactClassType>
 UE_EXACTCLASS_BOOL_DEPRECATED("FindObject")
-UE_NODEBUG FORCEINLINE T* FindObject(UObject* Outer, const TCHAR* Name, ExactClassType bExactClass)
+UE_NODEBUG FORCEINLINE T* FindObject(UObject* Outer, FStringView Name, ExactClassType bExactClass)
 {
 	return FindObject<T>(Outer, Name, bExactClass ? EFindObjectFlags::ExactClass : EFindObjectFlags::None);
 }
@@ -2056,13 +2056,13 @@ UE_NODEBUG FORCEINLINE T* FindObject(FTopLevelAssetPath InPath, ExactClassType b
  * @see StaticFindObjectChecked()
  */
 template< class T > 
-inline T* FindObjectChecked( UObject* Outer, const TCHAR* Name, EFindObjectFlags Flags = EFindObjectFlags::None )
+inline T* FindObjectChecked( UObject* Outer, FStringView Name, EFindObjectFlags Flags = EFindObjectFlags::None )
 {
 	return (T*)StaticFindObjectChecked( T::StaticClass(), Outer, Name, Flags );
 }
 template <typename T, UE::CSameAs<bool> ExactClassType>
 UE_EXACTCLASS_BOOL_DEPRECATED("FindObjectChecked")
-UE_NODEBUG FORCEINLINE T* FindObjectChecked( UObject* Outer, const TCHAR* Name, ExactClassType bExactClass )
+UE_NODEBUG FORCEINLINE T* FindObjectChecked( UObject* Outer, FStringView Name, ExactClassType bExactClass )
 {
 	return FindObjectChecked<T>(Outer, Name, bExactClass ? EFindObjectFlags::ExactClass : EFindObjectFlags::None);
 }
@@ -2072,13 +2072,13 @@ UE_NODEBUG FORCEINLINE T* FindObjectChecked( UObject* Outer, const TCHAR* Name,
  * @see StaticFindObjectSafe()
  */
 template< class T > 
-inline T* FindObjectSafe( UObject* Outer, const TCHAR* Name, EFindObjectFlags Flags = EFindObjectFlags::None )
+inline T* FindObjectSafe( UObject* Outer, FStringView Name, EFindObjectFlags Flags = EFindObjectFlags::None )
 {
 	return (T*)StaticFindObjectSafe( T::StaticClass(), Outer, Name, Flags );
 }
 template <typename T, UE::CSameAs<bool> ExactClassType>
 UE_EXACTCLASS_BOOL_DEPRECATED("FindObjectSafe")
-UE_NODEBUG FORCEINLINE T* FindObjectSafe( UObject* Outer, const TCHAR* Name, ExactClassType bExactClass )
+UE_NODEBUG FORCEINLINE T* FindObjectSafe( UObject* Outer, FStringView Name, ExactClassType bExactClass )
 {
 	return FindObjectSafe<T>(Outer, Name, bExactClass ? EFindObjectFlags::ExactClass : EFindObjectFlags::None);
 }
@@ -2104,7 +2104,7 @@ UE_NODEBUG FORCEINLINE T* FindObjectSafe(FTopLevelAssetPath InPath, ExactClassTy
  * @see StaticFindFirstObject()
  */
 template< class T >
-inline T* FindFirstObject(const TCHAR* Name, EFindFirstObjectOptions Options = EFindFirstObjectOptions::None, ELogVerbosity::Type AmbiguousMessageVerbosity = ELogVerbosity::NoLogging, const TCHAR* CurrentOperation = nullptr)
+inline T* FindFirstObject(FStringView Name, EFindFirstObjectOptions Options = EFindFirstObjectOptions::None, ELogVerbosity::Type AmbiguousMessageVerbosity = ELogVerbosity::NoLogging, const TCHAR* CurrentOperation = nullptr)
 {
 	return (T*)StaticFindFirstObject(T::StaticClass(), Name, Options, AmbiguousMessageVerbosity, CurrentOperation);
 }
@@ -2114,7 +2114,7 @@ inline T* FindFirstObject(const TCHAR* Name, EFindFirstObjectOptions Options = E
  * @see StaticFindFirstObject()
  */
 template< class T >
-inline T* FindFirstObjectSafe(const TCHAR* Name, EFindFirstObjectOptions Options = EFindFirstObjectOptions::None, ELogVerbosity::Type AmbiguousMessageVerbosity = ELogVerbosity::NoLogging, const TCHAR* CurrentOperation = nullptr)
+inline T* FindFirstObjectSafe(FStringView Name, EFindFirstObjectOptions Options = EFindFirstObjectOptions::None, ELogVerbosity::Type AmbiguousMessageVerbosity = ELogVerbosity::NoLogging, const TCHAR* CurrentOperation = nullptr)
 {
 	return (T*)StaticFindFirstObjectSafe(T::StaticClass(), Name, Options, AmbiguousMessageVerbosity, CurrentOperation);
 }
@@ -2124,7 +2124,7 @@ inline T* FindFirstObjectSafe(const TCHAR* Name, EFindFirstObjectOptions Options
  * @see StaticLoadObject()
  */
 template< class T > 
-inline T* LoadObject( UObject* Outer, const TCHAR* Name, const TCHAR* Filename=nullptr, uint32 LoadFlags=LOAD_None, UPackageMap* Sandbox=nullptr, const FLinkerInstancingContext* InstancingContext=nullptr )
+inline T* LoadObject( UObject* Outer, FStringView Name, FStringView Filename={}, uint32 LoadFlags=LOAD_None, UPackageMap* Sandbox=nullptr, const FLinkerInstancingContext* InstancingContext=nullptr )
 {
 	return (T*)StaticLoadObject( T::StaticClass(), Outer, Name, Filename, LoadFlags, Sandbox, true, InstancingContext );
 }
@@ -2134,7 +2134,7 @@ inline T* LoadObject( UObject* Outer, const TCHAR* Name, const TCHAR* Filename=n
  * @see StaticLoadClass
  */
 template< class T > 
-inline UClass* LoadClass( UObject* Outer, const TCHAR* Name, const TCHAR* Filename=nullptr, uint32 LoadFlags=LOAD_None, UPackageMap* Sandbox=nullptr )
+inline UClass* LoadClass( UObject* Outer, FStringView Name, FStringView Filename={}, uint32 LoadFlags=LOAD_None, UPackageMap* Sandbox=nullptr )
 {
 	return StaticLoadClass( T::StaticClass(), Outer, Name, Filename, LoadFlags, Sandbox );
 }
