[AR] Add ability to remove tags from the AssetRegistryTagsContext

diff --git a/Engine/Source/Runtime/CoreUObject/Private/UObject/AssetRegistryTagsContext.cpp b/Engine/Source/Runtime/CoreUObject/Private/UObject/AssetRegistryTagsContext.cpp
--- a/Engine/Source/Runtime/CoreUObject/Private/UObject/AssetRegistryTagsContext.cpp
+++ b/Engine/Source/Runtime/CoreUObject/Private/UObject/AssetRegistryTagsContext.cpp
@@ -4,6 +4,7 @@
 
 #include "Logging/LogMacros.h"
 #include "Misc/StringBuilder.h"
+#include "UObject/NameTypes.h"
 
 FAssetRegistryTagsContextData::FAssetRegistryTagsContextData(const UObject* CallTarget,
 	EAssetRegistryTagsCaller InCaller)
@@ -102,6 +103,11 @@ void FAssetRegistryTagsContext::AddTag(UObject::FAssetRegistryTag TagResult)
 	AddTagInternal(MoveTemp(TagResult), FStringView());
 }
 
+void FAssetRegistryTagsContext::RemoveTag(FName TagName)
+{
+	RemoveTagInternal(TagName, FStringView());
+}
+
 int32 FAssetRegistryTagsContext::GetNumTags() const
 {
 	return Data.Tags.Num();
@@ -173,6 +179,15 @@ void FAssetRegistryTagsContext::AddCookTag(UObject::FAssetRegistryTag TagResult)
 	AddTagInternal(MoveTemp(TagResult), UE::AssetRegistry::CookTagPrefix);
 }
 
+void FAssetRegistryTagsContext::RemoveCookTag(FName TagName)
+{
+	if (!WantsCookTags())
+	{
+		return;
+	}
+	RemoveTagInternal(TagName, UE::AssetRegistry::CookTagPrefix);
+}
+
 bool FAssetRegistryTagsContext::TryValidateTagName(FName TagName) const
 {
 	using namespace UE::AssetRegistry;
@@ -192,7 +207,7 @@ bool FAssetRegistryTagsContext::TryValidateTagName(FName TagName) const
 	return true;
 }
 
-FName FAssetRegistryTagsContext::TransformTagName(FName TagName, FStringView PrefixToAdd) const
+FName FAssetRegistryTagsContext::TransformTagName(FName TagName, FStringView PrefixToAdd, bool bCreateNameIfDoesNotExist) const
 {
 	if (!TryValidateTagName(TagName))
 	{
@@ -201,7 +216,7 @@ FName FAssetRegistryTagsContext::TransformTagName(FName TagName, FStringView Pre
 	if (!PrefixToAdd.IsEmpty())
 	{
 		TStringBuilder<64> TagNameStr(InPlace, PrefixToAdd, TagName);
-		TagName = FName(TagNameStr);
+		TagName = FName(TagNameStr, bCreateNameIfDoesNotExist ? FNAME_Add : FNAME_Find);
 	}
 	return TagName;
 }
@@ -241,7 +256,7 @@ UObject::FAssetRegistryTag& FAssetRegistryTagsContext::FindOrAddTagInternal(FNam
 
 UObject::FAssetRegistryTag* FAssetRegistryTagsContext::FindTagInternal(FName TagName, FStringView PrefixToAdd)
 {
-	TagName = TransformTagName(TagName, PrefixToAdd);
+	TagName = TransformTagName(TagName, PrefixToAdd, false);
 	if (TagName.IsNone())
 	{
 		return nullptr;
@@ -252,7 +267,7 @@ UObject::FAssetRegistryTag* FAssetRegistryTagsContext::FindTagInternal(FName Tag
 
 bool FAssetRegistryTagsContext::ContainsTagInternal(FName TagName, FStringView PrefixToAdd) const
 {
-	TagName = TransformTagName(TagName, PrefixToAdd);
+	TagName = TransformTagName(TagName, PrefixToAdd, false);
 	if (TagName.IsNone())
 	{
 		return false;
@@ -271,4 +286,14 @@ void FAssetRegistryTagsContext::AddTagInternal(UObject::FAssetRegistryTag&& TagR
 	UObject::FAssetRegistryTag& AddedTag = Data.Tags.FindOrAdd(TagName);
 	AddedTag = MoveTemp(TagResult);
 	AddedTag.Name = TagName;
-}
\ No newline at end of file
+}
+
+void FAssetRegistryTagsContext::RemoveTagInternal(FName TagName, FStringView PrefixToAdd)
+{
+	TagName = TransformTagName(TagName, PrefixToAdd, false);
+	if (TagName.IsNone())
+	{
+		return;
+	}
+	Data.Tags.Remove(TagName);
+}
diff --git a/Engine/Source/Runtime/CoreUObject/Public/UObject/AssetRegistryTagsContext.h b/Engine/Source/Runtime/CoreUObject/Public/UObject/AssetRegistryTagsContext.h
--- a/Engine/Source/Runtime/CoreUObject/Public/UObject/AssetRegistryTagsContext.h
+++ b/Engine/Source/Runtime/CoreUObject/Public/UObject/AssetRegistryTagsContext.h
@@ -156,6 +156,8 @@ public:
 	COREUOBJECT_API bool ContainsTag(FName TagName) const;
 	/** Output function: move the given tag into the results, overwriting previous result if it exists. */
 	COREUOBJECT_API void AddTag(UObject::FAssetRegistryTag TagResult);
+	/** Removes the tag from the output. */
+	COREUOBJECT_API void RemoveTag(FName TagName);
 	/** Return the number of tags so far reported to the context, including CookTags. */
 	COREUOBJECT_API int32 GetNumTags() const;
 	/** Pass each tag so far reported to the context into the provided Visitor, including CookTags. */
@@ -183,14 +185,17 @@ public:
 	COREUOBJECT_API bool ContainsCookTag(FName TagName) const;
 	/** AddTag but for a cooktag. @see AddTag and @see FAssetRegistryTagsContext */
 	COREUOBJECT_API void AddCookTag(UObject::FAssetRegistryTag TagResult);
+	/** Removes the cooktag from the output. */
+	COREUOBJECT_API void RemoveCookTag(FName TagName);
 
 private:
 	bool TryValidateTagName(FName TagName) const;
-	FName TransformTagName(FName TagName, FStringView PrefixToAdd) const;
+	FName TransformTagName(FName TagName, FStringView PrefixToAdd, bool bCreateNameIfDoesNotExist = true) const;
 	UObject::FAssetRegistryTag& FindOrAddTagInternal(FName TagName, FStringView PrefixToAdd, bool* bOutAlreadyExists);
 	UObject::FAssetRegistryTag* FindTagInternal(FName TagName, FStringView PrefixToAdd);
 	bool ContainsTagInternal(FName TagName, FStringView PrefixToAdd) const;
 	void AddTagInternal(UObject::FAssetRegistryTag&& TagResult, FStringView PrefixToAdd);
+	void RemoveTagInternal(FName TagName, FStringView PrefixToAdd);
 
 private:
 	FAssetRegistryTagsContextData& Data;
