[UObject]
- Fix benign race by emptying the map only if capacity is not already 0

diff --git a/Engine/Source/Runtime/CoreUObject/Private/UObject/UObjectBase.cpp b/Engine/Source/Runtime/CoreUObject/Private/UObject/UObjectBase.cpp
--- a/Engine/Source/Runtime/CoreUObject/Private/UObject/UObjectBase.cpp
+++ b/Engine/Source/Runtime/CoreUObject/Private/UObject/UObjectBase.cpp
@@ -1076,7 +1076,10 @@ void ProcessNewlyLoadedUObjects(FName InModuleName, bool bCanProcessNewlyLoadedO
 
 	if (TMap<UObjectBase*, FPendingRegistrantInfo>& PendingRegistrants = FPendingRegistrantInfo::GetMap(); PendingRegistrants.IsEmpty())
 	{
-		PendingRegistrants.Empty();
+		if (PendingRegistrants.GetMaxIndex() > 0)
+		{
+			PendingRegistrants.Empty();
+		}
 	}
 
 	if (bNewUObjects && !GIsInitialLoad)
