Change granularity of checking time limit during C++ object destruction to match FinishDestroy, which improves accuracy of the GC time limit.
Add gc.IncrementalBeginDestroyGranularity define so this can be overridden if necessary.

diff --git a/Engine/Source/Runtime/CoreUObject/Private/UObject/GarbageCollection.cpp b/Engine/Source/Runtime/CoreUObject/Private/UObject/GarbageCollection.cpp
--- a/Engine/Source/Runtime/CoreUObject/Private/UObject/GarbageCollection.cpp
+++ b/Engine/Source/Runtime/CoreUObject/Private/UObject/GarbageCollection.cpp
@@ -230,6 +230,14 @@ static FAutoConsoleVariableRef CIncrementalBeginDestroyEnabled(
 	ECVF_Default
 );
 
+static int32 GIncrementalBeginDestroyGranularity = 10;
+static FAutoConsoleVariableRef CIncrementalBeginDestroyGranularity(
+	TEXT("gc.IncrementalBeginDestroyGranularity"),
+	GIncrementalBeginDestroyGranularity,
+	TEXT("When using gc.IncrementalBeginDestroyEnabled, check the time limit after this number of destroys. A low number has more overhead but is more accurate."),
+	ECVF_Default
+);
+
 #if WITH_VERSE_VM || defined(__INTELLISENSE__)
 bool GEnableFrankenGC = true;
 static FAutoConsoleVariableRef CEnableFrankenGC(
@@ -850,7 +858,7 @@ public:
 	FORCENOINLINE bool DestroyObjects(bool bUseTimeLimit, double TimeLimit, double StartTime)
 	{
 		TRACE_CPUPROFILER_EVENT_SCOPE(FObjectPurge::DestroyObjects);
-		constexpr int32 TimeLimitEnforcementGranularityForDeletion = 100;
+		const int32 TimeLimitEnforcementGranularityForDeletion = GIncrementalBeginDestroyGranularity;
 		int32 ProcessedObjectsCount = 0;
 
 		// Global UObject Array needs to be locked only when freeing UObject indices
@@ -4762,8 +4770,7 @@ bool IncrementalDestroyGarbage(bool bUseTimeLimit, double TimeLimit)
 	// Depending on platform FPlatformTime::Seconds might take a noticeable amount of time if called thousands of times so we avoid
 	// enforcing the time limit too often, especially as neither Destroy nor actual deletion should take significant
 	// amounts of time.
-	const int32	TimeLimitEnforcementGranularityForDestroy = 10;
-	const int32	TimeLimitEnforcementGranularityForDeletion = 100;
+	const int32	TimeLimitEnforcementGranularityForDestroy = GIncrementalBeginDestroyGranularity;
 
 	// Set 'I'm garbage collecting' flag - might be checked inside UObject::Destroy etc.
 	TGuardValue GuardIsGarbageCollecting(GIsGarbageCollecting, true);
@@ -7259,4 +7266,4 @@ void ProcessAsync(void (*ProcessSync)(void*, FWorkerContext&), void* Processor,
 int32 GetNumCollectReferenceWorkers()
 {
 	return FMath::Clamp(FTaskGraphInterface::Get().GetNumWorkerThreads(), 1, UE::GC::MaxWorkers);
-}
+}
\ No newline at end of file
