Disabling GEnableScriptImportsDebugging by default as it causes standalone programs (BuildStorageTool, CrashReportClient) to hang on startup.

diff --git a/Engine/Source/Runtime/CoreUObject/Private/Serialization/AsyncLoading2.cpp b/Engine/Source/Runtime/CoreUObject/Private/Serialization/AsyncLoading2.cpp
--- a/Engine/Source/Runtime/CoreUObject/Private/Serialization/AsyncLoading2.cpp
+++ b/Engine/Source/Runtime/CoreUObject/Private/Serialization/AsyncLoading2.cpp
@@ -326,7 +326,7 @@ void FindAllRuntimeScriptPackages(FRuntimeScriptPackages& OutPackages)
 #define ALT2_ENABLE_SCRIPT_IMPORTS_DEBUGGING !UE_SHIPPING
 #endif
 #if ALT2_ENABLE_SCRIPT_IMPORTS_DEBUGGING
-static bool GEnableScriptImportsDebugging = true;
+static bool GEnableScriptImportsDebugging = false;
 static FAutoConsoleVariableRef CEnableScriptImportsDebugging(
 	TEXT("s.EnableScriptImportsDebugging"),
 	GEnableScriptImportsDebugging,
