[AsyncLoading2]

- Removes unnecessary GetAnnotation locking call to fetch the Linker ObjectIndex
- Only checks to clear the linker for objects when using linker load
- Clarifies some comments

#rnx

diff --git a/Engine/Source/Runtime/CoreUObject/Private/Serialization/AsyncLoading2.cpp b/Engine/Source/Runtime/CoreUObject/Private/Serialization/AsyncLoading2.cpp
--- a/Engine/Source/Runtime/CoreUObject/Private/Serialization/AsyncLoading2.cpp
+++ b/Engine/Source/Runtime/CoreUObject/Private/Serialization/AsyncLoading2.cpp
@@ -4811,11 +4811,9 @@ private:
 #endif
 		check(DataPtr - Data.MemoryBuffer0 == MemoryBufferSize);
 
-#if WITH_EDITOR
-		const bool bCanImportPackagesWithIdsOnly = false;
-#elif ALT2_ENABLE_LINKERLOAD_SUPPORT
-		const bool bIsZenPackage = !AsyncPackage->LinkerLoadState.IsSet();
-		bool bCanImportPackagesWithIdsOnly = bIsZenPackage;
+#if ALT2_ENABLE_LINKERLOAD_SUPPORT
+		// Only Zen packages can be loaded by package id only
+		const bool bCanImportPackagesWithIdsOnly = !AsyncPackage->LinkerLoadState.IsSet();
 #else
 		const bool bCanImportPackagesWithIdsOnly = true;
 #endif
@@ -5344,11 +5342,10 @@ FAsyncPackage2* FAsyncLoadingThread2::FindOrInsertPackage(FAsyncLoadingThreadSta
 
 	// PostLoadGroup are to protect some race conditions happening in runtime when ALT is active but
 	// it is causing crashes and soft-locks in editor when FlushAsyncLoading is called on specific requests
-	// that are in the same postload group as their outer, preventing fine grained flush. So we disable this
-	// in editor for the time being since ALT is not yet activated by default and we'll sort out the race conditions instead.
-#if ALT2_ENABLE_LINKERLOAD_SUPPORT
+	// that are in the same postload group as their outer, preventing fine grained flush. In runtime it can cause a runtime regression
+	// if you don't need protection from postload races. So we disable by default for the time being until we sort out the above concerns.
 	constexpr bool bIsPostLoadGroupFeatureActive = false;
-#else
+#if 0
 	// Prevents activating postload groups during boot because it was causing deadlock into the 
 	// highly recursive InitDefaultMaterials function on some platforms.
 	// Since postload groups are there to protect against race conditions between postloads and serialize,
@@ -6397,7 +6394,9 @@ void FAsyncPackage2::ImportPackagesRecursiveInner(FAsyncLoadingThreadState2& Thr
 		if (!ImportedPackageNameToLoad.IsNone() && FAsyncLoadingThreadSettings::Get().bLooseFileLoadingEnabled)
 		{
 			FPackagePath::TryFromPackageName(ImportedPackageNameToLoad, ImportedPackagePath);
+#if ALT2_ENABLE_LINKERLOAD_SUPPORT
 			PackageLoader = LinkerLoadState.IsSet() ? EPackageLoader::LinkerLoad : PackageLoader;
+#endif
 
 			if (AsyncLoadingThread.TryGetExistingLoadedPackagePath(ImportedPackageId, ImportedPackagePath, ImportedPackageLoader))
 			{
@@ -10340,7 +10339,7 @@ void FAsyncLoadingThread2::CollectUnreachableObjects(
 	{
 		TRACE_CPUPROFILER_EVENT_SCOPE(CollectUnreachableObjects);
 		ParallelFor(TEXT("CollectUnreachableObjectsTask"), UnreachableObjectItems.Num(), 2048,
-			[&UnreachableObjectItems, &OutUnreachableObjects](int32 Index)
+			[&UnreachableObjectItems, &OutUnreachableObjects, bUsingLinkerLoad = FAsyncLoadingThreadSettings::Get().bLooseFileLoadingEnabled](int32 Index)
 		{
 			UObject* Object = static_cast<UObject*>(UnreachableObjectItems[Index]->GetObject());
 
@@ -10358,15 +10357,15 @@ void FAsyncLoadingThread2::CollectUnreachableObjects(
 				}
 			}
 
-#if ALT2_ENABLE_LINKERLOAD_SUPPORT
-			// Clear garbage objects from linker export tables
-			// Normally done from UObject::BeginDestroy but we need to do it already here
-			if (FLinkerLoad* ObjectLinker = Object->GetLinker())
+			if (bUsingLinkerLoad)
 			{
-				const int32 CachedLinkerIndex = Object->GetLinkerIndex();
-				Object->SetLinker(nullptr, INDEX_NONE);
+				// Clear garbage objects from linker export tables
+				// Normally done from UObject::BeginDestroy but we need to do it already here
+				if (FLinkerLoad* ObjectLinker = Object->GetLinker())
+				{
+					Object->SetLinker(nullptr, INDEX_NONE);
+				}
 			}
-#endif // ALT2_ENABLE_LINKERLOAD_SUPPORT
 		});
 	}
 
