Changed how overridden properties are stored. While implementing the override all properties, it triggered a case where a node I was working on was being reallocated in the container.

#rnx

diff --git a/Engine/Source/Runtime/CoreUObject/Private/UObject/OverriddenPropertySet.cpp b/Engine/Source/Runtime/CoreUObject/Private/UObject/OverriddenPropertySet.cpp
--- a/Engine/Source/Runtime/CoreUObject/Private/UObject/OverriddenPropertySet.cpp
+++ b/Engine/Source/Runtime/CoreUObject/Private/UObject/OverriddenPropertySet.cpp
@@ -48,7 +48,6 @@ namespace Private
 		TEXT("Enables serialization of override state into/from T3D"),
 		FConsoleVariableDelegate::CreateLambda([](IConsoleVariable* Variable)
 		{
-			FOverridableManager& Manager = FOverridableManager::Get();
 			FOverridableSerializationLogicInternalAdapter::SetCapability(FOverridableSerializationLogic::ECapabilities::T3DSerialization, bEnableT3D);
 		}));
 
@@ -59,7 +58,6 @@ namespace Private
 		TEXT("Enables shadow serialization of subobject"),
 		FConsoleVariableDelegate::CreateLambda([](IConsoleVariable* Variable)
 		{
-			FOverridableManager& Manager = FOverridableManager::Get();
 			FOverridableSerializationLogicInternalAdapter::SetCapability(FOverridableSerializationLogic::ECapabilities::SubObjectsShadowSerialization, bEnableSubObjectsShadowSerialization);
 		}));
 
@@ -469,22 +467,49 @@ void FOverriddenPropertyNodeID::HandleDeadObjectReferences(const TSet<UObject*>&
 	}
 }
 
-void FOverriddenPropertySet::RestoreOverriddenState(const FOverriddenPropertySet& FromOverriddenProperties)
+//----------------------------------------------------------------------//
+// FOverriddenPropertyNode
+//----------------------------------------------------------------------//
+bool FOverriddenPropertyNode::Serialize(FArchive& Ar)
 {
-	bWasAdded = FromOverriddenProperties.bWasAdded;
+	FOverriddenPropertyNodeID::StaticStruct()->SerializeItem(Ar, &NodeID, nullptr);
+	Ar << Operation;
+	Ar << SubPropertyNodes;
+	return true;
+}
+
+void FOverriddenPropertyNode::HandleObjectsReInstantiated(const TMap<UObject*, UObject*>& Map)
+{
+	NodeID.HandleObjectsReInstantiated(Map);
+	for (FOverriddenPropertyNode& SubNode : SubPropertyNodes)
+	{
+		SubNode.HandleObjectsReInstantiated(Map);
+	}
+}
+
+void FOverriddenPropertyNode::HandleDeadObjectReferences(const TSet<UObject*>& ActiveInstances, const TSet<UObject*>& TemplateInstances)
+{
+	NodeID.HandleDeadObjectReferences(ActiveInstances, TemplateInstances);
+	for (FOverriddenPropertyNode& SubNode : SubPropertyNodes)
+	{
+		SubNode.HandleDeadObjectReferences(ActiveInstances, TemplateInstances);
+	}
 }
 
 //----------------------------------------------------------------------//
 // FOverriddenPropertySet
 //----------------------------------------------------------------------//
+
+void FOverriddenPropertySet::RestoreOverriddenState(const FOverriddenPropertySet& FromOverriddenProperties)
+{
+	bWasAdded = FromOverriddenProperties.bWasAdded;
+}
+
 FOverriddenPropertyNode& FOverriddenPropertySet::FindOrAddNode(FOverriddenPropertyNode& ParentNode, FOverriddenPropertyNodeID NodeID)
 {
-	FOverriddenPropertyNodeID& SubNodeID = ParentNode.SubPropertyNodeKeys.FindOrAdd(NodeID, FOverriddenPropertyNodeID());
-	if (SubNodeID.IsValid())
+	if (FOverriddenPropertyNode* SubNode = ParentNode.SubPropertyNodes.FindByKey(NodeID))
 	{
-		FOverriddenPropertyNode* FoundNode = OverriddenPropertyNodes.Find(SubNodeID);
-		checkf(FoundNode, TEXT("Expecting a node"));
-		return *FoundNode;
+		return *SubNode;
 	}
 
 	// We can safely assume that the parent node is at least modified from now on
@@ -494,9 +519,7 @@ FOverriddenPropertyNode& FOverriddenPropertySet::FindOrAddNode(FOverriddenProper
 	}
 
 	// Not found add the node
-	SubNodeID = FOverriddenPropertyNodeID(ParentNode.NodeID, NodeID);
-	const FSetElementId NewID = OverriddenPropertyNodes.Emplace(SubNodeID);
-	return OverriddenPropertyNodes.Get(NewID);
+	return ParentNode.SubPropertyNodes.Emplace_GetRef(NodeID);
 }
 
 EOverriddenPropertyOperation FOverriddenPropertySet::GetOverriddenPropertyOperation(const FOverriddenPropertyNode* ParentPropertyNode, FPropertyVisitorPath::Iterator PropertyIterator, bool* bOutInheritedOperation, const void* Data) const
@@ -516,10 +539,10 @@ EOverriddenPropertyOperation FOverriddenPropertySet::GetOverriddenPropertyOperat
 		const FOverriddenPropertyNode* CurrentOverriddenPropertyNode = nullptr;
 		if (OverriddenPropertyNode)
 		{
-			if (const FOverriddenPropertyNodeID* CurrentPropKey = OverriddenPropertyNode->SubPropertyNodeKeys.Find(FOverriddenPropertyNodeID(CurrentProperty)))
+			FOverriddenPropertyNodeID NodeID(CurrentProperty);
+			if (const FOverriddenPropertyNode* FoundNode = OverriddenPropertyNode->SubPropertyNodes.FindByKey(NodeID))
 			{
-				CurrentOverriddenPropertyNode = OverriddenPropertyNodes.Find(*CurrentPropKey);
-				checkf(CurrentOverriddenPropertyNode, TEXT("Expecting a node"));
+				CurrentOverriddenPropertyNode = FoundNode;
 			}
 		}
 
@@ -561,10 +584,8 @@ EOverriddenPropertyOperation FOverriddenPropertySet::GetOverriddenPropertyOperat
 						{
 							// Caller wants to know about any override state on the reference of the subobject itself
 							const FOverriddenPropertyNodeID  SubObjectID(SubObject);
-							if (const FOverriddenPropertyNodeID* CurrentPropKey = CurrentOverriddenPropertyNode->SubPropertyNodeKeys.Find(SubObjectID))
+							if (const FOverriddenPropertyNode* SubObjectOverriddenPropertyNode = CurrentOverriddenPropertyNode->SubPropertyNodes.FindByKey(SubObjectID))
 							{
-								const FOverriddenPropertyNode* SubObjectOverriddenPropertyNode = OverriddenPropertyNodes.Find(*CurrentPropKey);
-								checkf(SubObjectOverriddenPropertyNode, TEXT("Expecting a node"));
 								if (bOutInheritedOperation)
 								{
 									*bOutInheritedOperation = false;
@@ -605,10 +626,8 @@ EOverriddenPropertyOperation FOverriddenPropertySet::GetOverriddenPropertyOperat
 					checkf(MapProperty->KeyProp, TEXT("Expecting a key type for Maps"));
 					FOverriddenPropertyNodeID OverriddenKeyID = FOverriddenPropertyNodeID::FromMapKey(MapProperty->KeyProp, MapHelper.GetKeyPtr(InternalMapIndex));
 
-					if (const FOverriddenPropertyNodeID* CurrentPropKey = CurrentOverriddenPropertyNode->SubPropertyNodeKeys.Find(OverriddenKeyID))
+					if (const FOverriddenPropertyNode* SubObjectOverriddenPropertyNode = CurrentOverriddenPropertyNode->SubPropertyNodes.FindByKey(OverriddenKeyID))
 					{
-						const FOverriddenPropertyNode* SubObjectOverriddenPropertyNode = OverriddenPropertyNodes.Find(*CurrentPropKey);
-						checkf(SubObjectOverriddenPropertyNode, TEXT("Expecting a node"));
 						if (bOutInheritedOperation)
 						{
 							*bOutInheritedOperation = false;
@@ -655,8 +674,8 @@ bool FOverriddenPropertySet::ClearOverriddenProperty(FOverriddenPropertyNode& Pa
 	const void* SubValuePtr = Data;
 	FOverriddenPropertyNode* OverriddenPropertyNode = &ParentPropertyNode;
 	int32 ArrayIndex = INDEX_NONE;
-	TArray<FOverriddenPropertyNodeID> TraversedNodes;
-	TraversedNodes.Push(OverriddenPropertyNode->NodeID);
+	TArray<FOverriddenPropertyNode*> TraversedNodes;
+	TraversedNodes.Push(OverriddenPropertyNode);
 	while (PropertyIterator && (!OverriddenPropertyNode || OverriddenPropertyNode->Operation != EOverriddenPropertyOperation::Replace))
 	{
 		ArrayIndex = INDEX_NONE;
@@ -667,11 +686,11 @@ bool FOverriddenPropertySet::ClearOverriddenProperty(FOverriddenPropertyNode& Pa
 		FOverriddenPropertyNode* CurrentOverriddenPropertyNode = nullptr;
 		if (OverriddenPropertyNode)
 		{
-			if (const FOverriddenPropertyNodeID* CurrentPropKey = OverriddenPropertyNode->SubPropertyNodeKeys.Find(FOverriddenPropertyNodeID(CurrentProperty)))
+			FOverriddenPropertyNodeID NodeID(CurrentProperty);
+			if (FOverriddenPropertyNode* FoundPropertyNode = OverriddenPropertyNode->SubPropertyNodes.FindByKey(NodeID))
 			{
-				CurrentOverriddenPropertyNode = OverriddenPropertyNodes.Find(*CurrentPropKey);
-				checkf(CurrentOverriddenPropertyNode, TEXT("Expecting a node"));
-				TraversedNodes.Push(CurrentOverriddenPropertyNode->NodeID);
+				CurrentOverriddenPropertyNode = FoundPropertyNode;
+				TraversedNodes.Push(CurrentOverriddenPropertyNode);
 			}
 		}
 
@@ -724,13 +743,11 @@ bool FOverriddenPropertySet::ClearOverriddenProperty(FOverriddenPropertyNode& Pa
 						{
 							return OverridableManager.ClearOverriddenProperty(SubObject, NextPropertyIterator);
 						}
-						else if (CurrentOverriddenPropertyNode)
+						if (CurrentOverriddenPropertyNode)
 						{
 							const FOverriddenPropertyNodeID  SubObjectID(SubObject);
-							FOverriddenPropertyNodeID CurrentPropKey;
-							if (CurrentOverriddenPropertyNode->SubPropertyNodeKeys.RemoveAndCopyValue(SubObjectID, CurrentPropKey))
+							if (CurrentOverriddenPropertyNode->SubPropertyNodes.Remove(SubObjectID))
 							{
-								verifyf(OverriddenPropertyNodes.Remove(CurrentPropKey), TEXT("Expecting a node to be removed"));
 								OverridableManager.ClearInstancedSubObjectOverrides(Owner, SubObject);
 								return true;
 							}
@@ -786,10 +803,8 @@ bool FOverriddenPropertySet::ClearOverriddenProperty(FOverriddenPropertyNode& Pa
 				FOverriddenPropertyNodeID OverriddenKeyID = FOverriddenPropertyNodeID::FromMapKey(MapProperty->KeyProp, MapHelper.GetKeyPtr(InternalMapIndex));
 
 				FOverriddenPropertyNodeID CurrentPropKey;
-				if (CurrentOverriddenPropertyNode->SubPropertyNodeKeys.RemoveAndCopyValue(OverriddenKeyID, CurrentPropKey))
+				if (CurrentOverriddenPropertyNode->SubPropertyNodes.Remove(OverriddenKeyID))
 				{
-					verifyf(OverriddenPropertyNodes.Remove(CurrentPropKey), TEXT("Expecting a node to be removed"));
-
 					if (UObject* ValueSubObject = TryGetInstancedSubObjectValue(ValueObjectProperty, MapHelper.GetValuePtr(InternalMapIndex)))
 					{
 						// In the case of a instanced subobject, clear all the overrides on the subobject as well
@@ -809,27 +824,22 @@ bool FOverriddenPropertySet::ClearOverriddenProperty(FOverriddenPropertyNode& Pa
 	{
 		// Go through each traversed property in reversed order to do cleanup
 		// We need to continue the cleanup until there is more overrides than just the one we are removing
-		FOverriddenPropertyNodeID LastCleanedNode;
-		while (FOverriddenPropertyNode* CurrentNode = !TraversedNodes.IsEmpty() ? OverriddenPropertyNodes.Find(TraversedNodes.Top()) : nullptr)
+		FOverriddenPropertyNode* LastCleanedNode = nullptr;
+		while (FOverriddenPropertyNode* CurrentNode = !TraversedNodes.IsEmpty() ? TraversedNodes.Pop() : nullptr)
 		{
-			TraversedNodes.Pop();
-
-			if (LastCleanedNode.IsValid())
+			if (LastCleanedNode)
 			{
-				const FOverriddenPropertyNodeID* NodeToRemove = CurrentNode->SubPropertyNodeKeys.FindKey(LastCleanedNode);
-				checkf(NodeToRemove, TEXT("Expecting to find the last cleaned node"));
-
 				// In the case there are other overrides, just cleanup that node and stop.
-				if (CurrentNode->SubPropertyNodeKeys.Num() > 1)
+				if (CurrentNode->SubPropertyNodes.Num() > 1)
 				{
-					CurrentNode->SubPropertyNodeKeys.Remove(*NodeToRemove);
-					verifyf(OverriddenPropertyNodes.Remove(LastCleanedNode), TEXT("Expecting the node to be removed"));
+					verifyf(CurrentNode->SubPropertyNodes.Remove(LastCleanedNode->NodeID), TEXT("Expecting the node to always be removed"));
+					LastCleanedNode = nullptr;
 					break;
 				}
 			}
 
-			RemoveOverriddenSubProperties(*CurrentNode);
-			LastCleanedNode = CurrentNode->NodeID;
+			CurrentNode->Reset();
+			LastCleanedNode = CurrentNode;
 		}
 	};
 
@@ -860,7 +870,7 @@ void FOverriddenPropertySet::NotifyPropertyChange(FOverriddenPropertyNode* Paren
 	{
 		if (ParentPropertyNode && Notification == EPropertyNotificationType::PostEdit)
 		{
-			this->ClearOverriddenProperty(*ParentPropertyNode, PropertyIterator, Data);
+			ClearOverriddenProperty(*ParentPropertyNode, PropertyIterator, Data);
 		}
 		return;
 	}
@@ -871,15 +881,13 @@ void FOverriddenPropertySet::NotifyPropertyChange(FOverriddenPropertyNode* Paren
 		if (ParentPropertyNode && Notification == EPropertyNotificationType::PostEdit)
 		{
 			// Sub-property overrides are not needed from now on, so clear them
-			RemoveOverriddenSubProperties(*ParentPropertyNode);
+			ParentPropertyNode->Reset();
 
 			// Replacing this entire property
 			ParentPropertyNode->Operation = EOverriddenPropertyOperation::Replace;
 
 			// If we are overriding the root node, need to propagate the overrides to all instanced sub object
-			const FOverriddenPropertyNode* RootNode = OverriddenPropertyNodes.Find(RootNodeID);
-			checkf(RootNode, TEXT("Expecting to always have a "));
-			if (RootNode == ParentPropertyNode)
+			if (&RootNode == ParentPropertyNode)
 			{
 				OverridableManager.PropagateOverrideToInstancedSubObjects(Owner);
 			}
@@ -906,17 +914,14 @@ void FOverriddenPropertySet::NotifyPropertyChange(FOverriddenPropertyNode* Paren
 			return;
 		}
 
-		if (SubPropertyNode && SubPropertyNode->SubPropertyNodeKeys.IsEmpty() &&
+		if (SubPropertyNode && SubPropertyNode->SubPropertyNodes.IsEmpty() &&
 			(bNeedsCleanup || 
 			 SubPropertyNode->Operation == EOverriddenPropertyOperation::None || 
 			 SubPropertyNode->Operation == EOverriddenPropertyOperation::Modified))
 		{
-			FOverriddenPropertyNodeID RemovedNodeID;
-			if (ParentPropertyNode->SubPropertyNodeKeys.RemoveAndCopyValue(FOverriddenPropertyNodeID(Property), RemovedNodeID))
-			{
-				verifyf(OverriddenPropertyNodes.Remove(RemovedNodeID), TEXT("Expecting the node to be removed"));
-			}
-			if (ParentPropertyNode->Operation == EOverriddenPropertyOperation::Modified && ParentPropertyNode->SubPropertyNodeKeys.IsEmpty())
+			FOverriddenPropertyNodeID RemoveNodeID(Property);
+			ParentPropertyNode->SubPropertyNodes.Remove(RemoveNodeID);
+			if (ParentPropertyNode->Operation == EOverriddenPropertyOperation::Modified && ParentPropertyNode->SubPropertyNodes.IsEmpty())
 			{
 				ParentPropertyNode->Operation = EOverriddenPropertyOperation::None;
 			}
@@ -1058,10 +1063,8 @@ void FOverriddenPropertySet::NotifyPropertyChange(FOverriddenPropertyNode* Paren
 								if (RemovedSubObjectNode.Operation == EOverriddenPropertyOperation::Add)
 								{
 									// An add then a remove becomes no opt
-									FOverriddenPropertyNodeID RemovedNodeID;
-									if (SubPropertyNode->SubPropertyNodeKeys.RemoveAndCopyValue(RemovedSubObjectID, RemovedNodeID))
+									if (SubPropertyNode->SubPropertyNodes.Remove(RemovedSubObjectID))
 									{
-										verifyf(OverriddenPropertyNodes.Remove(RemovedNodeID), TEXT("Expecting the node to be removed"));
 										bNeedsCleanup = true;
 									}
 								}
@@ -1074,20 +1077,16 @@ void FOverriddenPropertySet::NotifyPropertyChange(FOverriddenPropertyNode* Paren
 							{
 								// Figure out if it is a remove of a previously added element
 								const FOverriddenPropertyNodeID RemovedSubObjectID (RemovedSubObject);
-								if (const FOverriddenPropertyNodeID* AddedSubObjectID = SubPropertyNode->SubPropertyNodeKeys.Find(RemovedSubObjectID))
+								if (const FOverriddenPropertyNode* AddedSubObjectNode = SubPropertyNode->SubPropertyNodes.FindByKey(RemovedSubObjectID))
 								{
-									FOverriddenPropertyNode* AddedSubObjectNode = OverriddenPropertyNodes.Find(*AddedSubObjectID);
-									checkf(AddedSubObjectNode, TEXT("Expecting a node"));
 									if (AddedSubObjectNode->Operation != EOverriddenPropertyOperation::Add)
 									{
 										UE_LOG(LogOverridableObject, Warning, TEXT("This removed object:%s(0x%p) was not tracked as an add in the overridden properties"), *GetNameSafe(RemovedSubObject), RemovedSubObject);
 									}
 
 									// An add then a remove becomes no opt
-									FOverriddenPropertyNodeID RemovedNodeID;
-									if (SubPropertyNode->SubPropertyNodeKeys.RemoveAndCopyValue(RemovedSubObjectID, RemovedNodeID))
+									if (SubPropertyNode->SubPropertyNodes.Remove(RemovedSubObjectID))
 									{
-										verifyf(OverriddenPropertyNodes.Remove(RemovedNodeID), TEXT("Expecting the node to be removed"));
 										bNeedsCleanup = true;
 									}
 								}
@@ -1282,10 +1281,8 @@ void FOverriddenPropertySet::NotifyPropertyChange(FOverriddenPropertyNode* Paren
 					if (RemovedKeyNode.Operation == EOverriddenPropertyOperation::Add)
 					{
 						// @Todo support remove/add/remove
-						FOverriddenPropertyNodeID RemovedNodeID;
-						if (SubPropertyNode->SubPropertyNodeKeys.RemoveAndCopyValue(RemovedKeyID, RemovedNodeID))
+						if (SubPropertyNode->SubPropertyNodes.Remove(RemovedKeyID))
 						{
-							verifyf(OverriddenPropertyNodes.Remove(RemovedNodeID), TEXT("Expecting the node to be removed"));
 							bNeedsCleanup = true;
 						}
 					}
@@ -1453,19 +1450,6 @@ void FOverriddenPropertySet::NotifyPropertyChange(FOverriddenPropertyNode* Paren
 	}
 }
 
-void FOverriddenPropertySet::RemoveOverriddenSubProperties(FOverriddenPropertyNode& PropertyNode)
-{
-	for (const auto& Pair : PropertyNode.SubPropertyNodeKeys)
-	{
-		FOverriddenPropertyNode* RemovedPropertyNode = OverriddenPropertyNodes.Find(Pair.Value);
-		checkf(RemovedPropertyNode, TEXT("Expecting a node"));
-		RemoveOverriddenSubProperties(*RemovedPropertyNode);
-		verifyf(OverriddenPropertyNodes.Remove(Pair.Value), TEXT("Expecting the node to be removed"));
-	}
-	PropertyNode.Operation = EOverriddenPropertyOperation::None;
-	PropertyNode.SubPropertyNodeKeys.Empty();
-}
-
 UObject* FOverriddenPropertySet::TryGetInstancedSubObjectValue(const FObjectPropertyBase* FromProperty, const void* ValuePtr) const
 {
 	// Property can be NULL - in that case there is no value.
@@ -1484,7 +1468,7 @@ UObject* FOverriddenPropertySet::TryGetInstancedSubObjectValue(const FObjectProp
 		RedirectMethod = UE::ResolveInstanceDataObject;
 		ExpectedOuter = Instance;
 	}
-#endif
+#endif // WITH_EDITORONLY_DATA
 
 	if (FromProperty->HasAnyPropertyFlags(CPF_PersistentInstance)
 		|| (FromProperty->IsA<FObjectProperty>() && SubObject && SubObject->IsIn(ExpectedOuter)))
@@ -1497,40 +1481,35 @@ UObject* FOverriddenPropertySet::TryGetInstancedSubObjectValue(const FObjectProp
 
 EOverriddenPropertyOperation FOverriddenPropertySet::GetOverriddenPropertyOperation(FPropertyVisitorPath::Iterator PropertyIterator, bool* bOutInheritedOperation /*= nullptr*/) const
 {
-	return GetOverriddenPropertyOperation(OverriddenPropertyNodes.Find(RootNodeID), PropertyIterator, bOutInheritedOperation, Owner);
+	return GetOverriddenPropertyOperation(&RootNode, PropertyIterator, bOutInheritedOperation, Owner);
 }
 
 bool FOverriddenPropertySet::ClearOverriddenProperty(FPropertyVisitorPath::Iterator PropertyIterator)
 {
-	if (FOverriddenPropertyNode* RootNode = OverriddenPropertyNodes.Find(RootNodeID))
-	{
-		return ClearOverriddenProperty(*RootNode, PropertyIterator, Owner);
-	}
-	return true;
+	return ClearOverriddenProperty(RootNode, PropertyIterator, Owner);
 }
 
 void FOverriddenPropertySet::OverrideProperty(FPropertyVisitorPath::Iterator PropertyIterator, const void* Data)
 {
-	FOverriddenPropertyNode& RootPropertyNode = OverriddenPropertyNodes.FindOrAdd(RootNodeID);
 	bool bNeedsCleanup = false;
-	NotifyPropertyChange(&RootPropertyNode, EPropertyNotificationType::PreEdit, PropertyIterator, EPropertyChangeType::Unspecified, Data, bNeedsCleanup);
-	NotifyPropertyChange(&RootPropertyNode, EPropertyNotificationType::PostEdit, PropertyIterator, EPropertyChangeType::Unspecified, Data, bNeedsCleanup);
+	NotifyPropertyChange(&RootNode, EPropertyNotificationType::PreEdit, PropertyIterator, EPropertyChangeType::Unspecified, Data, bNeedsCleanup);
+	NotifyPropertyChange(&RootNode, EPropertyNotificationType::PostEdit, PropertyIterator, EPropertyChangeType::Unspecified, Data, bNeedsCleanup);
 }
 
 void FOverriddenPropertySet::NotifyPropertyChange(const EPropertyNotificationType Notification, FPropertyVisitorPath::Iterator PropertyIterator, const EPropertyChangeType::Type ChangeType, const void* Data)
 {
 	bool bNeedsCleanup = false;
-	NotifyPropertyChange(&OverriddenPropertyNodes.FindOrAdd(RootNodeID), Notification, PropertyIterator, ChangeType, Data, bNeedsCleanup);
+	NotifyPropertyChange(&RootNode, Notification, PropertyIterator, ChangeType, Data, bNeedsCleanup);
 }
 
 EOverriddenPropertyOperation FOverriddenPropertySet::GetOverriddenPropertyOperation(const FArchiveSerializedPropertyChain* CurrentPropertyChain, FProperty* Property) const
 {
-	return GetOverriddenPropertyOperation(OverriddenPropertyNodes.Find(RootNodeID), CurrentPropertyChain, Property);
+	return GetOverriddenPropertyOperation(&RootNode, CurrentPropertyChain, Property);
 }
 
 FOverriddenPropertyNode* FOverriddenPropertySet::SetOverriddenPropertyOperation(EOverriddenPropertyOperation Operation, const FArchiveSerializedPropertyChain* CurrentPropertyChain, FProperty* Property)
 {
-	return SetOverriddenPropertyOperation(Operation, OverriddenPropertyNodes.FindOrAdd(RootNodeID), CurrentPropertyChain, Property);
+	return SetOverriddenPropertyOperation(Operation, RootNode, CurrentPropertyChain, Property);
 }
 
 FOverriddenPropertyNode* FOverriddenPropertySet::RestoreOverriddenPropertyOperation(EOverriddenPropertyOperation Operation, const FArchiveSerializedPropertyChain* CurrentPropertyChain, FProperty* Property)
@@ -1554,16 +1533,16 @@ FOverriddenPropertyNode* FOverriddenPropertySet::RestoreOverriddenPropertyOperat
 
 const FOverriddenPropertyNode* FOverriddenPropertySet::GetOverriddenPropertyNode(const FArchiveSerializedPropertyChain* CurrentPropertyChain) const
 {
-	if (const FOverriddenPropertyNode* RootNode = OverriddenPropertyNodes.Find(RootNodeID))
-	{
-		return GetOverriddenPropertyNode(*RootNode, CurrentPropertyChain);
-	}
-	return nullptr;
+	return GetOverriddenPropertyNode(RootNode, CurrentPropertyChain);
 }
 
 const FOverriddenPropertyNode* FOverriddenPropertySet::GetOverriddenPropertyNode(FOverriddenPropertyNodeID NodeID) const
 {
-	return OverriddenPropertyNodes.Find(NodeID);
+	if (const FOverriddenPropertyNode* FoundNode = RootNode.SubPropertyNodes.FindByKey(NodeID))
+	{
+		return FoundNode;
+	}
+	return nullptr;
 }
 
 EOverriddenPropertyOperation FOverriddenPropertySet::GetOverriddenPropertyOperation(const FOverriddenPropertyNode* ParentPropertyNode, const FArchiveSerializedPropertyChain* CurrentPropertyChain, FProperty* Property) const
@@ -1590,9 +1569,10 @@ EOverriddenPropertyOperation FOverriddenPropertySet::GetOverriddenPropertyOperat
 		const FProperty* CurrentProperty = (*PropertyIterator);
 		if (OverriddenPropertyNode)
 		{
-			if (const FOverriddenPropertyNodeID* CurrentPropKey = OverriddenPropertyNode->SubPropertyNodeKeys.Find(FOverriddenPropertyNodeID(CurrentProperty)))
+			FOverriddenPropertyNodeID NodeID(CurrentProperty);
+			if (const FOverriddenPropertyNode* FoundNode = OverriddenPropertyNode->SubPropertyNodes.FindByKey(NodeID))
 			{
-				OverriddenPropertyNode = OverriddenPropertyNodes.Find(*CurrentPropKey);
+				OverriddenPropertyNode = FoundNode;
 				checkf(OverriddenPropertyNode, TEXT("Expecting a node"));
 			}
 			else
@@ -1656,8 +1636,11 @@ FOverriddenPropertyNode* FOverriddenPropertySet::SetOverriddenPropertyOperation(
 
 EOverriddenPropertyOperation FOverriddenPropertySet::GetSubPropertyOperation(FOverriddenPropertyNodeID NodeID) const
 {
-	const FOverriddenPropertyNode* OverriddenPropertyNode = OverriddenPropertyNodes.Find(NodeID);
-	return OverriddenPropertyNode ? OverriddenPropertyNode->Operation : EOverriddenPropertyOperation::None;
+	if (const FOverriddenPropertyNode* SubPropertyNode = GetOverriddenPropertyNode(NodeID))
+	{
+		return SubPropertyNode->Operation;
+	}
+	return EOverriddenPropertyOperation::None;
 }
 
 FOverriddenPropertyNode* FOverriddenPropertySet::SetSubPropertyOperation(EOverriddenPropertyOperation Operation, FOverriddenPropertyNode& Node, FOverriddenPropertyNodeID NodeID)
@@ -1700,7 +1683,7 @@ bool FOverriddenPropertySet::IsCDOOwningProperty(const FProperty& Property) cons
 
 void FOverriddenPropertySet::Reset()
 {
-	OverriddenPropertyNodes.Reset();
+	RootNode.Reset();
 }
 
 void FOverriddenPropertySet::HandleObjectsReInstantiated(const TMap<UObject*, UObject*>& Map)
@@ -1714,28 +1697,12 @@ void FOverriddenPropertySet::HandleObjectsReInstantiated(const TMap<UObject*, UO
 	}
 #endif // WITH_EDITOR
 
-	for (FOverriddenPropertyNode& Node : OverriddenPropertyNodes)
-	{
-		Node.NodeID.HandleObjectsReInstantiated(Map);
-		for (auto& Pair : Node.SubPropertyNodeKeys)
-		{
-			Pair.Key.HandleObjectsReInstantiated(Map);
-			Pair.Value.HandleObjectsReInstantiated(Map);
-		}
-	}
+	RootNode.HandleObjectsReInstantiated(Map);
 }
 
 void FOverriddenPropertySet::HandleDeadObjectReferences(const TSet<UObject*>& ActiveInstances, const TSet<UObject*>& TemplateInstances)
 {
-	for (FOverriddenPropertyNode& Node : OverriddenPropertyNodes)
-	{
-		Node.NodeID.HandleDeadObjectReferences(ActiveInstances, TemplateInstances);
-		for (auto& Pair : Node.SubPropertyNodeKeys)
-		{
-			Pair.Key.HandleDeadObjectReferences(ActiveInstances, TemplateInstances);
-			Pair.Value.HandleDeadObjectReferences(ActiveInstances, TemplateInstances);
-		}
-	}
+	RootNode.HandleDeadObjectReferences(ActiveInstances, TemplateInstances);
 }
 
 const FOverriddenPropertyNode* FOverriddenPropertySet::GetOverriddenPropertyNode(const FOverriddenPropertyNode& ParentPropertyNode, const FArchiveSerializedPropertyChain* CurrentPropertyChain) const
@@ -1750,10 +1717,10 @@ const FOverriddenPropertyNode* FOverriddenPropertySet::GetOverriddenPropertyNode
 	while (PropertyIterator && OverriddenPropertyNode)
 	{
 		const FProperty* CurrentProperty = (*PropertyIterator);
-		if (const FOverriddenPropertyNodeID* CurrentPropKey = OverriddenPropertyNode->SubPropertyNodeKeys.Find(FOverriddenPropertyNodeID(CurrentProperty)))
+		FOverriddenPropertyNodeID NodeID(CurrentProperty);
+		if (const FOverriddenPropertyNode* FoundNode = OverriddenPropertyNode->SubPropertyNodes.FindByKey(NodeID))
 		{
-			OverriddenPropertyNode = OverriddenPropertyNodes.Find(*CurrentPropKey);
-			checkf(OverriddenPropertyNode, TEXT("Expecting a node"));
+			OverriddenPropertyNode = FoundNode;
 		}
 		else
 		{
diff --git a/Engine/Source/Runtime/CoreUObject/Private/UObject/PropertyArray.cpp b/Engine/Source/Runtime/CoreUObject/Private/UObject/PropertyArray.cpp
--- a/Engine/Source/Runtime/CoreUObject/Private/UObject/PropertyArray.cpp
+++ b/Engine/Source/Runtime/CoreUObject/Private/UObject/PropertyArray.cpp
@@ -526,14 +526,13 @@ void FArrayProperty::SerializeItem(FStructuredArchive::FSlot Slot, void* Value,
 
 						if (const FOverriddenPropertyNode* ArrayOverriddenPropertyNode = OverriddenProperties->GetOverriddenPropertyNode(UnderlyingArchive.GetSerializedPropertyChain()))
 						{
-							for (const auto& Pair : ArrayOverriddenPropertyNode->SubPropertyNodeKeys)
+							for (const FOverriddenPropertyNode& SubNode : ArrayOverriddenPropertyNode->SubPropertyNodes)
 							{
-								const EOverriddenPropertyOperation OverrideOp = OverriddenProperties->GetSubPropertyOperation(Pair.Value);
-								switch (OverrideOp)
+								switch (SubNode.Operation)
 								{
 								case EOverriddenPropertyOperation::Remove:
 									{
-										const int32 DefaultIndex = FindObject(Pair.Key, DefaultsArrayHelper);
+										const int32 DefaultIndex = FindObject(SubNode.NodeID, DefaultsArrayHelper);
 										if (DefaultIndex != INDEX_NONE)
 										{
 											// We don't support removing external entities from overriden arrays yet
@@ -544,13 +543,13 @@ void FArrayProperty::SerializeItem(FStructuredArchive::FSlot Slot, void* Value,
 										}
 										else
 										{
-											UE_LOG(LogOverridableObject, VeryVerbose, TEXT("Unable to save deleted item %s"), *Pair.Key.ToDebugString());
+											UE_LOG(LogOverridableObject, VeryVerbose, TEXT("Unable to save deleted item %s"), *SubNode.NodeID.ToDebugString());
 										}
 										break;
 									}
 								case EOverriddenPropertyOperation::Add:
 									{
-										const int32 Index = FindObject(Pair.Key, ArrayHelper);
+										const int32 Index = FindObject(SubNode.NodeID, ArrayHelper);
 										if (Index != INDEX_NONE)
 										{
 											if (!bFilterExternalObjects)
@@ -570,7 +569,7 @@ void FArrayProperty::SerializeItem(FStructuredArchive::FSlot Slot, void* Value,
 										}
 										else
 										{
-											UE_LOG(LogOverridableObject, VeryVerbose, TEXT("Unable to save added item %s"), *Pair.Key.ToDebugString());
+											UE_LOG(LogOverridableObject, VeryVerbose, TEXT("Unable to save added item %s"), *SubNode.NodeID.ToDebugString());
 										}
 										break;
 									}
diff --git a/Engine/Source/Runtime/CoreUObject/Private/UObject/PropertyMap.cpp b/Engine/Source/Runtime/CoreUObject/Private/UObject/PropertyMap.cpp
--- a/Engine/Source/Runtime/CoreUObject/Private/UObject/PropertyMap.cpp
+++ b/Engine/Source/Runtime/CoreUObject/Private/UObject/PropertyMap.cpp
@@ -554,14 +554,13 @@ void FMapProperty::SerializeItem(FStructuredArchive::FSlot Slot, void* Value, co
 							if (const FOverriddenPropertyNode* MapOverriddenPropertyNode = OverriddenProperties->GetOverriddenPropertyNode(UnderlyingArchive.GetSerializedPropertyChain()))
 							{
 								// Figure out the modifications of the map
-								for (const auto& Pair : MapOverriddenPropertyNode->SubPropertyNodeKeys)
+								for (const FOverriddenPropertyNode& SubNode : MapOverriddenPropertyNode->SubPropertyNodes)
 								{
-									const EOverriddenPropertyOperation OverrideOp = OverriddenProperties->GetSubPropertyOperation(Pair.Value);
-									switch (OverrideOp)
+									switch (SubNode.Operation)
 									{
 										case EOverriddenPropertyOperation::Remove:
 										{
-											const int32 InternalIndex = Pair.Key.ToMapInternalIndex(DefaultsMapHelper);
+											const int32 InternalIndex = SubNode.NodeID.ToMapInternalIndex(DefaultsMapHelper);
 											if (InternalIndex != INDEX_NONE)
 											{
 												RemovedIndices.Add(InternalIndex);
@@ -570,7 +569,7 @@ void FMapProperty::SerializeItem(FStructuredArchive::FSlot Slot, void* Value, co
 										}
 										case EOverriddenPropertyOperation::Add:
 										{
-											const int32 InternalIndex = Pair.Key.ToMapInternalIndex(MapHelper);
+											const int32 InternalIndex = SubNode.NodeID.ToMapInternalIndex(MapHelper);
 											if (InternalIndex != INDEX_NONE)
 											{
 												AddedIndices.Add(InternalIndex);
@@ -580,7 +579,7 @@ void FMapProperty::SerializeItem(FStructuredArchive::FSlot Slot, void* Value, co
 										}
 										case EOverriddenPropertyOperation::Modified:
 										{
-											const int32 InternalIndex = Pair.Key.ToMapInternalIndex(MapHelper);
+											const int32 InternalIndex = SubNode.NodeID.ToMapInternalIndex(MapHelper);
 											if (InternalIndex != INDEX_NONE)
 											{
 												ModifiedIndices.Add(InternalIndex);
diff --git a/Engine/Source/Runtime/CoreUObject/Public/UObject/OverriddenPropertySet.h b/Engine/Source/Runtime/CoreUObject/Public/UObject/OverriddenPropertySet.h
--- a/Engine/Source/Runtime/CoreUObject/Public/UObject/OverriddenPropertySet.h
+++ b/Engine/Source/Runtime/CoreUObject/Public/UObject/OverriddenPropertySet.h
@@ -333,21 +333,42 @@ struct FOverriddenPropertyNode
 	UPROPERTY()
 	EOverriddenPropertyOperation Operation = EOverriddenPropertyOperation::None;
 
-	UPROPERTY()
-	TMap<FOverriddenPropertyNodeID, FOverriddenPropertyNodeID> SubPropertyNodeKeys;
+	// Unable to make this a UPROPERTY() because of this UHT error:
+	// 'Struct' recursion via arrays is unsupported for properties.
+	// Instead will be serializing by hand
+	TArray<FOverriddenPropertyNode> SubPropertyNodes;
 
 	bool operator==(const FOverriddenPropertyNode& Other) const
 	{
 		return NodeID == Other.NodeID;
 	}
 
-	friend uint32 GetTypeHash(const FOverriddenPropertyNode& Node)
+	void Reset()
 	{
-		return GetTypeHash(Node.NodeID);
+		Operation = EOverriddenPropertyOperation::None;
+		SubPropertyNodes.Empty();
 	}
-};
 
+	void HandleObjectsReInstantiated(const TMap<UObject*, UObject*>& Map);
+	void HandleDeadObjectReferences(const TSet<UObject*>& ActiveInstances, const TSet<UObject*>& TemplateInstances);
 
+	UE_API bool Serialize(FArchive& Ar);
+	friend FArchive& operator<<(FArchive& Ar, FOverriddenPropertyNode& P)
+	{
+		P.Serialize(Ar);
+		return Ar;
+	}
+};
+
+template<>
+struct TStructOpsTypeTraits<FOverriddenPropertyNode> : public TStructOpsTypeTraitsBase2<FOverriddenPropertyNode>
+{
+	enum 
+	{
+		WithSerializer = true,
+	};
+	static constexpr EPropertyObjectReferenceType WithSerializerObjectReferences = EPropertyObjectReferenceType::None;
+};
 
 /*
  * Property change notification type mapping the Pre/PostEditChange callbacks
@@ -498,8 +519,6 @@ protected:
 	UE_API FOverriddenPropertyNode* SetOverriddenPropertyOperation(EOverriddenPropertyOperation Operation, FOverriddenPropertyNode& ParentPropertyNode, const FArchiveSerializedPropertyChain* CurrentPropertyChain, FProperty* Property);
 	UE_API const FOverriddenPropertyNode* GetOverriddenPropertyNode(const FOverriddenPropertyNode& ParentPropertyNode, const FArchiveSerializedPropertyChain* CurrentPropertyChain) const;
 
-	UE_API void RemoveOverriddenSubProperties(FOverriddenPropertyNode& PropertyIterator);
-
 	UE_API UObject* TryGetInstancedSubObjectValue(const FObjectPropertyBase* FromProperty, const void* ValuePtr) const;
 
 private:
@@ -510,7 +529,7 @@ private:
 	bool bWasAdded = false;
 
 	UPROPERTY()
-	TSet<FOverriddenPropertyNode> OverriddenPropertyNodes;
+	FOverriddenPropertyNode RootNode { RootNodeID };
 
 	static inline FOverriddenPropertyNodeID RootNodeID = FOverriddenPropertyNodeID::RootNodeId();
 
diff --git a/Engine/Source/Runtime/Engine/Private/UnrealExporter.cpp b/Engine/Source/Runtime/Engine/Private/UnrealExporter.cpp
--- a/Engine/Source/Runtime/Engine/Private/UnrealExporter.cpp
+++ b/Engine/Source/Runtime/Engine/Private/UnrealExporter.cpp
@@ -863,14 +863,13 @@ void ExportProperties_Overrides
 						TArray<int32> RemovedIndices;
 						TArray<int32> AddedIndices;
 
-						for (const auto& Pair : ArrayOverriddenPropertyNode->SubPropertyNodeKeys)
+						for (const FOverriddenPropertyNode& SubNode: ArrayOverriddenPropertyNode->SubPropertyNodes)
 						{
-							const EOverriddenPropertyOperation OverrideOp = OverriddenProperties->GetSubPropertyOperation(Pair.Value);
-							switch (OverrideOp)
+							switch (SubNode.Operation)
 							{
 							case EOverriddenPropertyOperation::Remove:
 								{
-									const int32 DefaultIndex = FindObject(Pair.Key, DefaultsArrayHelper);
+									const int32 DefaultIndex = FindObject(SubNode.NodeID, DefaultsArrayHelper);
 									if (DefaultIndex != INDEX_NONE)
 									{
 										RemovedIndices.Add(DefaultIndex);
@@ -879,7 +878,7 @@ void ExportProperties_Overrides
 								}
 							case EOverriddenPropertyOperation::Add:
 								{
-									const int32 Index = FindObject(Pair.Key, ArrayHelper);
+									const int32 Index = FindObject(SubNode.NodeID, ArrayHelper);
 									if (Index != INDEX_NONE)
 									{
 										AddedIndices.Add(Index);
@@ -953,14 +952,13 @@ void ExportProperties_Overrides
 
 
 						// Figure out the modifications of the map
-						for (const auto& Pair : MapOverriddenPropertyNode->SubPropertyNodeKeys)
+						for (const FOverriddenPropertyNode& SubNode : MapOverriddenPropertyNode->SubPropertyNodes)
 						{
-							const EOverriddenPropertyOperation OverrideOp = OverriddenProperties->GetSubPropertyOperation(Pair.Value);
-							switch (OverrideOp)
+							switch (SubNode.Operation)
 							{
 							case EOverriddenPropertyOperation::Remove:
 								{
-									const int32 InternalIndex = Pair.Key.ToMapInternalIndex(DiffMapHelper);
+									const int32 InternalIndex = SubNode.NodeID.ToMapInternalIndex(DiffMapHelper);
 									if (InternalIndex != INDEX_NONE)
 									{
 										RemovedIndices.Add(InternalIndex);
@@ -969,7 +967,7 @@ void ExportProperties_Overrides
 								}
 							case EOverriddenPropertyOperation::Add:
 								{
-									const int32 InternalIndex = Pair.Key.ToMapInternalIndex(MapHelper);
+									const int32 InternalIndex = SubNode.NodeID.ToMapInternalIndex(MapHelper);
 									if (InternalIndex != INDEX_NONE)
 									{
 										AddedIndices.Add(InternalIndex);
@@ -978,7 +976,7 @@ void ExportProperties_Overrides
 								}
 							case EOverriddenPropertyOperation::Modified:
 								{
-									const int32 InternalIndex = Pair.Key.ToMapInternalIndex(MapHelper);
+									const int32 InternalIndex = SubNode.NodeID.ToMapInternalIndex(MapHelper);
 									if (InternalIndex != INDEX_NONE)
 									{
 										ModifiedIndices.Add(InternalIndex);
