Add explicit nullptr constructor/equality operators to TWeakInterfacePtr
Mark functions in TWeakInterfacePtr [[nodiscard]]

diff --git a/Engine/Source/Runtime/AIModule/Private/Navigation/PathFollowingComponent.cpp b/Engine/Source/Runtime/AIModule/Private/Navigation/PathFollowingComponent.cpp
--- a/Engine/Source/Runtime/AIModule/Private/Navigation/PathFollowingComponent.cpp
+++ b/Engine/Source/Runtime/AIModule/Private/Navigation/PathFollowingComponent.cpp
@@ -558,7 +558,7 @@ void UPathFollowingComponent::ResumeMove(FAIRequestID RequestID)
 bool UPathFollowingComponent::ShouldCheckPathOnResume() const
 {
 	bool bCheckPath = true;
-	if (NavMovementInterface != NULL)
+	if (NavMovementInterface != nullptr)
 	{
 		float AgentRadius = 0.0f, AgentHalfHeight = 0.0f;
 		NavMovementInterface->GetSimpleCollisionCylinder(AgentRadius, AgentHalfHeight);
@@ -1102,7 +1102,7 @@ void UPathFollowingComponent::UpdatePathSegment()
 			}
 			else
 			{
-				if ((GetOwner() != NULL) && (NavMovementInterface != NULL))
+				if ((GetOwner() != nullptr) && (NavMovementInterface != nullptr))
 				{
 					UE_VLOG(GetOwner(), LogPathFollowing, Verbose, TEXT("Path blocked, but move segment indices are not valid: start %d, end %d of %d"), MoveSegmentStartIndex, MoveSegmentEndIndex, Path->GetPathPoints().Num());
 				}
diff --git a/Engine/Source/Runtime/CoreUObject/Public/UObject/WeakInterfacePtr.h b/Engine/Source/Runtime/CoreUObject/Public/UObject/WeakInterfacePtr.h
--- a/Engine/Source/Runtime/CoreUObject/Public/UObject/WeakInterfacePtr.h
+++ b/Engine/Source/Runtime/CoreUObject/Public/UObject/WeakInterfacePtr.h
@@ -19,13 +19,21 @@ struct TWeakInterfacePtr
 	using ElementType = T;
 	using UObjectType = TCopyQualifiersFromTo_T<T, UObject>;
 
-	FORCEINLINE TWeakInterfacePtr() = default;
-	FORCEINLINE TWeakInterfacePtr(const TWeakInterfacePtr& Other) = default;
-	FORCEINLINE TWeakInterfacePtr(TWeakInterfacePtr&& Other) = default;
+	[[nodiscard]] FORCEINLINE TWeakInterfacePtr() = default;
+	[[nodiscard]] FORCEINLINE TWeakInterfacePtr(const TWeakInterfacePtr& Other) = default;
+	[[nodiscard]] FORCEINLINE TWeakInterfacePtr(TWeakInterfacePtr&& Other) = default;
 	FORCEINLINE ~TWeakInterfacePtr() = default;
 	FORCEINLINE TWeakInterfacePtr& operator=(const TWeakInterfacePtr& Other) = default;
 	FORCEINLINE TWeakInterfacePtr& operator=(TWeakInterfacePtr&& Other) = default;
 
+	/**
+	* Construct from a null pointer
+	*/
+	[[nodiscard]] FORCEINLINE TWeakInterfacePtr(TYPE_OF_NULLPTR) :
+		ObjectInstance(nullptr)
+	{
+	}
+
 	/**
 	 * Construct from an object pointer
 	 * @param Object The object to create a weak pointer to. This object must implement interface T.
@@ -34,7 +42,7 @@ struct TWeakInterfacePtr
 		typename U
 		UE_REQUIRES(std::is_convertible_v<U, TCopyQualifiersFromTo_T<U, UObject>*>)
 	>
-	TWeakInterfacePtr(U&& Object)
+	[[nodiscard]] TWeakInterfacePtr(U&& Object)
 	{
 		InterfaceInstance = Cast<T>(ImplicitConv<TCopyQualifiersFromTo_T<U, UObject>*>(Object));
 		if (InterfaceInstance != nullptr)
@@ -47,7 +55,7 @@ struct TWeakInterfacePtr
 	 * Construct from an interface pointer
 	 * @param Interface The interface pointer to create a weak pointer to. There must be a UObject behind the interface.
 	 */
-	TWeakInterfacePtr(T* Interface)
+	[[nodiscard]] TWeakInterfacePtr(T* Interface)
 	{
 		ObjectInstance = Cast<UObject>(Interface);
 		if (ObjectInstance != nullptr)
@@ -61,7 +69,7 @@ struct TWeakInterfacePtr
 	 * @param ScriptInterface 	The TScriptInterface to copy from.
 	 * 							No validation is done here; passing an invalid TScriptInterface in will result in an invalid TWeakInterfacePtr.
 	 */
-	TWeakInterfacePtr(const TScriptInterface<T>& ScriptInterface)
+	[[nodiscard]] TWeakInterfacePtr(const TScriptInterface<T>& ScriptInterface)
 	{
 		ObjectInstance = ScriptInterface.GetObject();
 		InterfaceInstance = ScriptInterface.GetInterface();
@@ -79,7 +87,7 @@ struct TWeakInterfacePtr
 	/**
 	 * Test if this points to a live object. Parameters are passed to the underlying TWeakObjectPtr.
 	 */
-	FORCEINLINE bool IsValid(bool bEvenIfPendingKill, bool bThreadsafeTest = false) const
+	[[nodiscard]] FORCEINLINE bool IsValid(bool bEvenIfPendingKill, bool bThreadsafeTest = false) const
 	{
 		return InterfaceInstance != nullptr && ObjectInstance.IsValid(bEvenIfPendingKill, bThreadsafeTest);
 	}
@@ -87,7 +95,7 @@ struct TWeakInterfacePtr
 	/**
 	 * Test if this points to a live object. Calls the underlying TWeakObjectPtr's parameterless IsValid method.
 	 */
-	FORCEINLINE bool IsValid() const
+	[[nodiscard]] FORCEINLINE bool IsValid() const
 	{
 		return InterfaceInstance != nullptr && ObjectInstance.IsValid();
 	}
@@ -95,7 +103,7 @@ struct TWeakInterfacePtr
 	/**
 	 * Test if this pointer is stale. Parameters are passed to the underlying TWeakObjectPtr.
 	 */
-	FORCEINLINE bool IsStale(bool bEvenIfPendingKill = false, bool bThreadsafeTest = false) const
+	[[nodiscard]] FORCEINLINE bool IsStale(bool bEvenIfPendingKill = false, bool bThreadsafeTest = false) const
 	{
 		return InterfaceInstance != nullptr && ObjectInstance.IsStale(bEvenIfPendingKill, bThreadsafeTest);
 	}
@@ -103,7 +111,7 @@ struct TWeakInterfacePtr
 	/**
 	 * Dereference the weak pointer into an interface pointer.
 	 */
-	FORCEINLINE T* Get() const
+	[[nodiscard]] FORCEINLINE T* Get() const
 	{
 		return IsValid() ? InterfaceInstance : nullptr;
 	}
@@ -111,7 +119,7 @@ struct TWeakInterfacePtr
 	/**
 	 * Dereference the weak pointer into a UObject pointer.
 	 */
-	FORCEINLINE UObjectType* GetObject() const
+	[[nodiscard]] FORCEINLINE UObjectType* GetObject() const
 	{
 		return ObjectInstance.Get();
 	}
@@ -119,7 +127,7 @@ struct TWeakInterfacePtr
 	/**
 	 * Dereference the weak pointer.
 	 */
-	FORCEINLINE T& operator*() const
+	[[nodiscard]] FORCEINLINE T& operator*() const
 	{
 		check(IsValid());
 		return *InterfaceInstance;
@@ -128,7 +136,7 @@ struct TWeakInterfacePtr
 	/**
 	 * Dereference the weak pointer.
 	 */
-	FORCEINLINE T* operator->() const
+	[[nodiscard]] FORCEINLINE T* operator->() const
 	{
 		check(IsValid());
 		return InterfaceInstance;
@@ -153,17 +161,27 @@ struct TWeakInterfacePtr
 		return *this;
 	}
 
-	FORCEINLINE bool operator==(const TWeakInterfacePtr<T>& Other) const
+	[[nodiscard]] FORCEINLINE bool operator==(const TWeakInterfacePtr<T>& Other) const
 	{
 		return InterfaceInstance == Other.InterfaceInstance;
 	}
 
-	FORCEINLINE bool operator!=(const TWeakInterfacePtr<T>& Other) const
+	[[nodiscard]] FORCEINLINE bool operator!=(const TWeakInterfacePtr<T>& Other) const
 	{
 		return InterfaceInstance != Other.InterfaceInstance;
 	}
 
-	FORCEINLINE TScriptInterface<T> ToScriptInterface() const
+	[[nodiscard]] FORCENOINLINE bool operator==(TYPE_OF_NULLPTR) const
+	{
+		return !IsValid();
+	}
+
+	[[nodiscard]] FORCEINLINE bool operator!=(TYPE_OF_NULLPTR) const
+	{
+		return !(*this == nullptr);
+	}
+
+	[[nodiscard]] FORCEINLINE TScriptInterface<T> ToScriptInterface() const
 	{
 		UObjectType* Object = ObjectInstance.Get();
 		if (Object)
@@ -174,7 +192,7 @@ struct TWeakInterfacePtr
 		return TScriptInterface<T>();
 	}
 
-	FORCEINLINE TWeakObjectPtr<UObjectType> GetWeakObjectPtr() const
+	[[nodiscard]] FORCEINLINE TWeakObjectPtr<UObjectType> GetWeakObjectPtr() const
 	{
 		return ObjectInstance;
 	}
