[Feature] Enabling Impersonation for IDOs. This will save IDOs in place of their instance so that unknown properties get saved to disk.

diff --git a/Engine/Source/Runtime/CoreUObject/Private/UObject/InstanceDataObjectUtils.cpp b/Engine/Source/Runtime/CoreUObject/Private/UObject/InstanceDataObjectUtils.cpp
--- a/Engine/Source/Runtime/CoreUObject/Private/UObject/InstanceDataObjectUtils.cpp
+++ b/Engine/Source/Runtime/CoreUObject/Private/UObject/InstanceDataObjectUtils.cpp
@@ -321,7 +321,7 @@ namespace UE
 	);
 
 	// TODO: re-enable this ASAP. This disables most IDO features but disabling was necessary to unblock those experiencing IDO bugs
-	bool bEnableIDOImpersonationOnSave = false;
+	bool bEnableIDOImpersonationOnSave = true;
 	FAutoConsoleVariableRef EnableIDOImpersonationOnSaveCVar(
 		TEXT("IDO.Impersonation.EnableOnSave"),
 		bEnableIDOImpersonationOnSave,
