[Package]
- Make package flags thread-safe

diff --git a/Engine/Source/Programs/AutoRTFMTests/Private/UObjectTests.cpp b/Engine/Source/Programs/AutoRTFMTests/Private/UObjectTests.cpp
--- a/Engine/Source/Programs/AutoRTFMTests/Private/UObjectTests.cpp
+++ b/Engine/Source/Programs/AutoRTFMTests/Private/UObjectTests.cpp
@@ -11,10 +11,7 @@
 #include "UObject/UObjectAnnotation.h"
 #include "UObject/UObjectThreadContext.h"
 
-THIRD_PARTY_INCLUDES_START
 #include "Catch2Includes.h"
-THIRD_PARTY_INCLUDES_END
-
 
 TEST_CASE("UObject.NewObject")
 {
diff --git a/Engine/Source/Programs/AutoRTFMTests/Private/UPackageTests.cpp b/Engine/Source/Programs/AutoRTFMTests/Private/UPackageTests.cpp
new file mode 100644
--- /dev/null
+++ b/Engine/Source/Programs/AutoRTFMTests/Private/UPackageTests.cpp
@@ -0,0 +1,100 @@
+// Copyright Epic Games, Inc. All Rights Reserved.
+
+#include "AutoRTFM.h"
+#include "AutoRTFMTesting.h"
+#include "UObject/Package.h"
+
+#include "Catch2Includes.h"
+
+TEST_CASE("UPackage.SetPackageFlagsTo")
+{
+	SECTION("Commit")
+	{
+		UPackage* Package = NewObject<UPackage>();
+		Package->SetPackageFlagsTo(PKG_None);
+
+		AutoRTFM::Testing::Commit([&]
+		{
+			Package->SetPackageFlagsTo(PKG_TransientFlags);
+		});
+
+		REQUIRE(Package->GetPackageFlags() == PKG_TransientFlags);
+	}
+
+	SECTION("Abort")
+	{
+		UPackage* Package = NewObject<UPackage>();
+		Package->SetPackageFlagsTo(PKG_None);
+
+		AutoRTFM::Testing::Abort([&]
+		{
+			Package->SetPackageFlagsTo(PKG_TransientFlags);
+			REQUIRE(Package->GetPackageFlags() == PKG_TransientFlags);
+			AutoRTFM::AbortTransaction();
+		});
+
+		REQUIRE(Package->GetPackageFlags() == PKG_None);
+	}
+}
+
+TEST_CASE("UPackage.SetPackageFlags")
+{
+	SECTION("Commit")
+	{
+		UPackage* Package = NewObject<UPackage>();
+		Package->SetPackageFlagsTo(PKG_RuntimeGenerated);
+
+		AutoRTFM::Testing::Commit([&]
+		{
+			Package->SetPackageFlags(PKG_TransientFlags);
+		});
+
+		REQUIRE(Package->GetPackageFlags() == (PKG_RuntimeGenerated | PKG_TransientFlags));
+	}
+
+	SECTION("Abort")
+	{
+		UPackage* Package = NewObject<UPackage>();
+		Package->SetPackageFlagsTo(PKG_RuntimeGenerated);
+
+		AutoRTFM::Testing::Abort([&]
+		{
+			Package->SetPackageFlags(PKG_TransientFlags);
+			REQUIRE(Package->GetPackageFlags() == (PKG_RuntimeGenerated | PKG_TransientFlags));
+			AutoRTFM::AbortTransaction();
+		});
+
+		REQUIRE(Package->GetPackageFlags() == PKG_RuntimeGenerated);
+	}
+}
+
+TEST_CASE("UPackage.ClearPackageFlags")
+{
+	SECTION("Commit")
+	{
+		UPackage* Package = NewObject<UPackage>();
+		Package->SetPackageFlagsTo(PKG_RuntimeGenerated | PKG_TransientFlags);
+
+		AutoRTFM::Testing::Commit([&]
+		{
+			Package->ClearPackageFlags(PKG_TransientFlags);
+		});
+
+		REQUIRE(Package->GetPackageFlags() == PKG_RuntimeGenerated);
+	}
+
+	SECTION("Abort")
+	{
+		UPackage* Package = NewObject<UPackage>();
+		Package->SetPackageFlagsTo(PKG_RuntimeGenerated | PKG_TransientFlags);
+
+		AutoRTFM::Testing::Abort([&]
+		{
+			Package->ClearPackageFlags(PKG_TransientFlags);
+			REQUIRE(Package->GetPackageFlags() == PKG_RuntimeGenerated);
+			AutoRTFM::AbortTransaction();
+		});
+
+		REQUIRE(Package->GetPackageFlags() == (PKG_RuntimeGenerated | PKG_TransientFlags));
+	}
+}
\ No newline at end of file
diff --git a/Engine/Source/Runtime/CoreUObject/Public/UObject/Package.h b/Engine/Source/Runtime/CoreUObject/Public/UObject/Package.h
--- a/Engine/Source/Runtime/CoreUObject/Public/UObject/Package.h
+++ b/Engine/Source/Runtime/CoreUObject/Public/UObject/Package.h
@@ -319,7 +319,7 @@ private:
 #endif // WITH_EDITORONLY_DATA
 
 	/** Package Flags */
-	uint32 PackageFlagsPrivate;
+	std::atomic<uint32> PackageFlagsPrivate;
 	
 	/** Globally unique id */
 	FPackageId PackageId;
@@ -749,27 +749,113 @@ public:
 	*/
 	FORCEINLINE void SetPackageFlagsTo( uint32 NewFlags )
 	{
-		PackageFlagsPrivate = NewFlags;
+		uint32 OldFlags = GetPackageFlags();
+
+		// Fast path without atomics if already set
+		if (NewFlags == OldFlags)
+		{
+			return;
+		}
+
+		if (AutoRTFM::IsClosed())
+		{
+			UE_AUTORTFM_OPEN
+			{
+				OldFlags = PackageFlagsPrivate.exchange(NewFlags, std::memory_order_relaxed);
+			};
+
+			// If we abort we undo setting the flags we just set.
+			AutoRTFM::OnAbort([this, OldFlags]
+				{
+					PackageFlagsPrivate.store(OldFlags, std::memory_order_relaxed);
+				});
+		}
+		else
+		{
+			PackageFlagsPrivate.store(NewFlags, std::memory_order_relaxed);
+		}
 	}
 
 	/**
 	* Set the specified flags to true. Does not affect any other flags.
 	*
-	* @param	NewFlags		Package flags to enable
+	* @param	FlagsToSet		Package flags to enable
 	*/
-	FORCEINLINE void SetPackageFlags( uint32 NewFlags )
+	FORCEINLINE void SetPackageFlags( uint32 FlagsToSet )
 	{
-		SetPackageFlagsTo(PackageFlagsPrivate | NewFlags);
+		uint32 OldFlags = GetPackageFlags();
+		uint32 NewFlags = OldFlags | FlagsToSet;
+
+		// Fast path without atomics if already set
+		if (NewFlags == OldFlags)
+		{
+			return;
+		}
+
+		if (AutoRTFM::IsClosed())
+		{
+			UE_AUTORTFM_OPEN
+			{
+				OldFlags = PackageFlagsPrivate.fetch_or(FlagsToSet, std::memory_order_relaxed);
+			};
+
+			// If we abort we undo setting the flags we just set.
+			AutoRTFM::OnAbort([this, OldFlags, FlagsToSet]
+				{
+					uint32 MaskFlags = OldFlags;
+
+					// Now just extract out the old flags that mattered (the ones we were setting).
+					MaskFlags &= FlagsToSet;
+					// And unmask the flags we didn't mention.
+					MaskFlags |= ~FlagsToSet;
+
+					PackageFlagsPrivate.fetch_and(MaskFlags, std::memory_order_relaxed);
+				});
+		}
+		else
+		{
+			PackageFlagsPrivate.fetch_or(FlagsToSet, std::memory_order_relaxed);
+		}
 	}
 
 	/**
 	* Set the specified flags to false. Does not affect any other flags.
 	*
-	* @param	NewFlags		Package flags to disable
+	* @param	FlagsToClear		Package flags to disable
 	*/
-	FORCEINLINE void ClearPackageFlags( uint32 NewFlags )
+	FORCEINLINE void ClearPackageFlags( uint32 FlagsToClear )
 	{
-		SetPackageFlagsTo(PackageFlagsPrivate & ~NewFlags);
+		uint32 OldFlags = GetPackageFlags();
+		uint32 NewFlags = OldFlags & ~FlagsToClear;
+
+		// Fast path without atomics if already cleared
+		if (NewFlags == OldFlags)
+		{
+			return;
+		}
+
+		if (AutoRTFM::IsClosed())
+		{
+			UE_AUTORTFM_OPEN 
+			{
+				OldFlags = PackageFlagsPrivate.fetch_and(~FlagsToClear, std::memory_order_relaxed);
+			};
+
+			// If we abort we undo clearing the flags we just unset.
+			AutoRTFM::OnAbort([this, OldFlags, FlagsToClear]
+			{
+				int32 MaskFlags = OldFlags;
+
+				// Now just extract out the old flags that mattered (the ones we were setting).
+				MaskFlags &= FlagsToClear;
+
+				PackageFlagsPrivate.fetch_or(MaskFlags, std::memory_order_relaxed);
+			});
+		}
+		else
+		{
+			PackageFlagsPrivate.fetch_and(~FlagsToClear, std::memory_order_relaxed);
+		}
 	}
 
 	/**
@@ -782,7 +868,7 @@ public:
 	*/
 	FORCEINLINE bool HasAnyPackageFlags( uint32 FlagsToCheck ) const
 	{
-		return (PackageFlagsPrivate & FlagsToCheck) != 0;
+		return (GetPackageFlags() & FlagsToCheck) != 0;
 	}
 
 	/**
@@ -793,7 +879,7 @@ public:
 	*/
 	FORCEINLINE bool HasAllPackagesFlags( uint32 FlagsToCheck ) const
 	{
-		return ((PackageFlagsPrivate & FlagsToCheck) == FlagsToCheck);
+		return ((GetPackageFlags() & FlagsToCheck) == FlagsToCheck);
 	}
 
 	/**
@@ -803,7 +889,14 @@ public:
 	*/
 	FORCEINLINE uint32 GetPackageFlags() const
 	{
-		return PackageFlagsPrivate;
+		uint32 Result = 0;
+
+		UE_AUTORTFM_OPEN
+		{
+			Result = PackageFlagsPrivate.load(std::memory_order_relaxed);
+		};
+
+		return Result;
 	}
 
 	/**
@@ -827,11 +920,11 @@ public:
 	/** Gets how package can be referenced from other plugins and mount points */
 	FORCEINLINE EAssetAccessSpecifier GetAssetAccessSpecifier() const
 	{
-		if (PackageFlagsPrivate & PKG_NotExternallyReferenceable)
+		if (HasAnyPackageFlags(PKG_NotExternallyReferenceable))
 		{
 			return EAssetAccessSpecifier::Private;
 		}
-		else if (PackageFlagsPrivate & PKG_AccessSpecifierEpicInternal)
+		else if (HasAnyPackageFlags(PKG_AccessSpecifierEpicInternal))
 		{
 			return EAssetAccessSpecifier::EpicInternal;
 		}
