[AutoRTFM] Prevent a crash when TryCollectGarbage is called within a transaction.

If a transaction is active, we can simply return false immediately. This is within the API contract and avoids a hard crash.

This change has the added benefit of allowing us to write transactionally-safe test code without carving out special-case exceptions for AutoRTFM.

diff --git a/Engine/Plugins/Runtime/StateTree/Source/StateTreeTestSuite/Private/StateTreeBindingTest.cpp b/Engine/Plugins/Runtime/StateTree/Source/StateTreeTestSuite/Private/StateTreeBindingTest.cpp
--- a/Engine/Plugins/Runtime/StateTree/Source/StateTreeTestSuite/Private/StateTreeBindingTest.cpp
+++ b/Engine/Plugins/Runtime/StateTree/Source/StateTreeTestSuite/Private/StateTreeBindingTest.cpp
@@ -339,11 +339,8 @@ struct FStateTreeTest_CopyObjects : FStateTreeTestBase
 		ObjectB = nullptr;
 		Source.Object = nullptr;
 
-		// Run a garbage collection pass here when possible. (Garbage collection isn't allowed inside an AutoRTFM transaction.)
-		if (!AutoRTFM::IsTransactional())
-		{
-			CollectGarbage(GARBAGE_COLLECTION_KEEPFLAGS);
-		}
+		// @todo: Avoid relying on GC within this test.
+		TryCollectGarbage(GARBAGE_COLLECTION_KEEPFLAGS);
 
 		FStateTreeTest_PropertyCopyObjects TargetC;
 		bool bCopyResultC = true;
diff --git a/Engine/Source/Runtime/CoreUObject/Private/UObject/GarbageCollection.cpp b/Engine/Source/Runtime/CoreUObject/Private/UObject/GarbageCollection.cpp
--- a/Engine/Source/Runtime/CoreUObject/Private/UObject/GarbageCollection.cpp
+++ b/Engine/Source/Runtime/CoreUObject/Private/UObject/GarbageCollection.cpp
@@ -7,6 +7,7 @@
 #include "UObject/GarbageCollection.h"
 #include "Algo/Find.h"
 #include "Async/UniqueLock.h"
+#include "AutoRTFM.h"
 #include "Containers/ConsumeAllMpmcQueue.h"
 #include "Containers/Deque.h"
 #include "Containers/StaticBitArray.h"
@@ -6195,7 +6196,14 @@ bool TryCollectGarbage(EObjectFlags KeepFlags, bool bPerformFullPurge)
 	if (GIsInitialLoad)
 	{
 		// During initial load classes may not yet have their GC token streams assembled
-		UE_LOG(LogGarbage, Log, TEXT("Skipping CollectGarbage() call during initial load. It's not safe."));
+		UE_LOG(LogGarbage, Log, TEXT("Skipping TryCollectGarbage() call during initial load. It's not safe."));
+		return false;
+	}
+
+	if (AutoRTFM::IsTransactional())
+	{
+		// Memory cannot be freed within a transaction as this would prevent us from rolling back to the initial state.
+		UE_LOG(LogGarbage, Log, TEXT("TryCollectGarbage: skipping garbage collection because an AutoRTFM transaction is active."));
 		return false;
 	}
 
diff --git a/Engine/Source/Runtime/CoreUObject/Public/UObject/UObjectGlobals.h b/Engine/Source/Runtime/CoreUObject/Public/UObject/UObjectGlobals.h
--- a/Engine/Source/Runtime/CoreUObject/Public/UObject/UObjectGlobals.h
+++ b/Engine/Source/Runtime/CoreUObject/Public/UObject/UObjectGlobals.h
@@ -933,6 +933,8 @@ COREUOBJECT_API void CollectGarbage(EObjectFlags KeepFlags, bool bPerformFullPur
 *
 * @param	KeepFlags			objects with those flags will be kept regardless of being referenced or not
 * @param	bPerformFullPurge	if true, perform a full purge after the mark pass
+*
+* @return true if an actual attempt at garbage collection occurred
 */
 COREUOBJECT_API bool TryCollectGarbage(EObjectFlags KeepFlags, bool bPerformFullPurge = true);
 
