Replication
* IsFullNameStableForNetworking will ignore the result of IsNameStableForNetworking for templates
* This fixes CDO/Archetypes that are not referenceable due to a class overriding IsNameStableForNetworking and forcing it to always return false.

#[review]-43030321
#rnx

diff --git a/Engine/Source/Runtime/CoreUObject/Private/UObject/Obj.cpp b/Engine/Source/Runtime/CoreUObject/Private/UObject/Obj.cpp
--- a/Engine/Source/Runtime/CoreUObject/Private/UObject/Obj.cpp
+++ b/Engine/Source/Runtime/CoreUObject/Private/UObject/Obj.cpp
@@ -5759,6 +5759,7 @@ PRAGMA_ENABLE_DEPRECATION_WARNINGS
 /** IsNameStableForNetworking means an object can be referred to its path name (relative to outer) over the network */
 bool UObject::IsNameStableForNetworking() const
 {
+	//TODO: This should test for RF_ArchetypeObject too
 	return HasAnyFlags(RF_WasLoaded | RF_DefaultSubObject | RF_ClassDefaultObject) || IsNative() || IsDefaultSubobject();
 }
 
@@ -5767,10 +5768,11 @@ bool UObject::IsFullNameStableForNetworking() const
 {
 	if ( GetOuter() != NULL && !GetOuter()->IsFullNameStableForNetworking() )
 	{
-		return false;	// If any outer isn't stable, we can't consider the full name stable
+		return false; // If any outer isn't stable, we can't consider the full name stable
 	}
 
-	return IsNameStableForNetworking();
+	// Consider templates to always be stable
+	return IsTemplate() || IsNameStableForNetworking();
 }
 
 /** IsSupportedForNetworking means an object can be referenced over the network */
