[ZenLoader]
- Remove the check for the Async flag from VerifyObjectLoadFlagsWhenFinishedLoading since this flag is not owned by the loader and may conflict with other UObject activity

diff --git a/Engine/Source/Runtime/CoreUObject/Private/Serialization/AsyncLoading2.cpp b/Engine/Source/Runtime/CoreUObject/Private/Serialization/AsyncLoading2.cpp
--- a/Engine/Source/Runtime/CoreUObject/Private/Serialization/AsyncLoading2.cpp
+++ b/Engine/Source/Runtime/CoreUObject/Private/Serialization/AsyncLoading2.cpp
@@ -10273,7 +10273,7 @@ static void VerifyObjectLoadFlagsWhenFinishedLoading()
 	TRACE_CPUPROFILER_EVENT_SCOPE(VerifyObjectLoadFlagsWhenFinishedLoading);
 
 	const EInternalObjectFlags AsyncFlags =
-		EInternalObjectFlags::Async | EInternalObjectFlags_AsyncLoading;
+		EInternalObjectFlags_AsyncLoading;
 
 	const EObjectFlags LoadIntermediateFlags = 
 		EObjectFlags::RF_NeedLoad | EObjectFlags::RF_WillBeLoaded |
