Create Instance Data Objects for a sub-set of newly created objects (e.g., spawned).

This is a temporary solution until we can safely move the creation call to: StaticConstructObject_Internal

- The creation of IDOs in TryCreateIDOForConstructedObject is wrapped in an AutoRTFM on-commit block in order to handle the AutoRTFM validation in the Verse tests.
- I am disabling this feature during the execution of Verse tests because deserializing IDOs during these tests will encounter NULL values on non-nullable properties (as part of the test). These NULL values will cause warnings which are unavoidable.

This feature can be disabled using this CVar:
IDO.EnableOnSpawnedObjects 0

#rnx

diff --git a/Engine/Source/Runtime/CoreUObject/Internal/UObject/InstanceDataObjectUtils.h b/Engine/Source/Runtime/CoreUObject/Internal/UObject/InstanceDataObjectUtils.h
--- a/Engine/Source/Runtime/CoreUObject/Internal/UObject/InstanceDataObjectUtils.h
+++ b/Engine/Source/Runtime/CoreUObject/Internal/UObject/InstanceDataObjectUtils.h
@@ -6,13 +6,12 @@
 
 #if WITH_EDITORONLY_DATA
 
-#include "UObject/PropertyStateTracking.h" // Temporary...
-
 #define UE_API COREUOBJECT_API
 
 class FFieldVariant;
 class FProperty;
 class FStructuredArchiveRecord;
+class FArchive;
 class UClass;
 class UObject;
 class UStruct;
@@ -91,6 +90,11 @@ UE_API UObject* CreateInstanceDataObject(UObject* Owner, FArchive& Ar, int64 Sta
 /** Returns the IDO for the object if it has one, otherwise returns the object. */
 UE_API UObject* ResolveInstanceDataObject(UObject* Object);
 
+namespace Private
+{
+	UE_API bool TryCreateIDOForConstructedObject(UObject* Object, bool bCreateIDOsForSubObjects = true);
+} // Private
+
 } // UE
 
 #undef UE_API
diff --git a/Engine/Source/Runtime/CoreUObject/Private/UObject/InstanceDataObjectUtils.cpp b/Engine/Source/Runtime/CoreUObject/Private/UObject/InstanceDataObjectUtils.cpp
--- a/Engine/Source/Runtime/CoreUObject/Private/UObject/InstanceDataObjectUtils.cpp
+++ b/Engine/Source/Runtime/CoreUObject/Private/UObject/InstanceDataObjectUtils.cpp
@@ -17,6 +17,7 @@
 #include "UObject/Class.h"
 #include "UObject/EnumProperty.h"
 #include "UObject/Field.h"
+#include "UObject/PropertyStateTracking.h"
 #include "UObject/Package.h"
 #include "UObject/PropertyBagRepository.h"
 #include "UObject/PropertyHelper.h"
@@ -285,6 +286,13 @@ namespace UE
 		FConsoleVariableDelegate::CreateStatic(OnInstanceDataObjectSupportChanged)
 	);
 
+	bool bEnableIDOOnSpawnedObjects = true;
+	FAutoConsoleVariableRef EnableIDOOnSpawnedObjects(
+		TEXT("IDO.EnableOnSpawnedObjects"),
+		bEnableIDOOnSpawnedObjects,
+		TEXT("Allows an IDO to be created for every newly created (e.g., spawned) object.")
+	);
+
 	bool bEnableIDOUnknownProperties = true;
 	FAutoConsoleVariableRef EnableIDOUnknownProperties(
 		TEXT("IDO.Unknowns.EnableProperties"),
@@ -1601,6 +1609,92 @@ UObject* ResolveInstanceDataObject(UObject* Object)
 	return InstanceDataObject ? InstanceDataObject : Object;
 }
 
+namespace Private
+{
+
+static bool CanCreateIDOForConstructedObject(UObject* Object)
+{
+	bool bCanCreateIDO = false;
+
+	const UClass* ObjectClass = Object->GetClass();
+
+	// Ignore objects created very early in the initialization of CoreUObject.
+	if (ObjectClass && ObjectClass->GetClass())
+	{
+		// Do not create IDOs for the following types of objects:
+		// - Object that will be loaded: we handle those via FScopedObjectSerializeContext.
+		// - CDOs: created when the IDO class is created.
+		// - Archetypes: created on-demand.
+		// - Transient objects.
+		// - Objects with an obsolete type.
+		// - Types that do not support IDOs.
+		// - IDOs: we don't want to create IDOs for IDOs.
+		if (!Object->HasAnyFlags(RF_NeedLoad | RF_ClassDefaultObject | RF_ArchetypeObject | RF_Transient)
+			&& !Object->HasAnyInternalFlags(EInternalObjectFlags_AsyncLoading)
+			&& !ObjectClass->HasAnyClassFlags(CLASS_NewerVersionExists)
+			&& IsInstanceDataObjectSupportEnabledForClass(ObjectClass)
+			&& !Object->IsInPackage(GetTransientPackage()))
+		{
+			// Assume this object is getting loaded if an object is being serialized, so skip IDO creation.
+			const FUObjectSerializeContext* SerializeContext = FUObjectThreadContext::Get().GetSerializeContext();
+			if (!SerializeContext || !SerializeContext->SerializedObject)
+			{
+				bCanCreateIDO = !IsInstanceDataObject(Object) && IsInstanceDataObjectSupportEnabled(Object);
+			}
+		}
+	}
+
+	return bCanCreateIDO;
+}
+
+static bool TryCreateIDOForConstructedObject_Recursive(UObject* Object)
+{
+	bool bDidCreateIDO = false;
+
+	if (CanCreateIDOForConstructedObject(Object))
+	{
+		UObject* IDO = CreateInstanceDataObject(Object);
+		bDidCreateIDO = (IDO != nullptr);
+	}
+	TArray<UObject*> SubObjects;
+	constexpr bool bIncludeNestedObjects = false;
+	GetObjectsWithOuter(Object, SubObjects, bIncludeNestedObjects);
+	for (UObject* SubObject : SubObjects)
+	{
+		if (TryCreateIDOForConstructedObject_Recursive(SubObject))
+		{
+			bDidCreateIDO = true;
+		}
+	}
+
+	return bDidCreateIDO;
+}
+
+bool TryCreateIDOForConstructedObject(UObject* Object, bool bCreateIDOsForSubObjects)
+{
+	bool bDidCreateIDO = false;
+
+	if (bEnableIDOOnSpawnedObjects)
+	{
+		UE_AUTORTFM_ONCOMMIT(Object, bDidCreateIDO, bCreateIDOsForSubObjects)
+		{
+			if (bCreateIDOsForSubObjects)
+			{
+				bDidCreateIDO = TryCreateIDOForConstructedObject_Recursive(Object);
+			}
+			else if (CanCreateIDOForConstructedObject(Object))
+			{
+				UObject* IDO = CreateInstanceDataObject(Object);
+				bDidCreateIDO = (IDO != nullptr);
+			}
+		};
+	}
+
+	return bDidCreateIDO;
+}
+
+} // Private
+
 } // UE
 
 #endif // WITH_EDITORONLY_DATA
diff --git a/Engine/Source/Runtime/CoreUObject/Private/UObject/Property.cpp b/Engine/Source/Runtime/CoreUObject/Private/UObject/Property.cpp
--- a/Engine/Source/Runtime/CoreUObject/Private/UObject/Property.cpp
+++ b/Engine/Source/Runtime/CoreUObject/Private/UObject/Property.cpp
@@ -35,6 +35,10 @@
 #include "UObject/UObjectGlobals.h"
 #include "UObject/UObjectIterator.h"
 
+#if WITH_EDITORONLY_DATA
+#include "UObject/PropertyStateTracking.h"
+#endif
+
 DEFINE_LOG_CATEGORY(LogProperty);
 
 namespace UE::CoreUObject::Private
diff --git a/Engine/Source/Runtime/Engine/Private/PrintObjectUtils.cpp b/Engine/Source/Runtime/Engine/Private/PrintObjectUtils.cpp
--- a/Engine/Source/Runtime/Engine/Private/PrintObjectUtils.cpp
+++ b/Engine/Source/Runtime/Engine/Private/PrintObjectUtils.cpp
@@ -15,6 +15,7 @@
 #if WITH_EDITORONLY_DATA
 #include "UObject/InstanceDataObjectUtils.h"
 #include "UObject/PropertyBagRepository.h"
+#include "UObject/PropertyStateTracking.h"
 #endif
 
 namespace UE
