MetaData: fix regression from 40093672 that was causing the old meta data object not preloaded when queried as it was before the change, resulting in various meta data misses.

diff --git a/Engine/Source/Runtime/CoreUObject/Private/UObject/Package.cpp b/Engine/Source/Runtime/CoreUObject/Private/UObject/Package.cpp
--- a/Engine/Source/Runtime/CoreUObject/Private/UObject/Package.cpp
+++ b/Engine/Source/Runtime/CoreUObject/Private/UObject/Package.cpp
@@ -280,6 +280,16 @@ TArray<UPackage*> UPackage::GetExternalPackages() const
  */
 FMetaData& UPackage::GetMetaData()
 {
+	if (UDEPRECATED_MetaData* DeprecatedMetaData = FindObjectFast<UDEPRECATED_MetaData>(this, FName(NAME_PackageMetaData)))
+	{
+		if (DeprecatedMetaData->HasAnyFlags(RF_NeedLoad))
+		{
+			FLinkerLoad* MetaDataLinker = DeprecatedMetaData->GetLinker();
+			check(MetaDataLinker);
+			MetaDataLinker->Preload(DeprecatedMetaData);
+		}
+	}
+
 	return MetaData;
 }
 #endif
