[ALT/UObjectGlobals]
- PostRename can be called from ALT (async loading thread) and all implementers need to be thread-safe when ALT is active

diff --git a/Engine/Source/Runtime/CoreUObject/Public/UObject/UObjectGlobals.h b/Engine/Source/Runtime/CoreUObject/Public/UObject/UObjectGlobals.h
--- a/Engine/Source/Runtime/CoreUObject/Public/UObject/UObjectGlobals.h
+++ b/Engine/Source/Runtime/CoreUObject/Public/UObject/UObjectGlobals.h
@@ -3260,7 +3260,7 @@ struct FCoreUObjectDelegates
 	}
 
 	/** Called when an object is renamed this gets called in both the editor and standalone game editor builds, for every object renamed */
-	DECLARE_MULTICAST_DELEGATE_ThreeParams(FOnObjectRenamed, UObject*, UObject*, FName);
+	DECLARE_TS_MULTICAST_DELEGATE_ThreeParams(FOnObjectRenamed, UObject*, UObject*, FName);
 	static COREUOBJECT_API FOnObjectRenamed OnObjectRenamed;
 
 	/** Broadcast OnObjectRenamed event */
