Fix bug in GC timing where if it destroyed the right number of objects it would not check the time again and cause a large stall

diff --git a/Engine/Source/Runtime/CoreUObject/Private/UObject/GarbageCollection.cpp b/Engine/Source/Runtime/CoreUObject/Private/UObject/GarbageCollection.cpp
--- a/Engine/Source/Runtime/CoreUObject/Private/UObject/GarbageCollection.cpp
+++ b/Engine/Source/Runtime/CoreUObject/Private/UObject/GarbageCollection.cpp
@@ -890,7 +890,7 @@ public:
 				++ObjCurrentFreeIndexObjectIndex;
 
 				// Time slicing when running on the game thread
-				if (bUseTimeLimit && (ProcessedObjectsCount == TimeLimitEnforcementGranularityForDeletion) && (ObjCurrentFreeIndexObjectIndex < GUnreachableObjects.Num()))
+				if (bUseTimeLimit && (ProcessedObjectsCount >= TimeLimitEnforcementGranularityForDeletion) && (ObjCurrentFreeIndexObjectIndex < GUnreachableObjects.Num()))
 				{
 					ProcessedObjectsCount = 0;
 					if ((FPlatformTime::Seconds() - StartTime) > TimeLimit)
@@ -919,7 +919,7 @@ public:
 				++ObjCurrentPurgeObjectIndex;
 
 				// Time slicing when running on the game thread
-				if (bUseTimeLimit && (ProcessedObjectsCount == TimeLimitEnforcementGranularityForDeletion) && (ObjCurrentPurgeObjectIndex < GUnreachableObjects.Num()))
+				if (bUseTimeLimit && (ProcessedObjectsCount >= TimeLimitEnforcementGranularityForDeletion) && (ObjCurrentPurgeObjectIndex < GUnreachableObjects.Num()))
 				{
 					ProcessedObjectsCount = 0;
 					if ((FPlatformTime::Seconds() - StartTime) > TimeLimit)
