Added CDerivedFrom and cleaned up some header and trait usage in UnrealTypes.h.

diff --git a/Engine/Source/Runtime/Core/Public/Concepts/DerivedFrom.h b/Engine/Source/Runtime/Core/Public/Concepts/DerivedFrom.h
new file mode 100644
--- /dev/null
+++ b/Engine/Source/Runtime/Core/Public/Concepts/DerivedFrom.h
@@ -0,0 +1,16 @@
+// Copyright Epic Games, Inc. All Rights Reserved.
+
+#pragma once
+
+#include <type_traits>
+
+namespace UE
+{
+	/**
+	 * Concept which is satisfied if DerivedType is an unambiguous public base class of BaseType.
+	 *
+	 * We use this instead of std::same_as because <concepts> isn't a well supported header yet.
+	 */
+	template <typename DerivedType, typename BaseType>
+	concept CDerivedFrom = std::is_base_of_v<BaseType, DerivedType> && std::is_convertible_v<const volatile DerivedType*, const volatile BaseType*>;
+}
diff --git a/Engine/Source/Runtime/CoreUObject/Public/UObject/UnrealType.h b/Engine/Source/Runtime/CoreUObject/Public/UObject/UnrealType.h
--- a/Engine/Source/Runtime/CoreUObject/Public/UObject/UnrealType.h
+++ b/Engine/Source/Runtime/CoreUObject/Public/UObject/UnrealType.h
@@ -7,6 +7,7 @@
 #pragma once
 
 #include "Concepts/GetTypeHashable.h"
+#include "Concepts/DerivedFrom.h"
 #include "Containers/Array.h"
 #include "Containers/ArrayView.h"
 #include "Containers/ContainerAllocationPolicies.h"
@@ -20,12 +21,9 @@
 #include "Containers/UnrealString.h"
 #include "CoreGlobals.h"
 #include "HAL/MemoryBase.h"
-#include "HAL/PlatformMath.h"
 #include "HAL/UnrealMemory.h"
 #include "Logging/LogCategory.h"
 #include "Logging/LogMacros.h"
-#include "Math/NumericLimits.h"
-#include "Math/UnrealMathUtility.h"
 #include "Misc/AssertionMacros.h"
 #include "Misc/EnumClassFlags.h"
 #include "Misc/NotNull.h"
@@ -35,12 +33,7 @@
 #include "Serialization/SerializedPropertyScope.h"
 #include "Serialization/StructuredArchive.h"
 #include "Templates/Casts.h"
-#include "Templates/EnableIf.h"
-#include "Templates/Greater.h"
-#include "Templates/IsFloatingPoint.h"
-#include "Templates/IsIntegral.h"
 #include "Templates/IsPODType.h"
-#include "Templates/IsSigned.h"
 #include "Templates/IsUEnumClass.h"
 #include "Templates/MemoryOps.h"
 #include "Templates/Models.h"
@@ -1715,12 +1708,12 @@ class FNumericProperty : public FProperty
 	template <typename T>
 	bool CanHoldValue(T Value) const
 	{
-		if (TIsFloatingPoint<T>::Value)
+		if (std::is_floating_point_v<T>)
 		{
 			//@TODO: FLOATPRECISION: This feels wrong, it might be losing precision before it tests to see if it's going to lose precision...
 			return CanHoldDoubleValueInternal((double)Value);
 		}
-		else if (TIsSigned<T>::Value)
+		else if (std::is_signed_v<T>)
 		{
 			return CanHoldSignedValueInternal(Value);
 		}
@@ -1870,7 +1863,7 @@ private:
 			Property.SetPropertyValue_InContainer(Obj, NewValue, Tag.ArrayIndex);
 
 			UE_CLOG(
-				((TIsSigned<FromType>::Value || TIsFloatingPoint<FromType>::Value) && (!TIsSigned<ToType>::Value && !TIsFloatingPoint<ToType>::Value) && OldValue < 0) || ((FromType)NewValue != OldValue),
+				((std::is_signed_v<FromType> || std::is_floating_point_v<FromType>) && (!std::is_signed_v<ToType> && !std::is_floating_point_v<ToType>) && OldValue < 0) || ((FromType)NewValue != OldValue),
 				LogClass,
 				Warning,
 				TEXT("Potential data loss during conversion of integer property %s of %s - was (%s) now (%s) - for package: %s"),
@@ -1973,25 +1966,25 @@ public:
 	// FNumericProperty interface.
 	virtual bool IsFloatingPoint() const override
 	{
-		return TIsFloatingPoint<TCppType>::Value;
+		return std::is_floating_point_v<TCppType>;
 	}
 	virtual bool IsInteger() const override
 	{
-		return TIsIntegral<TCppType>::Value;
+		return std::is_integral_v<TCppType>;
 	}
 	virtual void SetIntPropertyValue(void* Data, uint64 Value) const override
 	{
-		check(TIsIntegral<TCppType>::Value);
+		check(std::is_integral_v<TCppType>);
 		TTypeFundamentals::SetPropertyValue(Data, (TCppType)Value);
 	}
 	virtual void SetIntPropertyValue(void* Data, int64 Value) const override
 	{
-		check(TIsIntegral<TCppType>::Value);
+		check(std::is_integral_v<TCppType>);
 		TTypeFundamentals::SetPropertyValue(Data, (TCppType)Value);
 	}
 	virtual void SetFloatingPointPropertyValue(void* Data, double Value) const override
 	{
-		check(TIsFloatingPoint<TCppType>::Value);
+		check(std::is_floating_point_v<TCppType>);
 		TTypeFundamentals::SetPropertyValue(Data, (TCppType)Value);
 	}
 	virtual void SetNumericPropertyValueFromString(void* Data, TCHAR const* Value) const override
@@ -2016,31 +2009,31 @@ public:
 	}
 	virtual int64 GetSignedIntPropertyValue(void const* Data) const override
 	{
-		check(TIsIntegral<TCppType>::Value);
+		check(std::is_integral_v<TCppType>);
 		return (int64)TTypeFundamentals::GetPropertyValue(Data);
 	}
 	virtual int64 GetSignedIntPropertyValue_InContainer(void const* Container) const override
 	{
-		check(TIsIntegral<TCppType>::Value);
+		check(std::is_integral_v<TCppType>);
 		TCppType LocalValue{};
 		FNumericProperty::GetValue_InContainer(Container, &LocalValue);
 		return (int64)LocalValue;
 	}
 	virtual uint64 GetUnsignedIntPropertyValue(void const* Data) const override
 	{
-		check(TIsIntegral<TCppType>::Value);
+		check(std::is_integral_v<TCppType>);
 		return (uint64)TTypeFundamentals::GetPropertyValue(Data);
 	}
 	virtual uint64 GetUnsignedIntPropertyValue_InContainer(void const* Container) const override
 	{
-		check(TIsIntegral<TCppType>::Value);
+		check(std::is_integral_v<TCppType>);
 		TCppType LocalValue{};
 		FNumericProperty::GetValue_InContainer(Container, &LocalValue);
 		return (uint64)LocalValue;
 	}
 	virtual double GetFloatingPointPropertyValue(void const* Data) const override
 	{
-		check(TIsFloatingPoint<TCppType>::Value);
+		check(std::is_floating_point_v<TCppType>);
 		return (double)TTypeFundamentals::GetPropertyValue(Data);
 	}
 	// End of FNumericProperty interface
@@ -6822,10 +6815,7 @@ struct TFieldRange
 	Field templates.
 -----------------------------------------------------------------------------*/
 
-template <
-	typename T
-	UE_REQUIRES(std::is_base_of_v<UField, T>)
->
+template <UE::CDerivedFrom<UField> T>
 T* FindUField(const UStruct* Owner, FName FieldName, EFieldIterationFlags IterationFlags = EFieldIterationFlags::Default)
 {
 	static_assert(sizeof(T) > 0, "T must not be an incomplete type");
@@ -6849,10 +6839,7 @@ T* FindUField(const UStruct* Owner, FName FieldName, EFieldIterationFlags Iterat
 	return nullptr;
 }
 
-template <
-	typename T
-	UE_REQUIRES(std::is_base_of_v<UField, T>)
->
+template <UE::CDerivedFrom<UField> T>
 T* FindUField(const UStruct* Owner, const TCHAR* FieldName, EFieldIterationFlags IterationFlags = EFieldIterationFlags::Default)
 {
 	static_assert(sizeof(T) > 0, "T must not be an incomplete type");
@@ -6862,10 +6849,7 @@ T* FindUField(const UStruct* Owner, const TCHAR* FieldName, EFieldIterationFlags
 	return FindUField<T>(Owner, Name, IterationFlags);
 }
 
-template <
-	typename T
-	UE_REQUIRES(std::is_base_of_v<FField, T>)
->
+template <UE::CDerivedFrom<FField> T>
 T* FindFProperty(const UStruct* Owner, FName FieldName, EFieldIterationFlags IterationFlags = EFieldIterationFlags::Default)
 {
 	static_assert(sizeof(T) > 0, "T must not be an incomplete type");
@@ -6889,10 +6873,7 @@ T* FindFProperty(const UStruct* Owner, FName FieldName, EFieldIterationFlags Ite
 	return nullptr;
 }
 
-template <
-	typename T
-	UE_REQUIRES(std::is_base_of_v<FField, T>)
->
+template <UE::CDerivedFrom<FField> T>
 T* FindFProperty(const UStruct* Owner, const TCHAR* FieldName, EFieldIterationFlags IterationFlags = EFieldIterationFlags::Default)
 {
 	static_assert(sizeof(T) > 0, "T must not be an incomplete type");
@@ -6922,37 +6903,25 @@ inline FFieldVariant FindUFieldOrFProperty(const UStruct* Owner, const TCHAR* Fi
 	return FindUFieldOrFProperty(Owner, Name, IterationFlags);
 }
 
-template <
-	typename T
-	UE_REQUIRES(std::is_base_of_v<UField, T>)
->
+template <UE::CDerivedFrom<UField> T>
 T* FindUFieldOrFProperty(const UStruct* Owner, FName FieldName, EFieldIterationFlags IterationFlags = EFieldIterationFlags::Default)
 {
 	return FindUField<T>(Owner, FieldName, IterationFlags);
 }
 
-template <
-	typename T
-	UE_REQUIRES(std::is_base_of_v<UField, T>)
->
+template <UE::CDerivedFrom<UField> T>
 T* FindUFieldOrFProperty(const UStruct* Owner, const TCHAR* FieldName, EFieldIterationFlags IterationFlags = EFieldIterationFlags::Default)
 {
 	return FindUField<T>(Owner, FieldName, IterationFlags);
 }
 
-template <
-	typename T
-	UE_REQUIRES(std::is_base_of_v<FField, T>)
->
+template <UE::CDerivedFrom<FField> T>
 T* FindUFieldOrFProperty(const UStruct* Owner, FName FieldName, EFieldIterationFlags IterationFlags = EFieldIterationFlags::Default)
 {
 	return FindFProperty<T>(Owner, FieldName, IterationFlags);
 }
 
-template <
-	typename T
-	UE_REQUIRES(std::is_base_of_v<FField, T>)
->
+template <UE::CDerivedFrom<FField> T>
 T* FindUFieldOrFProperty(const UStruct* Owner, const TCHAR* FieldName, EFieldIterationFlags IterationFlags = EFieldIterationFlags::Default)
 {
 	return FindFProperty<T>(Owner, FieldName, IterationFlags);
@@ -7337,8 +7306,15 @@ class UMulticastInlineDelegatePropertyWrapper : public UMulticastDelegatePropert
 
 #undef UE_API
 
-#if UE_ENABLE_INCLUDE_ORDER_DEPRECATED_IN_5_5
-#include "Templates/IsTriviallyDestructible.h"
+#if UE_ENABLE_INCLUDE_ORDER_DEPRECATED_IN_5_7
+#include "HAL/PlatformMath.h"
+#include "Math/NumericLimits.h"
+#include "Math/UnrealMathUtility.h"
+#include "Templates/EnableIf.h"
+#include "Templates/Greater.h"
+#include "Templates/IsFloatingPoint.h"
+#include "Templates/IsIntegral.h"
+#include "Templates/IsSigned.h"
 #endif
 
 // This should really be included everywhere that FStrProperty is used, but that would require a lot of fixup
