MetaData: fix regression from 40093672 that was causing the old meta data object not preloaded when queried as it was before the change, resulting in various meta data misses.

diff --git a/Engine/Source/Runtime/CoreUObject/Private/UObject/LinkerLoad.cpp b/Engine/Source/Runtime/CoreUObject/Private/UObject/LinkerLoad.cpp
--- a/Engine/Source/Runtime/CoreUObject/Private/UObject/LinkerLoad.cpp
+++ b/Engine/Source/Runtime/CoreUObject/Private/UObject/LinkerLoad.cpp
@@ -4409,6 +4409,13 @@ int32 FLinkerLoad::LoadMetaDataFromExportMap(bool bForcePreload)
 		}
 	}
 
+	if (LinkerRoot && IsValid(DeprecatedMetaData))
+	{
+		PRAGMA_DISABLE_DEPRECATION_WARNINGS
+		LinkerRoot->DeprecatedMetaData = DeprecatedMetaData;
+		PRAGMA_ENABLE_DEPRECATION_WARNINGS
+	}
+
 	return DeprecatedMetaDataIndex;
 }
 #endif // WITH_METADATA
diff --git a/Engine/Source/Runtime/CoreUObject/Private/UObject/MetaData.cpp b/Engine/Source/Runtime/CoreUObject/Private/UObject/MetaData.cpp
--- a/Engine/Source/Runtime/CoreUObject/Private/UObject/MetaData.cpp
+++ b/Engine/Source/Runtime/CoreUObject/Private/UObject/MetaData.cpp
@@ -241,6 +241,10 @@ void UDEPRECATED_MetaData::Serialize(FArchive& Ar)
 			ClearFlags(RF_Standalone);
 			MarkAsGarbage();
 
+			PRAGMA_DISABLE_DEPRECATION_WARNINGS
+			GetPackage()->DeprecatedMetaData = nullptr;
+			PRAGMA_ENABLE_DEPRECATION_WARNINGS
+
 			// This deprecated object is no longer needed now that we have converted it to FMetaData. However,
 			// if loading from disk, the linker might still be referring to the deprecated metadata export which
 			// may be GC'd during load. Thus, we remove the export to avoid issues during loading.
diff --git a/Engine/Source/Runtime/CoreUObject/Private/UObject/Package.cpp b/Engine/Source/Runtime/CoreUObject/Private/UObject/Package.cpp
--- a/Engine/Source/Runtime/CoreUObject/Private/UObject/Package.cpp
+++ b/Engine/Source/Runtime/CoreUObject/Private/UObject/Package.cpp
@@ -280,6 +280,16 @@ TArray<UPackage*> UPackage::GetExternalPackages() const
  */
 FMetaData& UPackage::GetMetaData()
 {
+	PRAGMA_DISABLE_DEPRECATION_WARNINGS
+	if (IsValid(DeprecatedMetaData) && DeprecatedMetaData->HasAllFlags(RF_NeedLoad))
+	{
+		FLinkerLoad* MetaDataLinker = DeprecatedMetaData->GetLinker();
+		check(MetaDataLinker);
+
+		MetaDataLinker->Preload(DeprecatedMetaData);
+	}
+	PRAGMA_ENABLE_DEPRECATION_WARNINGS
+
 	return MetaData;
 }
 #endif
diff --git a/Engine/Source/Runtime/CoreUObject/Public/UObject/Package.h b/Engine/Source/Runtime/CoreUObject/Public/UObject/Package.h
--- a/Engine/Source/Runtime/CoreUObject/Public/UObject/Package.h
+++ b/Engine/Source/Runtime/CoreUObject/Public/UObject/Package.h
@@ -55,6 +55,7 @@ class FSavePackageContext;
 class FString;
 class ITargetPlatform;
 class UFunction;
+class UDEPRECATED_MetaData;
 struct FAssetData;
 struct FMD5Hash;
 struct FPackageSaveInfo;
@@ -217,6 +218,8 @@ class UPackage : public UObject
 	{
 	};
 
+	friend UDEPRECATED_MetaData;
+
 public:
 
 	UPackage(const FObjectInitializer& ObjectInitializer = FObjectInitializer::Get())
@@ -364,6 +367,10 @@ private:
 #if WITH_METADATA
 	// MetaData for the editor
 	FMetaData MetaData;
+
+	// Deprecated UMetaData member when loading older packages
+	UE_DEPRECATED(5.6, "UMetaData was replaced by FMetaData, this member is only used for migrating t he existing data, do not use.")
+	UDEPRECATED_MetaData* DeprecatedMetaData = nullptr;
 #endif // WITH_METADATA
 
 #if WITH_EDITORONLY_DATA
