Optional Cooked Packages: Add workaround for CIS failures on projects that use multi-process cooks for client and server, and pass in -editoroptional for the client but don't want it on server; ignore the error "Can't save %s: Optional object (%s) has a reference to cooked object (%s) which is missing." on server. In the long-term we will replace this workaround with either removal of the error in general or with a per-platform way to specify editoroptional cooks.
#rnx

diff --git a/Engine/Source/Runtime/CoreUObject/Private/UObject/SavePackage2.cpp b/Engine/Source/Runtime/CoreUObject/Private/UObject/SavePackage2.cpp
--- a/Engine/Source/Runtime/CoreUObject/Private/UObject/SavePackage2.cpp
+++ b/Engine/Source/Runtime/CoreUObject/Private/UObject/SavePackage2.cpp
@@ -440,13 +440,29 @@ ESavePackageResult HarvestPackage(FSaveContext& SaveContext)
 					OptionalContext.AddImport(It->Obj);
 					// Flag the package itself to be an import
 					bHasNonOptionalSelfReference = true;
+					It.RemoveCurrent();
 				}
 				// if not found in the game context and the reference directly came from an optional object, record an illegal reference
 				else if (It->bFromOptionalReference)
 				{
-					SaveContext.RecordIllegalReference(nullptr, It->Obj, EIllegalRefReason::ReferenceFromOptionalToMissingGameExport);
+					// TODO: Suppress this error and leave the object in the optional realm on Server builds, because it does not work in
+					// some cases on multiplatform cooks. EditorOptional is a global commandline setting that applies to all platforms in
+					// a multiplatform cook, but on some projects the project cooks Client and Server together and wants editoroptional
+					// only on client, so it cooks both of them as editoroptional. But then server doesn't support editoroptional for
+					// objects that are in ClassesExcludedOnDedicatedServer but that are referenced both from optional and from gamerealm
+					// objects; we get this error because the object is stripped from the game realm but remains in the optional realm.
+					// 
+					// Change EditorOptional to be specifiable per platform, and remove this suppression.
+					if (!SaveContext.GetTargetPlatform()->IsServerOnly())
+					{
+						SaveContext.RecordIllegalReference(nullptr, It->Obj, EIllegalRefReason::ReferenceFromOptionalToMissingGameExport);
+						It.RemoveCurrent();
+					}
+				}
+				else
+				{
+					It.RemoveCurrent();
 				}
-				It.RemoveCurrent();
 			}
 		}
 		// Also add the current package itself as an import if we are referencing any non optional export
