[ClassDefaultObjectTest]
- Speculative fix for GC assumption issue in other tests depending on how GC is setup

diff --git a/Engine/Source/Runtime/CoreUObject/Tests/ClassDefaultObjectTest.cpp b/Engine/Source/Runtime/CoreUObject/Tests/ClassDefaultObjectTest.cpp
--- a/Engine/Source/Runtime/CoreUObject/Tests/ClassDefaultObjectTest.cpp
+++ b/Engine/Source/Runtime/CoreUObject/Tests/ClassDefaultObjectTest.cpp
@@ -72,6 +72,15 @@ TEST_CASE("UE::CoreUObject::GetDefaultObjectRace", "[CoreUObject][GetDefaultObje
 		CHECK(Task.Wait(FTimespan::FromSeconds(1.0f)));
 		CHECK(NumPostInitProperties != 0);
 	}
+
+	// Move old object out of the way and force recreation
+	if (UClassDefaultObjectTest* Obj = (UClassDefaultObjectTest*)UClassDefaultObjectTest::StaticClass()->GetDefaultObject(false))
+	{
+		Obj->Rename(nullptr, GetTransientPackage());
+	}
+	UClassDefaultObjectTest::StaticClass()->SetDefaultObject(nullptr);
+
+	CollectGarbage(GARBAGE_COLLECTION_KEEPFLAGS);
 }
 
 #endif // WITH_LOW_LEVEL_TESTS
