Avoid TInstancedStruct::GetPtr/GetMutablePtr and TStructView/TConstStructView::GetPtr crashing when the underlying struct data is incompatible with BaseStructT

This will now ensure instead, though will still perform the bad cast in a shipping build, so any ensures need to be fixed!

diff --git a/Engine/Source/Runtime/CoreUObject/Public/StructUtils/InstancedStruct.h b/Engine/Source/Runtime/CoreUObject/Public/StructUtils/InstancedStruct.h
--- a/Engine/Source/Runtime/CoreUObject/Public/StructUtils/InstancedStruct.h
+++ b/Engine/Source/Runtime/CoreUObject/Public/StructUtils/InstancedStruct.h
@@ -428,14 +428,7 @@ public:
 	template<typename T = BaseStructT, typename = std::enable_if_t<std::is_base_of_v<BaseStructT, std::decay_t<T>>>>
 	const T* GetPtr() const
 	{
-		if constexpr (std::is_same_v<BaseStructT, std::decay_t<T>>)
-		{
-			return ((T*)InstancedStruct.GetMemory());
-		}
-		else
-		{
-			return InstancedStruct.GetPtr<T>();
-		}
+		return UE::StructUtils::GetStructPtr<T, BaseStructT>(InstancedStruct.GetScriptStruct(), InstancedStruct.GetMemory());
 	}
 
 	/** Returns a mutable pointer to raw struct memory. */
@@ -455,14 +448,7 @@ public:
 	template<typename T = BaseStructT, typename = std::enable_if_t<std::is_base_of_v<BaseStructT, std::decay_t<T>>>>
 	T* GetMutablePtr()
 	{
-		if constexpr (std::is_same_v<BaseStructT, std::decay_t<T>>)
-		{
-			return ((T*)InstancedStruct.GetMutableMemory());
-		}
-		else
-		{
-			return InstancedStruct.GetMutablePtr<T>();
-		}
+		return UE::StructUtils::GetStructPtr<T, BaseStructT>(InstancedStruct.GetScriptStruct(), InstancedStruct.GetMutableMemory());
 	}
 
 	/** Returns True if the struct is valid.*/
diff --git a/Engine/Source/Runtime/CoreUObject/Public/StructUtils/StructUtils.h b/Engine/Source/Runtime/CoreUObject/Public/StructUtils/StructUtils.h
--- a/Engine/Source/Runtime/CoreUObject/Public/StructUtils/StructUtils.h
+++ b/Engine/Source/Runtime/CoreUObject/Public/StructUtils/StructUtils.h
@@ -34,9 +34,9 @@ namespace UE::StructUtils
 		check(StructMemory != nullptr);
 		check(ScriptStruct != nullptr);
 		checkf(ScriptStruct == TBaseStructure<T>::Get() || ScriptStruct->IsChildOf(TBaseStructure<T>::Get()),
-			TEXT("Incompatible struct types. Source type: '%s' Target type: '%s'"),
-			*(TBaseStructure<T>::Get()->GetName()),
-			ScriptStruct ? *(ScriptStruct->GetName()) : TEXT("nullptr"));
+			TEXT("Incompatible struct types: Cannot cast '%s' to '%s'"),
+			*ScriptStruct->GetPathName(),
+			*TBaseStructure<T>::Get()->GetPathName());
 		return *((T*)StructMemory);
 	}
 
@@ -54,6 +54,33 @@ namespace UE::StructUtils
 		return nullptr;
 	}
 
+	/** Returns pointer to the struct, or nullptr if cast is not valid. */
+	template<typename T, typename BaseStructT>
+	T* GetStructPtr(const UScriptStruct* ScriptStruct, void* StructMemory)
+	{
+		if constexpr (std::is_same_v<BaseStructT, std::decay_t<T>>)
+		{
+#if WITH_EDITOR || !UE_BUILD_SHIPPING
+			if (StructMemory && ScriptStruct)
+			{
+				// An extra guard to verify that ScriptStruct is actually safe to cast to T
+				if (ensureAlwaysMsgf(ScriptStruct == TBaseStructure<T>::Get() || ScriptStruct->IsChildOf(TBaseStructure<T>::Get()),
+					TEXT("Incompatible struct types: Cannot cast '%s' to '%s'. Falling back to return nullptr, but this may crash in shipping!"), *ScriptStruct->GetPathName(), *TBaseStructure<T>::Get()->GetPathName()))
+				{
+					return ((T*)StructMemory);
+				}
+			}
+			return nullptr;
+#else
+			return ((T*)StructMemory);
+#endif
+		}
+		else
+		{
+			return GetStructPtr<T>(ScriptStruct, StructMemory);
+		}
+	}
+
 	/** Returns const reference to the struct, this assumes that all data is valid. */
 	template<typename T>
 	const T& GetStructRef(const UScriptStruct* ScriptStruct, const void* StructMemory)
@@ -68,6 +95,13 @@ namespace UE::StructUtils
 		return GetStructPtr<T>(ScriptStruct, const_cast<void*>(StructMemory));
 	}
 
+	/** Returns const pointer to the struct, or nullptr if cast is not valid. */
+	template<typename T, typename BaseStructT>
+	const T* GetStructPtr(const UScriptStruct* ScriptStruct, const void* StructMemory)
+	{
+		return GetStructPtr<T, BaseStructT>(ScriptStruct, const_cast<void*>(StructMemory));
+	}
+
 	/** 
 	 *  Returns the middle part of an array or view by taking up to the given number of elements from the given position.
 	 *  Based on TArrayView::Mid().
diff --git a/Engine/Source/Runtime/CoreUObject/Public/StructUtils/StructView.h b/Engine/Source/Runtime/CoreUObject/Public/StructUtils/StructView.h
--- a/Engine/Source/Runtime/CoreUObject/Public/StructUtils/StructView.h
+++ b/Engine/Source/Runtime/CoreUObject/Public/StructUtils/StructView.h
@@ -150,14 +150,7 @@ public:
 	template<typename T = BaseStructT, typename = std::enable_if_t<std::is_base_of_v<BaseStructT, std::decay_t<T>>>>
 	T* GetPtr() const
 	{
-		if constexpr (std::is_same_v<BaseStructT, std::decay_t<T>>)
-		{
-			return ((T*)StructMemory);
-		}
-		else
-		{
-			return UE::StructUtils::GetStructPtr<T>(ScriptStruct, StructMemory);
-		}
+		return UE::StructUtils::GetStructPtr<T, BaseStructT>(ScriptStruct, StructMemory);
 	}
 
 	/** Simplified accessor for members of the guaranteed minimum structure type */
@@ -353,14 +346,7 @@ struct [[nodiscard]] TConstStructView : FConstStructView
 	template<typename T = BaseStructT, typename = std::enable_if_t<std::is_base_of_v<BaseStructT, std::decay_t<T>>>>
 	constexpr const T* GetPtr() const
 	{
-		if constexpr (std::is_same_v<BaseStructT, std::decay_t<T>>)
-		{
-			return ((T*)StructMemory);
-		}
-		else
-		{
-			return UE::StructUtils::GetStructPtr<T>(ScriptStruct, StructMemory);
-		}
+		return UE::StructUtils::GetStructPtr<T, BaseStructT>(ScriptStruct, StructMemory);
 	}
 
 	/** Simplified accessor for members of the guaranteed minimum structure type */
