[ObjectVisibility]
- Apply object visibility rules when searching for objects from parallel loading threads

diff --git a/Engine/Source/Runtime/CoreUObject/Private/UObject/ObjectVisibility.cpp b/Engine/Source/Runtime/CoreUObject/Private/UObject/ObjectVisibility.cpp
--- a/Engine/Source/Runtime/CoreUObject/Private/UObject/ObjectVisibility.cpp
+++ b/Engine/Source/Runtime/CoreUObject/Private/UObject/ObjectVisibility.cpp
@@ -17,7 +17,7 @@ namespace UE
 
 	EInternalObjectFlags GetAsyncLoadingInternalFlagsExclusion()
 	{
-		if (IsInAsyncLoadingThread())
+		if (IsInAsyncLoadingThread() || IsInParallelLoadingThread())
 		{
 			return GUseObjectVisibilityFilterForAsyncLoading ? FUObjectThreadContext::Get().AsyncVisibilityFilter : EInternalObjectFlags::None;
 		}
