Add a version of the vtable contructor that includes a namespace.

#rnx

diff --git a/Engine/Source/Runtime/CoreUObject/Public/UObject/ObjectMacros.h b/Engine/Source/Runtime/CoreUObject/Public/UObject/ObjectMacros.h
--- a/Engine/Source/Runtime/CoreUObject/Public/UObject/ObjectMacros.h
+++ b/Engine/Source/Runtime/CoreUObject/Public/UObject/ObjectMacros.h
@@ -1877,8 +1877,11 @@ public: \
 	/** DO NOT USE. This constructor is for internal usage only for hot-reload purposes. */ \
 	API TClass(FVTableHelper& Helper);
 
-#define DEFINE_VTABLE_PTR_HELPER_CTOR(TClass) \
-	TClass::TClass(FVTableHelper& Helper) : Super(Helper) {};
+// Namespace should be in the form of "Whatever::" or blank for no namespace
+#define DEFINE_VTABLE_PTR_HELPER_CTOR_NS(Namespace, TClass) \
+	Namespace TClass::TClass(FVTableHelper& Helper) : Super(Helper) {};
+
+#define DEFINE_VTABLE_PTR_HELPER_CTOR(TClass) DEFINE_VTABLE_PTR_HELPER_CTOR_NS(, TClass)
 
 #define DEFINE_VTABLE_PTR_HELPER_CTOR_CALLER_DUMMY() \
 	static UObject* __VTableCtorCaller(FVTableHelper& Helper) \
