Fix cis

#rnx

diff --git a/Engine/Source/Runtime/CoreUObject/Private/UObject/OverrideEventHelper.cpp b/Engine/Source/Runtime/CoreUObject/Private/UObject/OverrideEventHelper.cpp
--- a/Engine/Source/Runtime/CoreUObject/Private/UObject/OverrideEventHelper.cpp
+++ b/Engine/Source/Runtime/CoreUObject/Private/UObject/OverrideEventHelper.cpp
@@ -7,40 +7,40 @@
 namespace UE
 {
 	// helper class that constructs a change event using FPropertyVisitorPaths
-	class ProxyChangeEvent : public FPropertyChangedChainEvent
+	class ProxyChangeEvent : public FEditPropertyChain, public FPropertyChangedChainEvent
 	{
 	public:
 		ProxyChangeEvent(TNotNull<UObject*> Object, const FPropertyVisitorPath& PropertyPath, EPropertyChangeType::Type ChangeType)
-			: FPropertyChangedChainEvent(PropertyChain, CreateChangeEvent(Object, PropertyPath, ChangeType))
+			: FPropertyChangedChainEvent(*this, CreateChangeEvent(Object, PropertyPath, ChangeType))
 		{
 			ArrayIndicesPerObject.AddDefaulted(ObjectIteratorIndex + 1);
 			for (const FPropertyVisitorInfo& Info : PropertyPath.GetPath())
 			{
-				PropertyChain.AddTail(const_cast<FProperty*>(Info.Property));
+				FEditPropertyChain::AddTail(const_cast<FProperty*>(Info.Property));
 				ArrayIndicesPerObject[ObjectIteratorIndex].Add(Info.Property->GetName(), Info.Index);
 			}
-			SetArrayIndexPerObject(ArrayIndicesPerObject);
+			FPropertyChangedChainEvent::SetArrayIndexPerObject(ProxyChangeEvent::ArrayIndicesPerObject);
 		}
 
 		ProxyChangeEvent(TNotNull<UObject*> Object, EPropertyChangeType::Type ChangeType)
-			: FPropertyChangedChainEvent(PropertyChain, CreateChangeEvent(Object, {}, ChangeType))
+			: FPropertyChangedChainEvent(*this, CreateChangeEvent(Object, {}, ChangeType))
 		{}
 
-		ProxyChangeEvent(TNotNull<UObject*> Object, const FPropertyChangedEvent& PropertyEvent, const FEditPropertyChain& InPropertyChain, EPropertyChangeType::Type ChangeType)
-			: FPropertyChangedChainEvent(PropertyChain, const_cast<FPropertyChangedEvent&>(PropertyEvent))
+		ProxyChangeEvent(TNotNull<UObject*> Object, const FPropertyChangedEvent& PropertyEvent, const FEditPropertyChain& PropertyChain, EPropertyChangeType::Type ChangeType)
+			: FPropertyChangedChainEvent(*this, const_cast<FPropertyChangedEvent&>(PropertyEvent))
 		{
 			FPropertyChangedChainEvent::ChangeType = ChangeType;
 			// The copy constructor of PropertyChain was deleted so copy it manually here.
-			for (FEditPropertyChain::TDoubleLinkedListNode* Itr = InPropertyChain.GetHead(); Itr && Itr->GetValue(); Itr = Itr->GetNextNode())
+			for (FEditPropertyChain::TDoubleLinkedListNode* Itr = PropertyChain.GetHead(); Itr && Itr->GetValue(); Itr = Itr->GetNextNode())
 			{
-				PropertyChain.AddTail(Itr->GetValue());
-				if (InPropertyChain.GetActiveMemberNode() == Itr)
+				FEditPropertyChain::AddTail(Itr->GetValue());
+				if (PropertyChain.GetActiveMemberNode() == Itr)
 				{
-					PropertyChain.SetActiveMemberPropertyNode(PropertyChain.GetTail()->GetValue());
+					FEditPropertyChain::ActiveMemberPropertyNode = FEditPropertyChain::GetTail();
 				}
-				if (InPropertyChain.GetActiveNode() == Itr)
+				if (PropertyChain.GetActiveNode() == Itr)
 				{
-					PropertyChain.SetActivePropertyNode(PropertyChain.GetTail()->GetValue());
+					FEditPropertyChain::ActivePropertyNode = FEditPropertyChain::GetTail();
 				}
 			}
 		}
@@ -55,9 +55,6 @@ namespace UE
 		}
 
 		TArray<TMap<FString, int32>> ArrayIndicesPerObject;
-
-	public:
-		FEditPropertyChain PropertyChain;
 	};
 
 	void SendOverrideObjectEvent(TNotNull<UObject*> Object)
@@ -65,7 +62,7 @@ namespace UE
 #if UE_EDITOR
 		// an event with no provided property path will override the entire object
 		ProxyChangeEvent ChangeEvent(Object, EPropertyChangeType::Unspecified);
-		FOverridableManager::Get().PreOverrideProperty(Object, ChangeEvent.PropertyChain);
+		FOverridableManager::Get().PreOverrideProperty(Object, ChangeEvent);
 		FOverridableManager::Get().PostOverrideProperty(Object, ChangeEvent);
 		FPropertyBagRepository::Get().PostEditChangeChainProperty(Object, ChangeEvent);
 #endif
@@ -76,7 +73,7 @@ namespace UE
 #if UE_EDITOR
 		// a ResetToDefault event with no provided property path will clear all of the object's overrides
 		ProxyChangeEvent ChangeEvent(Object, EPropertyChangeType::ResetToDefault);
-		FOverridableManager::Get().PreOverrideProperty(Object, ChangeEvent.PropertyChain);
+		FOverridableManager::Get().PreOverrideProperty(Object, ChangeEvent);
 		FOverridableManager::Get().PostOverrideProperty(Object, ChangeEvent);
 		FPropertyBagRepository::Get().PostEditChangeChainProperty(Object, ChangeEvent);
 #endif
@@ -86,7 +83,7 @@ namespace UE
 	{
 #if UE_EDITOR
 		ProxyChangeEvent ChangeEvent(Object, PropertyPath, ChangeType);
-		FOverridableManager::Get().PreOverrideProperty(Object, ChangeEvent.PropertyChain);
+		FOverridableManager::Get().PreOverrideProperty(Object, ChangeEvent);
 		FOverridableManager::Get().PostOverrideProperty(Object, ChangeEvent);
 		FPropertyBagRepository::Get().PostEditChangeChainProperty(Object, ChangeEvent);
 #endif
@@ -96,7 +93,7 @@ namespace UE
 	{
 #if UE_EDITOR
 		ProxyChangeEvent ChangeEvent(Object, PropertyEvent, PropertyChain, PropertyEvent.ChangeType);
-		FOverridableManager::Get().PreOverrideProperty(Object, ChangeEvent.PropertyChain);
+		FOverridableManager::Get().PreOverrideProperty(Object, ChangeEvent);
 		FOverridableManager::Get().PostOverrideProperty(Object, ChangeEvent);
 		FPropertyBagRepository::Get().PostEditChangeChainProperty(Object, ChangeEvent);
 #endif
@@ -106,7 +103,7 @@ namespace UE
 	{
 #if UE_EDITOR
 		ProxyChangeEvent ChangeEvent(Object, PropertyPath, EPropertyChangeType::ResetToDefault);
-		FOverridableManager::Get().PreOverrideProperty(Object, ChangeEvent.PropertyChain);
+		FOverridableManager::Get().PreOverrideProperty(Object, ChangeEvent);
 		FOverridableManager::Get().PostOverrideProperty(Object, ChangeEvent);
 		FPropertyBagRepository::Get().PostEditChangeChainProperty(Object, ChangeEvent);
 #endif
@@ -116,7 +113,7 @@ namespace UE
 	{
 #if UE_EDITOR
 		ProxyChangeEvent ChangeEvent(Object, PropertyEvent, PropertyChain, EPropertyChangeType::ResetToDefault);
-		FOverridableManager::Get().PreOverrideProperty(Object, ChangeEvent.PropertyChain);
+		FOverridableManager::Get().PreOverrideProperty(Object, ChangeEvent);
 		FOverridableManager::Get().PostOverrideProperty(Object, ChangeEvent);
 		FPropertyBagRepository::Get().PostEditChangeChainProperty(Object, ChangeEvent);
 #endif
@@ -126,7 +123,7 @@ namespace UE
 	{
 #if UE_EDITOR
 		ProxyChangeEvent ChangeEvent(Object, PropertyPath, EPropertyChangeType::Unspecified);
-		FOverridableManager::Get().PreOverrideProperty(Object, ChangeEvent.PropertyChain);
+		FOverridableManager::Get().PreOverrideProperty(Object, ChangeEvent);
 #endif
 	}
 
