[UObjectArray]
- Fix WITH_INSTRUMENTATION that should have been USING_INSTRUMENTATION

diff --git a/Engine/Source/Runtime/CoreUObject/Private/UObject/UObjectArray.cpp b/Engine/Source/Runtime/CoreUObject/Private/UObject/UObjectArray.cpp
--- a/Engine/Source/Runtime/CoreUObject/Private/UObject/UObjectArray.cpp
+++ b/Engine/Source/Runtime/CoreUObject/Private/UObject/UObjectArray.cpp
@@ -302,7 +302,7 @@ void FUObjectArray::AllocateUObjectIndex(UObjectBase* Object, EInternalObjectFla
 	ObjectItem->ClusterRootIndex = 0;
 
 	// AutoRTFM doesn't like atomics even when relaxed, so we need to differentiate the code here
-#if WITH_INSTRUMENTATION || USING_THREAD_SANITISER
+#if USING_INSTRUMENTATION || USING_THREAD_SANITISER
 	// This can race with weakptr trying to resolve an old object in this slot.
 	// Avoid TSAN warning here since this is safe, the ObjectItem can't possibly match as its been
 	// cleaned up during GC.
diff --git a/Engine/Source/Runtime/CoreUObject/Public/UObject/UObjectArray.h b/Engine/Source/Runtime/CoreUObject/Public/UObject/UObjectArray.h
--- a/Engine/Source/Runtime/CoreUObject/Public/UObject/UObjectArray.h
+++ b/Engine/Source/Runtime/CoreUObject/Public/UObject/UObjectArray.h
@@ -193,7 +193,7 @@ public:
 	FORCEINLINE int32 GetSerialNumber() const
 	{
 		// AutoRTFM doesn't like atomics even when relaxed, so we need to differentiate the code here
-#if WITH_INSTRUMENTATION || USING_THREAD_SANITISER
+#if USING_INSTRUMENTATION || USING_THREAD_SANITISER
 		// Resolving an old WeakPtr can race with new object construction.
 		// Use a relaxed atomic here so that TSAN knows this is considered safe.
 		return FPlatformAtomics::AtomicRead_Relaxed(&SerialNumber);
