- Made a LooseFilePackageStore
- Made a HybridPackageStore that wraps around a cooked PS and the Loose, moving logic into the PS and out of AsyncLoading2.cpp
- No longer allow packages without packagedata in the StorageServerPackageStore
- Reordered code in AsyncLoading2.cpp to use the LooseFIlePackageStore instead of direct file system checks

diff --git a/Engine/Source/Runtime/CoreUObject/Private/Serialization/AsyncLoading2.cpp b/Engine/Source/Runtime/CoreUObject/Private/Serialization/AsyncLoading2.cpp
--- a/Engine/Source/Runtime/CoreUObject/Private/Serialization/AsyncLoading2.cpp
+++ b/Engine/Source/Runtime/CoreUObject/Private/Serialization/AsyncLoading2.cpp
@@ -489,15 +489,6 @@ static FAutoConsoleVariableRef CVarGOnlyProcessRequiredPackagesWhenSyncLoading(
 	ECVF_Default
 );
 
-// In editor builds, for packages that are both local and in zenstore, do we prefer uncooked (local) or cooked (zenstore) packages
-static bool GPreferCookedPackagesInEditor = false;
-static FAutoConsoleVariableRef CVarPreferCookedPackagesInEditor(
-	TEXT("s.PreferCookedPackagesInEditor"),
-	GPreferCookedPackagesInEditor,
-	TEXT("In editor builds, for packages that are both local and in zenstore, do we prefer uncooked (local) or cooked (zenstore) packages"),
-	ECVF_Default
-);
-
 static bool GFailLoadOnNotInstalledImport = true;
 static FAutoConsoleVariableRef CVarGFailLoadOnNotInstalledImport(
 	TEXT("s.FailLoadOnNotInstalledImport"),
@@ -566,6 +557,10 @@ TRACE_DECLARE_ATOMIC_INT_COUNTER(AsyncLoadingLoadingPackages, TEXT("AsyncLoading
 TRACE_DECLARE_ATOMIC_INT_COUNTER(AsyncLoadingPackagesWithRemainingWork, TEXT("AsyncLoading/PackagesWithRemainingWork"));
 TRACE_DECLARE_ATOMIC_INT_COUNTER(AsyncLoadingPendingIoRequests, TEXT("AsyncLoading/PendingIoRequests"));
 TRACE_DECLARE_ATOMIC_MEMORY_COUNTER(AsyncLoadingTotalLoaded, TEXT("AsyncLoading/TotalLoaded"));
+#if ALT2_ENABLE_LINKERLOAD_SUPPORT
+TRACE_DECLARE_ATOMIC_INT_COUNTER(AsyncLoadingFileSystemLoads, TEXT("AsyncLoadingThread/FileSystemLoads"));
+TRACE_DECLARE_ATOMIC_INT_COUNTER(AsyncLoadingIoStoreLoads, TEXT("AsyncLoadingThread/IoStoreLoads"));
+#endif
 
 FString FormatPackageId(FPackageId PackageId)
 {
@@ -722,18 +717,12 @@ struct FPackageRequest
 	}
 };
 
-enum class EPackageLoader : uint8
-{
-	Unknown = 0, // Indicates the package was loaded from somewhere outside of the loader
-	LinkerLoad,
-	IoDispatcher
-};
 const TCHAR* LexToString(EPackageLoader PackageLoader)
 {
 	switch (PackageLoader)
 	{
 	case EPackageLoader::LinkerLoad: return TEXT("LinkerLoad");
-	case EPackageLoader::IoDispatcher: return TEXT("IoDispatcher");
+	case EPackageLoader::Zen: return TEXT("Zen");
 	case EPackageLoader::Unknown: return TEXT("Unknown");
 	default:
 		checkNoEntry();
@@ -4540,62 +4529,62 @@ public:
 	}
 
 private:
-	bool TryGetExistingLoadedPackagePath(FPackageId InPackageIdToLoad, FPackagePath& OutPackagePath, EPackageLoader& OutPackageLoader)
+	bool TryGetExistingLoadedPackagePath(FPackageId InPackageIdToLoad, FName InPackageNameToLoad, EPackageExtension& OutPackageExtension, EPackageLoader& OutPackageLoader, FPackageStoreEntry& OutZenStoreEntry)
 	{
 		// Lookup if we have the package in the importstore already to avoid expensive calls in TryGetPackagePathFromFileSystem
 		if (FLoadedPackageRef* PackageRef = GlobalImportStore.FindPackageRef(InPackageIdToLoad))
 		{
 			// If a package has failed or wasn't loaded from the loader, fail the lookup: we cannot trust these references
-			// and it's possible their actual state is now different than what was cached in the GlobalImportStore
-			if (!PackageRef->HasErrors() && PackageRef->GetPackageLoader() != EPackageLoader::Unknown)
+			// and it's possible their actual state is now different than what was cached in the GlobalImportStore.
+			// Also, if a newer version of the package exists, we don't want to look up in this cache, as we may be going
+			// between loaders (cooked -> uncooked)
+			if (!PackageRef->HasErrors() && PackageRef->GetPackageLoader() != EPackageLoader::Unknown &&
+				(PackageRef->GetPackage() == nullptr || !PackageRef->GetPackage()->HasAnyFlags(RF_NewerVersionExists)))
 			{
 				OutPackageLoader = PackageRef->GetPackageLoader();
-				OutPackagePath.SetHeaderExtension(PackageRef->GetPackageHeaderExtension());
+				OutPackageExtension = PackageRef->GetPackageHeaderExtension();
+
+				if (OutPackageLoader == EPackageLoader::Zen)
+				{
+					// @todo: get the name? it's not used really 
+					PackageStore.GetPackageStoreEntry(InPackageIdToLoad, InPackageNameToLoad, OutZenStoreEntry);
+				}
 				return true;
 			}
 		}
 		return false;
 	}
 
-	bool TryGetPackagePathFromFileSystem(FName& InOutPackageNameToLoad, FName& InOutUPackageName, FPackagePath& InOutPackagePath)
+	EPackageStoreEntryStatus TryGetPackageInfoFromPackageStore(FPackageId& InOutPackageIdToLoad, FName& InOutPackageNameToLoad, FName& InOutUPackageName, FPackageStoreEntry& OutPackageStoreEntry)
 	{
-		// if we prefer cooked packages in the editor, then we don't want to use the linkerload version
-		// NOTE that we won't have fixed up the FNames, but if we aren't using linkerloads, it shouldn't matter
-		if (GPreferCookedPackagesInEditor)
+		EPackageStoreEntryStatus PackageStatus = PackageStore.GetPackageStoreEntry(InOutPackageIdToLoad, InOutPackageNameToLoad, OutPackageStoreEntry);
+
+		if (PackageStatus != EPackageStoreEntryStatus::Missing)
 		{
-			bool bExistsInIoStore = FIoDispatcher::Get().DoesChunkExist(CreatePackageDataChunkId(FPackageId::FromName(InOutPackagePath.GetPackageFName())));
-			if (bExistsInIoStore)
+#if WITH_EDITORONLY_DATA
+			if (OutPackageStoreEntry.LoaderType == EPackageLoader::LinkerLoad)
 			{
-				return false;
+				// in case we are loading loose, update the casing to match the name on disk
+				if (InOutUPackageName == InOutPackageNameToLoad)
+				{
+					InOutUPackageName = OutPackageStoreEntry.LinkerLoadCaseCorrectedPackageName;
+				}
+				InOutPackageNameToLoad = OutPackageStoreEntry.LinkerLoadCaseCorrectedPackageName;
 			}
-		}
-
-#if WITH_EDITORONLY_DATA
-		// In editor, set MatchCaseOnDisk=true so that we set the capitalization of the Package's FName to match the
-		// capitalization on disk. Different capitalizations can arise from imports of the package that were somehow
-		// constructed with a different captialization (most often because the disk captialization changed).
-		// We need the captialization to match so that source control operations in case-significant source control
-		// depots succeed, and to avoid indetermism in the cook.
-		constexpr bool bMatchCaseOnDisk = true;
-#else
-		constexpr bool bMatchCaseOnDisk = false;
 #endif
-
-		// now, check to see if the package exists in the local file system, and if it does, correct the case as needed
-		if (FPackageName::DoesPackageExistEx(InOutPackagePath, FPackageName::EPackageLocationFilter::FileSystem,
-				bMatchCaseOnDisk, &InOutPackagePath) != FPackageName::EPackageLocationFilter::None)
-		{
-			FName CaseCorrectedPackageName = InOutPackagePath.GetPackageFName();
-			if (InOutPackageNameToLoad == InOutUPackageName)
+			if (OutPackageStoreEntry.LoaderType == EPackageLoader::LinkerLoad)
 			{
-				InOutUPackageName = CaseCorrectedPackageName;
+				TRACE_COUNTER_INCREMENT(AsyncLoadingFileSystemLoads);
+			}
+			else
+			{
+				TRACE_COUNTER_INCREMENT(AsyncLoadingIoStoreLoads);
 			}
-			InOutPackageNameToLoad = CaseCorrectedPackageName;
-
-			return true;
 		}
-		return false;
-	}
+
+		return PackageStatus;
+
+	};
 
 #if WITH_EDITOR
 	void ConditionalProcessEditorCallbacks();
@@ -5094,12 +5083,6 @@ void FAsyncLoadingThread2::InitializeLoading()
 		ParsePackageNames(VerbosePackageNamesString, GAsyncLoading2_VerbosePackageIds);
 		ParsePackageNames(DebugPackageNamesString, GAsyncLoading2_VerbosePackageIds);
 		GAsyncLoading2_VerboseLogFilter = GAsyncLoading2_VerbosePackageIds.Num() > 0 ? 1 : 2;
-
-		// force allowing cooked packages in hybrid cooked editor mode
-		if (IsRunningHybridCookedEditor())
-		{
-			GPreferCookedPackagesInEditor = true;
-		}
 	}
 
 	FileOpenLogWrapper = (FPlatformFileOpenLog*)(FPlatformFileManager::Get().FindPlatformFile(FPlatformFileOpenLog::GetTypeName()));
@@ -5510,7 +5493,7 @@ bool FAsyncLoadingThread2::CreateAsyncPackagesFromQueue(FAsyncLoadingThreadState
 		
 			FPackageRequest& Request = OptionalRequest.GetValue();
 			EPackageStoreEntryStatus PackageStatus = EPackageStoreEntryStatus::Missing;
-			EPackageLoader PackageLoader = EPackageLoader::IoDispatcher;
+			EPackageLoader PackageLoader = EPackageLoader::Zen;
 			FPackageStoreEntry PackageEntry;
 			FName PackageNameToLoad = Request.PackagePath.GetPackageFName();
 			TCHAR NameBuffer[FName::StringBufferSize];
@@ -5548,37 +5531,37 @@ bool FAsyncLoadingThread2::CreateAsyncPackagesFromQueue(FAsyncLoadingThreadState
 				}
 			}
 
-			PackageStatus = PackageStore.GetPackageStoreEntry(PackageIdToLoad, UPackageName, PackageEntry);
-			if (PackageStatus == EPackageStoreEntryStatus::Missing)
-			{
-				PackageLoader = EPackageLoader::Unknown;
-				// While there is an active load request for (InName=/Temp/PackageABC_abc, InPackageToLoadFrom=/Game/PackageABC), then allow these requests too:
-				// (InName=/Temp/PackageA_abc, InPackageToLoadFrom=/Temp/PackageABC_abc) and (InName=/Temp/PackageABC_xyz, InPackageToLoadFrom=/Temp/PackageABC_abc)
-				FAsyncPackage2* Package = GetAsyncPackage(PackageIdToLoad);
-				if (Package)
-				{
-					PackageIdToLoad = Package->Desc.PackageIdToLoad;
-					Request.PackagePath = Package->Desc.PackagePathToLoad;
-					PackageNameToLoad = Request.PackagePath.GetPackageFName();
-					PackageStatus = PackageStore.GetPackageStoreEntry(PackageIdToLoad, UPackageName, PackageEntry);
-				}
-			}
-
-			if (!PackageNameToLoad.IsNone() && FAsyncLoadingThreadSettings::Get().bLooseFileLoadingEnabled)
-			{
-				FPackagePath::TryFromPackageName(PackageNameToLoad, Request.PackagePath);
+			FPackagePath::TryFromPackageName(PackageNameToLoad, Request.PackagePath);
+			EPackageExtension Extension;
+		    if (!PackageNameToLoad.IsNone() && TryGetExistingLoadedPackagePath(PackageIdToLoad, PackageNameToLoad, Extension, PackageLoader, PackageEntry))
+		    {
+				Request.PackagePath.SetHeaderExtension(Extension);
+			    PackageStatus = EPackageStoreEntryStatus::Ok;
+		    }
+		    else
+		    {
+			    PackageStatus = TryGetPackageInfoFromPackageStore(PackageIdToLoad, PackageNameToLoad, UPackageName, PackageEntry);
 
-				if (TryGetExistingLoadedPackagePath(PackageIdToLoad, Request.PackagePath, PackageLoader))
+				// try again after looking for an matching active load request
+				if (PackageStatus == EPackageStoreEntryStatus::Missing)
 				{
-					PackageStatus = EPackageStoreEntryStatus::Ok;
+					// While there is an active load request for (InName=/Temp/PackageABC_abc, InPackageToLoadFrom=/Game/PackageABC), then allow these requests too:
+					// (InName=/Temp/PackageA_abc, InPackageToLoadFrom=/Temp/PackageABC_abc) and (InName=/Temp/PackageABC_xyz, InPackageToLoadFrom=/Temp/PackageABC_abc)
+					FAsyncPackage2* Package = GetAsyncPackage(PackageIdToLoad);
+					if (Package)
+					{
+						PackageIdToLoad = Package->Desc.PackageIdToLoad;
+						PackageNameToLoad = Request.PackagePath.GetPackageFName();
+						PackageStatus = TryGetPackageInfoFromPackageStore(PackageIdToLoad, PackageNameToLoad, UPackageName, PackageEntry);
+					}
 				}
-				else if (TryGetPackagePathFromFileSystem(PackageNameToLoad, UPackageName, Request.PackagePath))
+
+				if (PackageStatus != EPackageStoreEntryStatus::Missing)
 				{
-					PackageLoader = EPackageLoader::LinkerLoad;
-					PackageStatus = EPackageStoreEntryStatus::Ok;
+					Request.PackagePath.SetHeaderExtension(PackageEntry.PackageExtension);
+					PackageLoader = PackageEntry.LoaderType;
 				}
 			}
-
 			// Fixup CustomName to handle any input string that can be converted to a long package name.
 			if (!Request.CustomName.IsNone())
 			{
@@ -5636,7 +5619,7 @@ bool FAsyncLoadingThread2::CreateAsyncPackagesFromQueue(FAsyncLoadingThreadState
 						else
 #endif
 						{
-							check(PackageLoader == EPackageLoader::IoDispatcher);
+							check(PackageLoader == EPackageLoader::Zen);
 							InitializeAsyncPackageFromPackageStore(ThreadState, &IoBatch, Package, PackageEntry);
 						}
 						Package->StartLoading(ThreadState, IoBatch);
@@ -6305,6 +6288,9 @@ void FAsyncPackage2::ImportPackagesRecursiveInner(FAsyncLoadingThreadState2& Thr
 	{
 		return;
 	}
+
+	TRACE_CPUPROFILER_EVENT_SCOPE(ImportPackagesRecursiveInner);
+
 	TArray<FName>& ImportedPackageNames = Header.ImportedPackageNames;
 	bool bHasImportedPackageNames = !ImportedPackageNames.IsEmpty();
 	check(ImportedPackageNames.Num() == 0 || ImportedPackageNames.Num() == ImportedPackageCount);
@@ -6312,7 +6298,7 @@ void FAsyncPackage2::ImportPackagesRecursiveInner(FAsyncLoadingThreadState2& Thr
 	{
 		FPackageId ImportedPackageId = ImportedPackageIds[LocalImportedPackageIndex];
 		EPackageStoreEntryStatus ImportedPackageStatus = EPackageStoreEntryStatus::Missing;
-		EPackageLoader ImportedPackageLoader = EPackageLoader::IoDispatcher;
+		EPackageLoader ImportedPackageLoader = EPackageLoader::Unknown;
 		FPackageStoreEntry ImportedPackageEntry;
 		FName ImportedPackageUPackageName = bHasImportedPackageNames ? ImportedPackageNames[LocalImportedPackageIndex] : NAME_None;
 		FName ImportedPackageNameToLoad = ImportedPackageUPackageName;
@@ -6385,31 +6371,34 @@ void FAsyncPackage2::ImportPackagesRecursiveInner(FAsyncLoadingThreadState2& Thr
 				ImportedPackageNameToLoad = NAME_None;
 			}
 		}
-		ImportedPackageStatus = PackageStore.GetPackageStoreEntry(ImportedPackageIdToLoad, ImportedPackageUPackageName, ImportedPackageEntry);
-		if (ImportedPackageStatus == EPackageStoreEntryStatus::Missing)
-		{
-			ImportedPackageLoader = EPackageLoader::Unknown;
-		}
 
 		FPackagePath ImportedPackagePath;
-		EPackageLoader PackageLoader = EPackageLoader::IoDispatcher;
-		if (!ImportedPackageNameToLoad.IsNone() && FAsyncLoadingThreadSettings::Get().bLooseFileLoadingEnabled)
-		{
-			FPackagePath::TryFromPackageName(ImportedPackageNameToLoad, ImportedPackagePath);
+		EPackageLoader PackageLoader = EPackageLoader::Zen;
+		EPackageExtension ImportedPackageExtension = EPackageExtension::Unspecified;
+
+		FPackagePath::TryFromPackageName(ImportedPackageNameToLoad, ImportedPackagePath);
 #if ALT2_ENABLE_LINKERLOAD_SUPPORT
-			PackageLoader = LinkerLoadState.IsSet() ? EPackageLoader::LinkerLoad : PackageLoader;
+		PackageLoader = LinkerLoadState.IsSet() ? EPackageLoader::LinkerLoad : PackageLoader;
 #endif
 
-			if (AsyncLoadingThread.TryGetExistingLoadedPackagePath(ImportedPackageId, ImportedPackagePath, ImportedPackageLoader))
+		if (!ImportedPackageNameToLoad.IsNone() && AsyncLoadingThread.TryGetExistingLoadedPackagePath(ImportedPackageId, ImportedPackageNameToLoad, ImportedPackageExtension, ImportedPackageLoader, ImportedPackageEntry))
+		{
+			ImportedPackageStatus = EPackageStoreEntryStatus::Ok;
+		}
+		else
+		{
+			ImportedPackageStatus = AsyncLoadingThread.TryGetPackageInfoFromPackageStore(ImportedPackageIdToLoad, ImportedPackageNameToLoad, ImportedPackageUPackageName, ImportedPackageEntry);
+			if (ImportedPackageStatus != EPackageStoreEntryStatus::Missing)
 			{
-				ImportedPackageStatus = EPackageStoreEntryStatus::Ok;
+				ImportedPackageExtension = ImportedPackageEntry.PackageExtension;
+				ImportedPackageLoader = ImportedPackageEntry.LoaderType;
 			}
-			else if (AsyncLoadingThread.TryGetPackagePathFromFileSystem(ImportedPackageNameToLoad, ImportedPackageUPackageName, ImportedPackagePath))
+			else
 			{
-				ImportedPackageLoader = EPackageLoader::LinkerLoad;
-				ImportedPackageStatus = EPackageStoreEntryStatus::Ok;
+				ImportedPackageLoader = EPackageLoader::Unknown;
 			}
 		}
+		ImportedPackagePath.SetHeaderExtension(ImportedPackageExtension);
 
 		FLoadedPackageRef& ImportedPackageRef = ImportStore.AddImportedPackageReference(ImportedPackageId, ImportedPackageUPackageName, ImportedPackageLoader, ImportedPackagePath.GetHeaderExtension());
 		checkf(ImportedPackageStatus == EPackageStoreEntryStatus::Missing || ImportedPackageRef.GetPackageLoader() == ImportedPackageLoader, TEXT("ImportPackageStatus(%d) != EPackageStoreEntryStatus::Missing(%d) || ImportedPackageRef.GetPackageLoader()(%s) != ImportedPackageLoader(%s)"), ImportedPackageStatus, EPackageStoreEntryStatus::Missing, LexToString(ImportedPackageRef.GetPackageLoader()), LexToString(ImportedPackageLoader));
@@ -6600,7 +6589,7 @@ void FAsyncPackage2::ImportPackagesRecursiveInner(FAsyncLoadingThreadState2& Thr
 		if (PackageLoader != ImportedPackageLoader)
 		{
 			UE_ASYNC_PACKAGE_LOG(VeryVerbose, Desc, TEXT("ImportPackages: AddDependency"),
-				TEXT("Adding package dependency to %s import '%s'."), ImportedPackageLoader == EPackageLoader::IoDispatcher ? TEXT("cooked") : TEXT("non-cooked"), * ImportedPackage->Desc.UPackageName.ToString());
+				TEXT("Adding package dependency to %s import '%s'."), LexToString(ImportedPackageLoader), *ImportedPackage->Desc.UPackageName.ToString());
 
 			// When importing a linker load package from a zen package or vice versa we need to wait for all the exports in the imported package to be created
 			// and serialized before we can start processing our own exports
@@ -6620,7 +6609,7 @@ void FAsyncPackage2::ImportPackagesRecursiveInner(FAsyncLoadingThreadState2& Thr
 				if (ImportedPackageLoader == EPackageLoader::LinkerLoad)
 				{
 					// Only propagate the instancing context if the imported package is also instanced and it isn't a zen package
-					ImportedPackage->InitializeLinkerLoadState(PackageLoader == EPackageLoader::IoDispatcher || !bIsInstanced ? nullptr : &LinkerLoadState->Linker->GetInstancingContext());
+					ImportedPackage->InitializeLinkerLoadState(PackageLoader == EPackageLoader::Zen || !bIsInstanced ? nullptr : &LinkerLoadState->Linker->GetInstancingContext());
 				}
 				else
 #endif
diff --git a/Engine/Source/Runtime/CoreUObject/Private/Serialization/PackageStore.cpp b/Engine/Source/Runtime/CoreUObject/Private/Serialization/PackageStore.cpp
--- a/Engine/Source/Runtime/CoreUObject/Private/Serialization/PackageStore.cpp
+++ b/Engine/Source/Runtime/CoreUObject/Private/Serialization/PackageStore.cpp
@@ -6,14 +6,173 @@
 #include "Containers/StringView.h"
 #include "Containers/UnrealString.h"
 #include "CoreGlobals.h"
+#include "CoreGlobalsInternal.h"
 #include "HAL/Platform.h"
 #include "IO/PackageId.h"
 #include "Misc/AssertionMacros.h"
+#include "Misc/ConfigCacheIni.h"
+#include "Misc/CoreDelegates.h"
+#include "Misc/PackageName.h"
 #include "Serialization/Archive.h"
+#include "Serialization/AsyncPackageLoader.h"
 #include "Serialization/CompactBinary.h"
 #include "Serialization/CompactBinaryWriter.h"
 #include "Templates/Greater.h"
 
+
+///////////////////////////////////////
+// FLinkerLoadPackageStoreBackend
+///////////////////////////////////////
+
+class FLinkerLoadPackageStoreBackend final : public IPackageStoreBackend
+{
+	EPackageLoader GetSupportedLoaders()
+	{
+		return EPackageLoader::LinkerLoad;
+	}
+
+	virtual void OnMounted(TSharedRef<const FPackageStoreBackendContext>) override
+	{
+	}
+
+	virtual void BeginRead() override
+	{
+	}
+
+	virtual void EndRead() override
+	{
+	}
+
+	virtual EPackageStoreEntryStatus GetPackageStoreEntry(FPackageId PackageId, FName PackageName, FPackageStoreEntry& OutPackageStoreEntry) override
+	{
+		FPackageStoreEntry* Entry = PackageStoreCache.Find(PackageId);
+		if (Entry)
+		{
+			if (Entry->LoaderType == EPackageLoader::LinkerLoad)
+			{
+				OutPackageStoreEntry = *Entry;
+				return EPackageStoreEntryStatus::Ok;
+			}
+			return EPackageStoreEntryStatus::Missing;
+		}
+
+
+#if WITH_EDITORONLY_DATA
+		// In editor, set MatchCaseOnDisk=true so that we set the capitalization of the Package's FName to match the
+		// capitalization on disk. Different capitalizations can arise from imports of the package that were somehow
+		// constructed with a different captialization (most often because the disk captialization changed).
+		// We need the captialization to match so that source control operations in case-significant source control
+		// depots succeed, and to avoid indetermism in the cook.
+		constexpr bool bMatchCaseOnDisk = true;
+#else
+		constexpr bool bMatchCaseOnDisk = false;
+#endif
+
+		FPackagePath PackagePath = FPackagePath::FromPackageNameUnchecked(PackageName);
+		// now, check to see if the package exists in the local file system, and if it does, correct the case as needed
+		if (FPackageName::DoesPackageExistEx(PackagePath, FPackageName::EPackageLocationFilter::FileSystem,
+			bMatchCaseOnDisk, &PackagePath) != FPackageName::EPackageLocationFilter::None)
+		{
+			OutPackageStoreEntry.LoaderType = EPackageLoader::LinkerLoad;
+			OutPackageStoreEntry.PackageExtension = PackagePath.GetHeaderExtension();
+			check(OutPackageStoreEntry.PackageExtension != EPackageExtension::Unspecified);
+
+#if WITH_EDITORONLY_DATA
+			OutPackageStoreEntry.LinkerLoadCaseCorrectedPackageName = PackagePath.GetPackageFName();
+#endif
+
+			PackageStoreCache.Add(PackageId, OutPackageStoreEntry);
+			return EPackageStoreEntryStatus::Ok;
+		}
+
+		PackageStoreCache.Add(PackageId, FPackageStoreEntry());
+		return EPackageStoreEntryStatus::Missing;
+	}
+
+	virtual bool GetPackageRedirectInfo(FPackageId PackageId, FName& OutSourcePackageName, FPackageId& OutRedirectedToPackageId) override
+	{
+		return false;
+	}
+
+private:
+	TMap<FPackageId, FPackageStoreEntry> PackageStoreCache;
+};
+
+
+///////////////////////////////////////
+// FHybridPackageStoreBackend
+///////////////////////////////////////
+
+
+namespace
+{
+	TSet<FPackageId> ForcedLoosePackages;
+}
+
+FHybridPackageStoreBackend::FHybridPackageStoreBackend(TSharedPtr<IPackageStoreBackend> InLoosePackageStore, TSharedPtr<IPackageStoreBackend> InCookedPackageStore)
+: LoosePackageStore(InLoosePackageStore)
+, CookedPackageStore(InCookedPackageStore)
+{
+}
+
+EPackageLoader FHybridPackageStoreBackend::GetSupportedLoaders()
+{
+	return EPackageLoader::LinkerLoad | EPackageLoader::Zen;
+}
+
+void FHybridPackageStoreBackend::OnMounted(TSharedRef<const FPackageStoreBackendContext>)
+{
+
+}
+
+void FHybridPackageStoreBackend::BeginRead()
+{
+	LoosePackageStore->BeginRead();
+	CookedPackageStore->BeginRead();
+}
+
+void FHybridPackageStoreBackend::EndRead()
+{
+	LoosePackageStore->EndRead();
+	CookedPackageStore->EndRead();
+}
+
+bool FHybridPackageStoreBackend::GetPackageRedirectInfo(FPackageId PackageId, FName& OutSourcePackageName, FPackageId& OutRedirectedToPackageId)
+{
+	return CookedPackageStore->GetPackageRedirectInfo(PackageId, OutSourcePackageName, OutRedirectedToPackageId);
+}
+
+
+EPackageStoreEntryStatus FHybridPackageStoreBackend::GetPackageStoreEntry(FPackageId PackageId, FName PackageName, FPackageStoreEntry& OutPackageStoreEntry)
+{
+	EPackageStoreEntryStatus InnerStatus = EPackageStoreEntryStatus::Missing;
+
+	if (!ForcedLoosePackages.Contains(PackageId))
+	{
+		InnerStatus = CookedPackageStore->GetPackageStoreEntry(PackageId, PackageName, OutPackageStoreEntry);
+	}
+
+	if (InnerStatus == EPackageStoreEntryStatus::Missing && FAsyncLoadingThreadSettings::Get().bLooseFileLoadingEnabled)
+	{
+		// reset it in case the cooked store partially filled this out
+		OutPackageStoreEntry = {};
+		InnerStatus = LoosePackageStore->GetPackageStoreEntry(PackageId, PackageName, OutPackageStoreEntry);
+	}
+
+	return InnerStatus;
+}
+
+void FHybridPackageStoreBackend::ForceLoadPackageAsLoose(FPackageId PackageId)
+{
+	ForcedLoosePackages.Add(PackageId);
+}
+
+
+///////////////////////////////////////
+// FPackageStore
+///////////////////////////////////////
+
+
 FArchive& operator<<(FArchive& Ar, FPackageStoreEntryResource& PackageStoreEntry)
 {
 	uint32 Flags = static_cast<uint32>(PackageStoreEntry.Flags);
@@ -168,7 +327,37 @@ FPackageStoreReadScope::~FPackageStoreReadScope()
 FPackageStore::FPackageStore()
 	: BackendContext(MakeShared<FPackageStoreBackendContext>())
 {
-	
+	auto SetupLinkerLoadBackend = [this]()
+	{
+		// if loose file loading is enabled, make the package store that manages looking on disk for loose packages (cooked or uncooked)
+		if (FAsyncLoadingThreadSettings::Get().bLooseFileLoadingEnabled)
+		{
+			LinkerLoadBackend = MakeShared<FLinkerLoadPackageStoreBackend>();
+			if (LinkerLoadBackend)
+			{
+				bool bLooseTakesPrecedence = false;
+				// non-hybrid editors prefer loose over cooked, but editor -game would prefer cooked, if there was a way for
+				// it to load cooked data without hybred
+				if (GIsEditor)
+				{
+					bLooseTakesPrecedence = true;
+				}
+				// since most package stores are prio 0, we use -1 or 1 to put the loose store above or below the default,
+				// depending on which order we prefer
+				Mount(LinkerLoadBackend.ToSharedRef(), bLooseTakesPrecedence ? 1 : -1);
+			}
+		}
+	};
+
+	if (GConfig && GConfig->IsReadyForUse())
+	{
+		SetupLinkerLoadBackend();
+	}
+	else
+	{
+		// delay setting up the loose loader untyil GConfig exists
+		FCoreDelegates::TSConfigReadyForUse().AddLambda(SetupLinkerLoadBackend);
+	}
 }
 
 FPackageStore& FPackageStore::Get()
@@ -181,6 +370,28 @@ FPackageStore& FPackageStore::Get()
 void FPackageStore::Mount(TSharedRef<IPackageStoreBackend> Backend, int32 Priority)
 {
 	check(IsInGameThread());
+
+#if WITH_EDITOR
+	// delay adding the hybrid backend until we get a second package store to wrap along with
+	// the loose package store
+	if (IsRunningHybridCookedEditor())
+	{
+		if (!HybridBackend && LinkerLoadBackend && !!(Backend->GetSupportedLoaders() & EPackageLoader::Zen))
+		{
+			// remove the loose backend
+			Backends.RemoveAll([this](const FBackendAndPriority& Pair) { return Pair.Value == LinkerLoadBackend; });
+
+			// make a new hybrid backend which wraps the existing loose and the new backend 
+			HybridBackend = MakeShared<FHybridPackageStoreBackend>(LinkerLoadBackend, Backend);
+			Mount(HybridBackend.ToSharedRef(), Priority);
+
+			// call callbacks 
+			Backend->OnMounted(BackendContext);
+			return;
+		}
+	}
+#endif
+
 	int32 Index = Algo::LowerBoundBy(Backends, Priority, &FBackendAndPriority::Key, TGreater<>());
 	Backends.Insert(MakeTuple(Priority, Backend), Index);
 	Backend->OnMounted(BackendContext);
@@ -199,6 +410,7 @@ EPackageStoreEntryStatus FPackageStore::GetPackageStoreEntry(FPackageId PackageI
 
 		if (!bContinueSearch)
 		{
+			checkSlow((Backend.Value->GetSupportedLoaders() & OutPackageStoreEntry.LoaderType));
 			return Status;
 		}
 	}
diff --git a/Engine/Source/Runtime/CoreUObject/Public/Serialization/PackageStore.h b/Engine/Source/Runtime/CoreUObject/Public/Serialization/PackageStore.h
--- a/Engine/Source/Runtime/CoreUObject/Public/Serialization/PackageStore.h
+++ b/Engine/Source/Runtime/CoreUObject/Public/Serialization/PackageStore.h
@@ -9,6 +9,7 @@
 #include "Delegates/Delegate.h"
 #include "HAL/PlatformCrt.h"
 #include "Misc/EnumClassFlags.h"
+#include "Misc/PackagePath.h"
 #include "Misc/SecureHash.h"
 #include "IO/PackageId.h"
 #include "Templates/SharedPointer.h"
@@ -19,6 +20,17 @@ class FCbObjectView;
 class FCbWriter;
 class FStructuredArchiveSlot;
 
+/**
+ * The type of loader that a package store entry needs to load with
+ */
+enum class EPackageLoader : uint8
+{
+	Unknown		= 0x0, // Indicates the package was loaded from somewhere outside of the loader
+	LinkerLoad	= 0x1,
+	Zen			= 0x2,
+};
+ENUM_CLASS_FLAGS(EPackageLoader);
+
 /**
  * Package store entry status.
  */
@@ -36,6 +48,18 @@ enum class EPackageStoreEntryStatus
  */ 
 struct FPackageStoreEntry
 {
+#if WITH_EDITORONLY_DATA
+	// the corrected package name, which may be a different case than PackageName that's passed in to GetPackageStoreEntry
+	// only applies to LinkerLoad (loose on the local filesystem) zen packages are never submitted to source control
+	FName LinkerLoadCaseCorrectedPackageName;
+#endif
+
+	// for LinkerLoad entries, this will contain the extension of the package. for Zen entries, this will stay Unspecified
+	EPackageExtension PackageExtension = EPackageExtension::Unspecified;
+
+	// default to IoDispatcher, since almost all package store entries will use it
+	EPackageLoader LoaderType = EPackageLoader::Zen;
+
 	TArrayView<const FPackageId> ImportedPackageIds;
 	TArrayView<const FSHAHash> ShaderMapHashes;
 #if WITH_EDITOR
@@ -167,6 +191,12 @@ public:
 	/* Destructor. */
 	virtual ~IPackageStoreBackend() { }
 
+	/** Returns what possible loader types are supported by this PackageStore backend */
+	virtual EPackageLoader GetSupportedLoaders()
+	{
+		return EPackageLoader::Zen;
+	}
+
 	/** Called when the backend is mounted */
 	virtual void OnMounted(TSharedRef<const FPackageStoreBackendContext> Context) = 0;
 
@@ -225,6 +255,11 @@ private:
 	using FBackendAndPriority = TTuple<int32, TSharedRef<IPackageStoreBackend>>;
 	TArray<FBackendAndPriority> Backends;
 
+	TSharedPtr<IPackageStoreBackend> LinkerLoadBackend;
+#if WITH_EDITOR
+	TSharedPtr<IPackageStoreBackend> HybridBackend;
+#endif
+
 	static thread_local int32 ThreadReadCount;
 };
 
@@ -237,3 +272,27 @@ public:
 private:
 	FPackageStore& PackageStore;
 };
+
+
+class FHybridPackageStoreBackend final : public IPackageStoreBackend
+{
+public:
+	FHybridPackageStoreBackend(TSharedPtr<IPackageStoreBackend> InLoosePackageStore, TSharedPtr<IPackageStoreBackend> InCookedPackageStore);
+
+	virtual EPackageLoader GetSupportedLoaders() override;
+	virtual void OnMounted(TSharedRef<const FPackageStoreBackendContext>) override;
+	virtual void BeginRead() override;
+	virtual void EndRead() override;
+	virtual bool GetPackageRedirectInfo(FPackageId PackageId, FName& OutSourcePackageName, FPackageId& OutRedirectedToPackageId) override;
+		virtual EPackageStoreEntryStatus GetPackageStoreEntry(FPackageId PackageId, FName PackageName, FPackageStoreEntry& OutPackageStoreEntry) override;
+
+	/**
+	 * Used by the editor to tell a package to stop loading the cooked version.
+	 * This can be used to "uncook" an asset to allow for editing 
+	 */
+	COREUOBJECT_API static void ForceLoadPackageAsLoose(FPackageId PackageId);
+
+private:
+	TSharedPtr<IPackageStoreBackend> LoosePackageStore;
+	TSharedPtr<IPackageStoreBackend> CookedPackageStore;
+};
diff --git a/Engine/Source/Runtime/StorageServerClient/Private/StorageServerPackageStore.cpp b/Engine/Source/Runtime/StorageServerClient/Private/StorageServerPackageStore.cpp
--- a/Engine/Source/Runtime/StorageServerClient/Private/StorageServerPackageStore.cpp
+++ b/Engine/Source/Runtime/StorageServerClient/Private/StorageServerPackageStore.cpp
@@ -105,13 +105,18 @@ uint32 FStorageServerPackageStoreBackend::FAsyncInitRunnable::Run()
 	TRACE_CPUPROFILER_EVENT_SCOPE(StorageServerPackageStoreRequest);
 	Connection.PackageStoreRequest([this](FPackageStoreEntryResource&& PackageStoreEntryResource)
 		{
-			FStoreEntry& LocalEntry = Owner.StoreEntriesMap.FindOrAdd(PackageStoreEntryResource.GetPackageId());
-			LocalEntry.ImportedPackages = MoveTemp(PackageStoreEntryResource.ImportedPackageIds);
+			// if there's no package data, we do not want to report that this package is in the PackageStore - these entries
+			// can exist in the ZenStore for incremental cooking purposes, but are not real usable packages in the store
+			if (!!(PackageStoreEntryResource.Flags & EPackageStoreEntryFlags::HasPackageData))
+			{
+				FStoreEntry& LocalEntry = Owner.StoreEntriesMap.FindOrAdd(PackageStoreEntryResource.GetPackageId());
+				LocalEntry.ImportedPackages = MoveTemp(PackageStoreEntryResource.ImportedPackageIds);
 #if WITH_EDITOR
-			LocalEntry.OptionalSegmentImportedPackageIds = MoveTemp(PackageStoreEntryResource.OptionalSegmentImportedPackageIds);
-			LocalEntry.Flags = PackageStoreEntryResource.Flags;
+				LocalEntry.OptionalSegmentImportedPackageIds = MoveTemp(PackageStoreEntryResource.OptionalSegmentImportedPackageIds);
+				LocalEntry.Flags = PackageStoreEntryResource.Flags;
 #endif
-			LocalEntry.ShaderMapHashes = MoveTemp(PackageStoreEntryResource.ShaderMapHashes);
+				LocalEntry.ShaderMapHashes = MoveTemp(PackageStoreEntryResource.ShaderMapHashes);
+			}
 		});
 
 	IsCompleted->Trigger();
