[LinkerLoad]
- Add a mutex and protect LinkerLoad against multithreaded access while manipulating its flags or serializing from it

diff --git a/Engine/Source/Runtime/CoreUObject/Private/Blueprint/BlueprintSupport.cpp b/Engine/Source/Runtime/CoreUObject/Private/Blueprint/BlueprintSupport.cpp
--- a/Engine/Source/Runtime/CoreUObject/Private/Blueprint/BlueprintSupport.cpp
+++ b/Engine/Source/Runtime/CoreUObject/Private/Blueprint/BlueprintSupport.cpp
@@ -15,6 +15,7 @@
 #include "Serialization/DuplicatedDataWriter.h"
 #include "Misc/PackageAccessTrackingOps.h"
 #include "Misc/PackageName.h"
+#include "Async/UniqueLock.h"
 #include "UObject/ObjectResource.h"
 #include "UObject/GCObject.h"
 #include "UObject/LinkerLoad.h"
@@ -1487,6 +1488,7 @@ UPackage* LoadPackageInternal(UPackage* InOuter, const FPackagePath& PackagePath
 void FLinkerLoad::ResolveDeferredDependencies(UStruct* LoadStruct)
 {
 #if USE_CIRCULAR_DEPENDENCY_LOAD_DEFERRING
+	UE::TUniqueLock ScopeLock(Mutex);
 	//--------------------------------------
 	// Phase 2: Resolve Dependency Stubs
 	//--------------------------------------
diff --git a/Engine/Source/Runtime/CoreUObject/Private/UObject/LinkerLoad.cpp b/Engine/Source/Runtime/CoreUObject/Private/UObject/LinkerLoad.cpp
--- a/Engine/Source/Runtime/CoreUObject/Private/UObject/LinkerLoad.cpp
+++ b/Engine/Source/Runtime/CoreUObject/Private/UObject/LinkerLoad.cpp
@@ -4643,6 +4643,11 @@ void FLinkerLoad::Preload( UObject* Object )
 	//check(IsValidLowLevel());
 	check(Object);
 
+	// This is important to take the lock before checking the RF_NeedLoad flag in case
+	// Preload is called on the same objects from 2 different threads, one might see the RF_NeedLoad and remove it
+	// before the other, making one thread return from Preload before the object is fully serialized.
+	UE::TUniqueLock UniqueLock(Mutex);
+
 	// Preload the object if necessary.
 	if (Object->HasAnyFlags(RF_NeedLoad))
 	{
@@ -5034,6 +5039,8 @@ void FLinkerLoad::Preload( UObject* Object )
 		else if (FLinkerLoad* Linker = Object->GetLinker())
 		{
 #if USE_CIRCULAR_DEPENDENCY_LOAD_DEFERRING
+			// We need to take the lock before changing the flags
+			UE::TUniqueLock UniqueLockInner(Linker->Mutex);
 			uint32 const DeferredLoadFlag = (LoadFlags & LOAD_DeferDependencyLoads);
 			TGuardValue<uint32> LoadFlagsGuard(Linker->LoadFlags, Linker->LoadFlags | DeferredLoadFlag);
 #endif
diff --git a/Engine/Source/Runtime/CoreUObject/Private/UObject/PropertyPrivate.cpp b/Engine/Source/Runtime/CoreUObject/Private/UObject/PropertyPrivate.cpp
--- a/Engine/Source/Runtime/CoreUObject/Private/UObject/PropertyPrivate.cpp
+++ b/Engine/Source/Runtime/CoreUObject/Private/UObject/PropertyPrivate.cpp
@@ -805,6 +805,7 @@ static inline void PreloadInnerStructMembers(UStructProperty* StructProperty)
 	{
 		if (FLinkerLoad* StructLinker = Struct->GetLinker())
 		{
+			UE::TUniqueLock ScopeLock(StructLinker->Mutex);
 			TGuardValue<uint32> LoadFlagGuard(StructLinker->LoadFlags, StructLinker->LoadFlags | PropagatedLoadFlags);
 			Struct->RecursivelyPreload();
 		}
diff --git a/Engine/Source/Runtime/CoreUObject/Private/UObject/PropertyStruct.cpp b/Engine/Source/Runtime/CoreUObject/Private/UObject/PropertyStruct.cpp
--- a/Engine/Source/Runtime/CoreUObject/Private/UObject/PropertyStruct.cpp
+++ b/Engine/Source/Runtime/CoreUObject/Private/UObject/PropertyStruct.cpp
@@ -26,6 +26,7 @@ static inline void PreloadInnerStructMembers(FStructProperty* StructProperty)
 		{
 			if (FLinkerLoad* StructLinker = Struct->GetLinker())
 			{
+				UE::TUniqueLock ScopeLock(StructLinker->Mutex);
 				TGuardValue<uint32> LoadFlagGuard(StructLinker->LoadFlags, StructLinker->LoadFlags | PropagatedLoadFlags);
 				Struct->RecursivelyPreload();
 			}
@@ -295,8 +296,10 @@ const TCHAR* FStructProperty::ImportText_Internal(const TCHAR* InBuffer, void* C
 
 	uint32 OldFlags = 0;
 	FLinkerLoad* StructLinker = Struct->GetLinker();
+	TOptional<UE::TUniqueLock<UE::FRecursiveMutex>> OptionalMutex;
 	if (StructLinker)
 	{
+		OptionalMutex.Emplace(StructLinker->Mutex);
 		OldFlags = StructLinker->LoadFlags;
 		StructLinker->LoadFlags |= OldFlags | PropagatedLoadFlags;
 	}
diff --git a/Engine/Source/Runtime/CoreUObject/Public/UObject/LinkerLoad.h b/Engine/Source/Runtime/CoreUObject/Public/UObject/LinkerLoad.h
--- a/Engine/Source/Runtime/CoreUObject/Public/UObject/LinkerLoad.h
+++ b/Engine/Source/Runtime/CoreUObject/Public/UObject/LinkerLoad.h
@@ -3,6 +3,7 @@
 #pragma once
 
 #include "Async/Mutex.h"
+#include "Async/RecursiveMutex.h"
 #include "Containers/Array.h"
 #include "Containers/Map.h"
 #include "Containers/Set.h"
@@ -170,6 +171,9 @@ public:
 	bool					bForceSimpleIndexToObject;
 	bool					bLockoutLegacyOperations;
 
+	// Used to protect LinkerLoad from multithreaded access
+	UE::FRecursiveMutex     Mutex;
+
 	/** True if Loader is FAsyncArchive  */
 	bool					bIsAsyncLoader;
 	FORCEINLINE FAsyncArchive* GetAsyncLoader()
