[GCObject]
- Make the delegate thread-safe as its called from ALT

diff --git a/Engine/Source/Runtime/CoreUObject/Public/UObject/GCObject.h b/Engine/Source/Runtime/CoreUObject/Public/UObject/GCObject.h
--- a/Engine/Source/Runtime/CoreUObject/Public/UObject/GCObject.h
+++ b/Engine/Source/Runtime/CoreUObject/Public/UObject/GCObject.h
@@ -87,7 +87,7 @@ public:
 
 #if WITH_EDITORONLY_DATA
 	/** Called when a new FGCObject is added to this referencer */
-	DECLARE_MULTICAST_DELEGATE_OneParam(FGCObjectAddedDelegate, FGCObject*);
+	DECLARE_TS_MULTICAST_DELEGATE_OneParam(FGCObjectAddedDelegate, FGCObject*);
 
 private:
 
