[UObject]
- Add a ConditionalPreload function to remove duplicated code and improve the pattern it uses

diff --git a/Engine/Source/Runtime/CoreUObject/Private/UObject/Obj.cpp b/Engine/Source/Runtime/CoreUObject/Private/UObject/Obj.cpp
--- a/Engine/Source/Runtime/CoreUObject/Private/UObject/Obj.cpp
+++ b/Engine/Source/Runtime/CoreUObject/Private/UObject/Obj.cpp
@@ -1376,6 +1376,20 @@ void UObject::ConditionalPostLoad()
 	}
 }
 
+void UObject::ConditionalPreload()
+{
+	// The check for EInternalObjectFlags::AsyncLoadingPhase1 is important here to make sure we go
+	// through the Preload function to get rid of any potential race condition that might happen
+	// with the parallel loader threads when going wide.
+	if (HasAnyFlags(RF_NeedLoad) || HasAnyInternalFlags(EInternalObjectFlags::AsyncLoadingPhase1))
+	{
+		if (FLinkerLoad* LinkerLoad = GetLinker())
+		{
+			LinkerLoad->Preload(this);
+		}
+	}
+}
+
 void UObject::PostLoadSubobjects( FObjectInstancingGraph* OuterInstanceGraph/*=NULL*/ )
 {
 	// if this class contains instanced object properties and a new object property has been added since this object was saved,
diff --git a/Engine/Source/Runtime/CoreUObject/Public/UObject/Object.h b/Engine/Source/Runtime/CoreUObject/Public/UObject/Object.h
--- a/Engine/Source/Runtime/CoreUObject/Public/UObject/Object.h
+++ b/Engine/Source/Runtime/CoreUObject/Public/UObject/Object.h
@@ -1179,6 +1179,9 @@ public:
 	/** PostLoad if needed. */
 	COREUOBJECT_API void ConditionalPostLoad();
 
+	/** Preload if needed. Mostly used in the loader or loader threads. */
+	COREUOBJECT_API void ConditionalPreload();
+
 	/**
 	 * Instances subobjects and components for objects being loaded from disk, if necessary.  Ensures that references
 	 * between nested components are fixed up correctly.
