Updated TouchResidentObject system to use Sets and Maps instead of arrays as containers to improve search performance.

These arrays were used in multiple places to find specific objects, and the linear search becomes too costly as the number of objects per server grow.

diff --git a/Engine/Source/Runtime/CoreUObject/Private/UObject/RemoteObjectTransfer.cpp b/Engine/Source/Runtime/CoreUObject/Private/UObject/RemoteObjectTransfer.cpp
--- a/Engine/Source/Runtime/CoreUObject/Private/UObject/RemoteObjectTransfer.cpp
+++ b/Engine/Source/Runtime/CoreUObject/Private/UObject/RemoteObjectTransfer.cpp
@@ -51,15 +51,16 @@ public:
 	FRemoteTransactionId RequestId;
 	FRemoteWorkPriority Priority;
 
-	// this array fills up with object ids that we have ever touched
+	// this set fills up with object ids that we have ever touched
 	// while running the transaction
-	TArray<FRemoteObjectId> RequiredObjects;
-	TArray<FRemoteObjectId> RequiredObjectsCanonicalRoot;
+	TSet<FRemoteObjectId> RequiredObjects;
+	TMap<FRemoteObjectId, FRemoteObjectId> RequiredObjectsCanonicalRootByObjectID;
+	TSet<FRemoteObjectId> RequiredObjectsCanonicalRoots;
 
-	// this array gets cleared before running the 
+	// this set gets cleared before running the 
 	// transaction and only tracks the objects touched 
 	// during the most recent run
-	TArray<FRemoteObjectId> UsedObjects;
+	TSet<FRemoteObjectId> UsedObjects;
 
 	FRemoteTransactionId GetRequestId() const { return RequestId; }
 	FRemoteWorkPriority GetPriority() const { return Priority; }
@@ -170,18 +171,7 @@ public:
 					if (IsHigherPriority(ExistingRequest->GetPriority(), HighestPriorityRequest->RequestPriority))
 					{
 						// ExistingRequest is higher priority, now see if it actually needs this object
-						for (int32 RequiredObjectIndex = 0; RequiredObjectIndex < ExistingRequest->RequiredObjectsCanonicalRoot.Num(); RequiredObjectIndex++)
-						{
-							FRemoteObjectId RequiredObjectId = ExistingRequest->RequiredObjectsCanonicalRoot[RequiredObjectIndex];
-
-							if (RequiredObjectId == CanonicalRootObjectId)
-							{
-								// we found an existing request that is higher priority than the highest priority request,
-								// so we have to give it up
-								bObjectLocked = true;
-								break;
-							}
-						}
+						bObjectLocked = ExistingRequest->RequiredObjectsCanonicalRoots.Contains(CanonicalRootObjectId);
 					}
 
 					if (bObjectLocked)
@@ -331,25 +321,26 @@ public:
 		using namespace UE::RemoteObject::Handle;
 		using namespace UE::RemoteObject::Serialization;
 
-		check(ActiveRequest->RequiredObjects.Num() == ActiveRequest->RequiredObjectsCanonicalRoot.Num());
+		check(ActiveRequest->RequiredObjects.Num() == ActiveRequest->RequiredObjectsCanonicalRootByObjectID.Num());
 
 		// reset the list of used objects each time before we run
-		ActiveRequest->UsedObjects.SetNum(0, EAllowShrinking::No);
+		ActiveRequest->UsedObjects.Reset();
 
 		UE_LOG(LogRemoteObject, VeryVerbose, TEXT("FRemoteObjectTransferQueue: TickRequest(%s) %d required objs (%d pending requests):"), *ActiveRequest->GetRequestId().ToString(), ActiveRequest->RequiredObjects.Num(), PendingObjectRequests.Num());
-		for (int32 RequiredObjectIndex = 0; RequiredObjectIndex < ActiveRequest->RequiredObjects.Num(); RequiredObjectIndex++)
+
+		if (UE_GET_LOG_VERBOSITY(LogRemoteObject) == ELogVerbosity::VeryVerbose)
 		{
-			FRemoteObjectId RequiredObjectId = ActiveRequest->RequiredObjects[RequiredObjectIndex];
-			FRemoteObjectId CanonicalObjectId = ActiveRequest->RequiredObjectsCanonicalRoot[RequiredObjectIndex];
-			UE_LOG(LogRemoteObject, VeryVerbose, TEXT("FRemoteObjectTransferQueue: TickRequest(%s) %s [root %s] %s"), *ActiveRequest->GetRequestId().ToString(), *RequiredObjectId.ToString(), *CanonicalObjectId.ToString(), Handle::IsRemote(RequiredObjectId) ? TEXT("<remote>") : TEXT("<local>"));
+			for (const FRemoteObjectId& RequiredObjectId : ActiveRequest->RequiredObjects)
+			{
+				FRemoteObjectId CanonicalObjectId = ActiveRequest->RequiredObjectsCanonicalRootByObjectID.FindChecked(RequiredObjectId);
+				UE_LOG(LogRemoteObject, VeryVerbose, TEXT("FRemoteObjectTransferQueue: TickRequest(%s) %s [root %s] %s"), *ActiveRequest->GetRequestId().ToString(), *RequiredObjectId.ToString(), *CanonicalObjectId.ToString(), Handle::IsRemote(RequiredObjectId) ? TEXT("<remote>") : TEXT("<local>"));
+			}
 		}
 
 		// send the requests
-		for (int32 RequiredObjectIndex = 0; RequiredObjectIndex < ActiveRequest->RequiredObjects.Num(); RequiredObjectIndex++)
+		for (const FRemoteObjectId& RequiredObjectId : ActiveRequest->RequiredObjects)
 		{
-			FRemoteObjectId RequiredObjectId = ActiveRequest->RequiredObjects[RequiredObjectIndex];
-			FRemoteObjectId RequiredCanonicalRootObjectId = ActiveRequest->RequiredObjectsCanonicalRoot[RequiredObjectIndex];
-			
+			FRemoteObjectId RequiredCanonicalRootObjectId = ActiveRequest->RequiredObjectsCanonicalRootByObjectID.FindChecked(RequiredObjectId);
 			FRemoteObjectId ObjectIdToRequest = RequiredCanonicalRootObjectId.IsValid() ? RequiredCanonicalRootObjectId : RequiredObjectId;
 
 			if (Handle::IsRemote(ObjectIdToRequest))
@@ -360,10 +351,9 @@ public:
 				for (int32 RequestIndex = 0; RequestIndex < GetRequestCount(); RequestIndex++)
 				{
 					FRemoteObjectRequest* ExistingRequest = GetRequestByIndex(RequestIndex);
-
 					bool bFoundObject =
-						(ExistingRequest->RequiredObjectsCanonicalRoot.Find(ObjectIdToRequest) != INDEX_NONE) ||
-						(ExistingRequest->RequiredObjects.Find(ObjectIdToRequest) != INDEX_NONE);
+						ExistingRequest->RequiredObjectsCanonicalRoots.Contains(ObjectIdToRequest)||
+						ExistingRequest->RequiredObjects.Contains(ObjectIdToRequest);
 
 					if (bFoundObject)
 					{
@@ -428,7 +418,7 @@ public:
 				{
 					UObject* const RootObject = FindCanonicalRootObjectForSerialization(Object);
 					const FRemoteObjectId CanonicalRootObjectId = UE::CoreUObject::Private::FObjectHandleUtils::GetRemoteId(RootObject);
-					ActiveRequest->RequiredObjectsCanonicalRoot[RequiredObjectIndex] = CanonicalRootObjectId;
+					ActiveRequest->RequiredObjectsCanonicalRootByObjectID.FindChecked(RequiredObjectId) = CanonicalRootObjectId;
 				}
 			}
 		}
@@ -436,10 +426,8 @@ public:
 
 	bool AreDependenciesSatisfied() const final
 	{
-		for (int32 RequiredObjectIndex = 0; RequiredObjectIndex < ActiveRequest->RequiredObjects.Num(); RequiredObjectIndex++)
+		for (const FRemoteObjectId& RequiredObjectId : ActiveRequest->RequiredObjects)
 		{
-			FRemoteObjectId RequiredObjectId = ActiveRequest->RequiredObjects[RequiredObjectIndex];
-
 			if (Handle::IsRemote(RequiredObjectId))
 			{
 				return false;
@@ -913,10 +901,9 @@ UObject* MigrateObjectFromObjectData(FRemoteObjectData& ObjectData, const FUObje
 
 		for (FRemoteObjectId SubObjectId : ReceivedObjectRemoteIds)
 		{
-			const int32 RequiredObjectIndex = ExistingRequest->RequiredObjects.Find(SubObjectId);
-			if (RequiredObjectIndex != INDEX_NONE)
+			if (FRemoteObjectId* ExistingCanonicalRootForSubObject = ExistingRequest->RequiredObjectsCanonicalRootByObjectID.Find(SubObjectId))
 			{
-				ExistingRequest->RequiredObjectsCanonicalRoot[RequiredObjectIndex] = RootObjectRemoteId;
+				*ExistingCanonicalRootForSubObject = RootObjectRemoteId;
 			}
 		}
 	}
@@ -938,7 +925,7 @@ void MigrateObjectFromRemoteServer(FRemoteObjectId ObjectId, FRemoteServerId Cur
 	{
 		// check bInMultiServerCommit here to ensure the UsedObjects list doesn't 
 		// accidentally mutate during the actual multi-server commit process
-		if (!GTransferQueue.bInMultiServerCommit && GTransferQueue.ActiveRequest->UsedObjects.Find(ObjectId) == INDEX_NONE)
+		if (!GTransferQueue.bInMultiServerCommit && !GTransferQueue.ActiveRequest->UsedObjects.Contains(ObjectId))
 		{
 			GTransferQueue.ActiveRequest->UsedObjects.Add(ObjectId);
 
@@ -946,7 +933,7 @@ void MigrateObjectFromRemoteServer(FRemoteObjectId ObjectId, FRemoteServerId Cur
 		}
 
 		// add it to the active request's list of required objects
-		if (GTransferQueue.ActiveRequest->RequiredObjects.Find(ObjectId) == INDEX_NONE)
+		if (!GTransferQueue.ActiveRequest->RequiredObjects.Contains(ObjectId))
 		{
 			// does any other existing request also require this object? if so, try to
 			// grab the cached canonical id from them
@@ -956,20 +943,15 @@ void MigrateObjectFromRemoteServer(FRemoteObjectId ObjectId, FRemoteServerId Cur
 			{
 				FRemoteObjectRequest* ExistingRequest = GTransferQueue.GetRequestByIndex(RequestIndex);
 
-				for (int32 RequiredObjectIndex = 0; RequiredObjectIndex < ExistingRequest->RequiredObjects.Num(); RequiredObjectIndex++)
+				if (FRemoteObjectId* FoundRequiredRemoteObjectID = ExistingRequest->RequiredObjects.Find(ObjectId))
 				{
-					FRemoteObjectId RequiredObjectId = ExistingRequest->RequiredObjects[RequiredObjectIndex];
-
-					if (ObjectId == RequiredObjectId)
-					{
-						CanonicalRootObjectId = ExistingRequest->RequiredObjectsCanonicalRoot[RequiredObjectIndex];
-						break;
-					}
+					CanonicalRootObjectId = ExistingRequest->RequiredObjectsCanonicalRootByObjectID.FindChecked(*FoundRequiredRemoteObjectID);
 				}
 			}
 
 			GTransferQueue.ActiveRequest->RequiredObjects.Add(ObjectId);
-			GTransferQueue.ActiveRequest->RequiredObjectsCanonicalRoot.Add(CanonicalRootObjectId);
+			GTransferQueue.ActiveRequest->RequiredObjectsCanonicalRootByObjectID.Add(ObjectId, CanonicalRootObjectId);
+			GTransferQueue.ActiveRequest->RequiredObjectsCanonicalRoots.Add(CanonicalRootObjectId);
 		}
 		
 		if (UE_LOG_ACTIVE(LogRemoteObject, VeryVerbose))
@@ -1009,7 +991,7 @@ void TouchResidentObject(UObject* Object)
 		{
 			const FRemoteObjectId ObjectId = UE::CoreUObject::Private::FObjectHandleUtils::GetRemoteId(Object);
 
-			if (GTransferQueue.ActiveRequest->UsedObjects.Find(ObjectId) == INDEX_NONE)
+			if (!GTransferQueue.ActiveRequest->UsedObjects.Contains(ObjectId))
 			{
 				GTransferQueue.ActiveRequest->UsedObjects.Add(ObjectId);
 
@@ -1017,7 +999,7 @@ void TouchResidentObject(UObject* Object)
 			}
 
 			// add it to the active request's list of required objects
-			if (GTransferQueue.ActiveRequest->RequiredObjects.Find(ObjectId) == INDEX_NONE)
+			if (!GTransferQueue.ActiveRequest->RequiredObjects.Contains(ObjectId))
 			{
 				GTransferQueue.ActiveRequest->RequiredObjects.Add(ObjectId);
 				
@@ -1025,7 +1007,8 @@ void TouchResidentObject(UObject* Object)
 				// object because it could trigger endless recursion into this function which is
 				// called when object pointers get resolved
 				FRemoteObjectId CanonicalRootObjectId;
-				GTransferQueue.ActiveRequest->RequiredObjectsCanonicalRoot.Add(CanonicalRootObjectId);
+				GTransferQueue.ActiveRequest->RequiredObjectsCanonicalRootByObjectID.Add(ObjectId, CanonicalRootObjectId);
+				GTransferQueue.ActiveRequest->RequiredObjectsCanonicalRoots.Add(CanonicalRootObjectId);
 			}
 		};
 	}
@@ -1185,7 +1168,7 @@ void PrefetchRemoteObject(FRemoteObjectId ObjectId, UObject* DestinationOuter)
 	{
 		// check bInMultiServerCommit here to ensure the UsedObjects list doesn't 
 		// accidentally mutate during the actual multi-server commit process
-		if (!GTransferQueue.bInMultiServerCommit && GTransferQueue.ActiveRequest->UsedObjects.Find(ObjectId) == INDEX_NONE)
+		if (!GTransferQueue.bInMultiServerCommit && !GTransferQueue.ActiveRequest->UsedObjects.Contains(ObjectId))
 		{
 			GTransferQueue.ActiveRequest->UsedObjects.Add(ObjectId);
 
@@ -1193,12 +1176,12 @@ void PrefetchRemoteObject(FRemoteObjectId ObjectId, UObject* DestinationOuter)
 		}
 
 		// it's possible a previous prefetch or strong fetch already added this to the RequiredObjects list
-		bAdded = (GTransferQueue.ActiveRequest->RequiredObjects.Find(ObjectId) == INDEX_NONE);
+		bAdded = !GTransferQueue.ActiveRequest->RequiredObjects.Contains(ObjectId);
 
 		if (bAdded)
 		{
 			GTransferQueue.ActiveRequest->RequiredObjects.Add(ObjectId);
-			GTransferQueue.ActiveRequest->RequiredObjectsCanonicalRoot.Add(ObjectId);
+			GTransferQueue.ActiveRequest->RequiredObjectsCanonicalRoots.Add(ObjectId);
 			UE_LOG(LogRemoteObject, VeryVerbose, TEXT("PrefetchRemoteObject(%s) : obj %s"), *GTransferQueue.ActiveRequest->GetRequestId().ToString(), *ObjectId.ToString());
 		}
 	};
