Fixed a FARFilter with only WithPackageFlags or WithoutPackageFlags set failing to return any assets from IAssetRegistry::GetAssets(...)

Such a FARFilter is not considered empty via its IsEmpty() function, but the compiled version (FARCompiledFilter) was considered empty and would fail the query.

With FARFilter::IsEmpty() and FARCompiledFilter::IsEmpty() made consistent, FAssetRegistryState::EnumerateAssets also needed updating to handle the fact that no accelerated filters were set in the FARCompiledFilter (as the flag checks aren't accelerated).

diff --git a/Engine/Source/Runtime/AssetRegistry/Private/AssetRegistryState.cpp b/Engine/Source/Runtime/AssetRegistry/Private/AssetRegistryState.cpp
--- a/Engine/Source/Runtime/AssetRegistry/Private/AssetRegistryState.cpp
+++ b/Engine/Source/Runtime/AssetRegistry/Private/AssetRegistryState.cpp
@@ -1444,6 +1444,17 @@ bool FAssetRegistryState::EnumerateAssets(const FARCompiledFilter& Filter, const
 		}
 	}
 
+	// If AccumulatedResults is still empty at this point, then it means that none of the above accelerated filters were set 
+	// (as they would have early returned), and so we need to run all assets through the non-accelerated filter
+	if (AccumulatedResults.IsEmpty())
+	{
+		AccumulatedResults.Reserve(CachedAssets.Num());
+		for (const FAssetData* AssetData : CachedAssets)
+		{
+			AccumulatedResults.Add(AssetData);
+		}
+	}
+
 	// Run the remaining non-accelerated filters on every element of AccumulatedResults
 
 #if WITH_ENGINE && WITH_EDITOR
diff --git a/Engine/Source/Runtime/CoreUObject/Public/AssetRegistry/ARFilter.h b/Engine/Source/Runtime/CoreUObject/Public/AssetRegistry/ARFilter.h
--- a/Engine/Source/Runtime/CoreUObject/Public/AssetRegistry/ARFilter.h
+++ b/Engine/Source/Runtime/CoreUObject/Public/AssetRegistry/ARFilter.h
@@ -105,7 +105,7 @@ PRAGMA_ENABLE_DEPRECATION_WARNINGS
 	/** Returns true if this filter has no entries */
 	bool IsEmpty() const
 	{
-		return PackageNames.Num() + PackagePaths.Num() + SoftObjectPaths.Num() + ClassPaths.Num() + TagsAndValues.Num() + WithoutPackageFlags + WithPackageFlags == 0;
+		return (PackageNames.Num() | PackagePaths.Num() | SoftObjectPaths.Num() | ClassPaths.Num() | TagsAndValues.Num() | WithoutPackageFlags | WithPackageFlags) == 0;
 	}
 
 	/** Returns true if this filter is recursive */
@@ -218,7 +218,7 @@ PRAGMA_ENABLE_DEPRECATION_WARNINGS
 	/** Returns true if this filter has no entries */
 	bool IsEmpty() const
 	{
-		return PackageNames.Num() + PackagePaths.Num() + SoftObjectPaths.Num() + ClassPaths.Num() + TagsAndValues.Num() == 0;
+		return (PackageNames.Num() | PackagePaths.Num() | SoftObjectPaths.Num() | ClassPaths.Num() | TagsAndValues.Num() | WithoutPackageFlags | WithPackageFlags) == 0;
 	}
 
 	/** Clears this filter of all entries */
