[CrashFix] CheckValidObject is being called from within a package save and it's causing a crash when it calls StaticFindObject.

diff --git a/Engine/Source/Runtime/CoreUObject/Private/UObject/InstanceDataObjectUtils.cpp b/Engine/Source/Runtime/CoreUObject/Private/UObject/InstanceDataObjectUtils.cpp
--- a/Engine/Source/Runtime/CoreUObject/Private/UObject/InstanceDataObjectUtils.cpp
+++ b/Engine/Source/Runtime/CoreUObject/Private/UObject/InstanceDataObjectUtils.cpp
@@ -329,7 +329,7 @@ namespace UE
 	);
 
 	// TODO: re-enable this ASAP. This disables most IDO features but disabling was necessary to unblock those experiencing IDO bugs
-	bool bEnableIDOImpersonationOnSave = false;
+	bool bEnableIDOImpersonationOnSave = true;
 	FAutoConsoleVariableRef EnableIDOImpersonationOnSaveCVar(
 		TEXT("IDO.Impersonation.EnableOnSave"),
 		bEnableIDOImpersonationOnSave,
diff --git a/Engine/Source/Runtime/CoreUObject/Private/UObject/PropertyObject.cpp b/Engine/Source/Runtime/CoreUObject/Private/UObject/PropertyObject.cpp
--- a/Engine/Source/Runtime/CoreUObject/Private/UObject/PropertyObject.cpp
+++ b/Engine/Source/Runtime/CoreUObject/Private/UObject/PropertyObject.cpp
@@ -328,7 +328,7 @@ void FObjectProperty::PostSerializeObjectItem(FArchive& SerializingArchive, void
 		//        to accommodate this (as it depends on finding itself as the set value)
 #endif // USE_CIRCULAR_DEPENDENCY_LOAD_DEFERRING
 
-		if (!SerializingArchive.IsSaving())
+		if (!SerializingArchive.IsSaving() && !UE::IsSavingPackage(nullptr))
 		{
 			CheckValidObject(Value, ObjectValue, Defaults);
 		}
