dstm trace vizualization:

- added labels for tick transactions
- when clicking a transaction there is now a vizualizer for all of the separate executions
- marked up object migration regions
- NetPump events

diff --git a/Engine/Source/Runtime/CoreUObject/Private/UObject/RemoteExecutor.cpp b/Engine/Source/Runtime/CoreUObject/Private/UObject/RemoteExecutor.cpp
--- a/Engine/Source/Runtime/CoreUObject/Private/UObject/RemoteExecutor.cpp
+++ b/Engine/Source/Runtime/CoreUObject/Private/UObject/RemoteExecutor.cpp
@@ -19,9 +19,12 @@ TDelegate<TOptional<TTuple<FName, FRemoteWorkPriority, bool, TFunction<void(void
 
 TMulticastDelegate<void(FRemoteTransactionId, FName)> OnTransactionQueuedDelegate;
 TMulticastDelegate<void(FRemoteTransactionId, FName)> OnTransactionStartingDelegate;
-TMulticastDelegate<void(FRemoteTransactionId)> OnTransactionCompletedDelegate;
-TMulticastDelegate<void(FRemoteTransactionId)> OnTransactionAbortedDelegate;
+TMulticastDelegate<void(FRemoteTransactionId, uint32)> OnTransactionCompletedDelegate;
+TMulticastDelegate<void(FRemoteTransactionId, uint32, const FString&)> OnTransactionAbortedDelegate;
 TMulticastDelegate<void(FRemoteTransactionId)> OnTransactionReleasedDelegate;
+
+TDelegate<void(FName /*Text*/)> OnRegionBeginDelegate;
+TDelegate<void(const FString& /*Text*/)> OnRegionEndDelegate;
 }
 
 enum class ERemoteExecutorAbortReason : uint8
@@ -505,7 +508,7 @@ public:
 					{
 						bWorkTransactionAborted = true;
 
-						UE::RemoteExecutor::OnTransactionAbortedDelegate.Broadcast(RequestId);
+						UE::RemoteExecutor::OnTransactionAbortedDelegate.Broadcast(RequestId, ExecutingWork->ExecutionAttempts, AbortReasonDescription);
 
 						if (AbortReason == ERemoteExecutorAbortReason::RequiresDependencies) // -V547
 						{
@@ -538,7 +541,7 @@ public:
 							LocalIterationNumber,
 							*ExecutingWork->RequestId.ToString(),
 							ExecutingWork->ExecutionAttempts);
-						UE::RemoteExecutor::OnTransactionCompletedDelegate.Broadcast(RequestId);
+						UE::RemoteExecutor::OnTransactionCompletedDelegate.Broadcast(RequestId, ExecutingWork->ExecutionAttempts);
 
 						bWorkComplete = true;
 					}
diff --git a/Engine/Source/Runtime/CoreUObject/Private/UObject/RemoteObjectTransfer.cpp b/Engine/Source/Runtime/CoreUObject/Private/UObject/RemoteObjectTransfer.cpp
--- a/Engine/Source/Runtime/CoreUObject/Private/UObject/RemoteObjectTransfer.cpp
+++ b/Engine/Source/Runtime/CoreUObject/Private/UObject/RemoteObjectTransfer.cpp
@@ -182,6 +182,11 @@ public:
 
 				if (!bObjectLocked)
 				{
+					FString ObjectName = Object->GetName();
+
+					static FName SendObjectRegionName(TEXT("ObjectTransferQueue_SendObject"));
+					UE::RemoteExecutor::OnRegionBeginDelegate.ExecuteIfBound(SendObjectRegionName);
+
 					// send the object
 					UE_LOG(LogRemoteObject, VeryVerbose, TEXT("TickObjectMigrations sending obj %s to %s %s"),
 						*HighestPriorityRequest->ObjectId.ToString(),
@@ -269,7 +274,9 @@ public:
 							}
 						}
 					}
-					
+
+					UE::RemoteExecutor::OnRegionEndDelegate.ExecuteIfBound(FString::Format(TEXT("{0} {1}"), {*RootObject->GetName(), *CanonicalRootObjectId.ToString()}));
+
 					// all done, delete all requests for this object
 					PendingObjectMigrationRequests.RemoveAt(MigrationRequestsIndex, EAllowShrinking::No);
 					// do not increment MigrationRequestsIndex
@@ -792,6 +799,15 @@ UObject* MigrateObjectFromObjectData(FRemoteObjectData& ObjectData, const FUObje
 	using namespace UE::RemoteObject::Serialization;
 	using namespace UE::RemoteObject::Private;
 
+	uint32 ByteSize = 0;
+	for (const FRemoteObjectBytes& Chunk : ObjectData.Bytes)
+	{
+		ByteSize += Chunk.Bytes.Num();
+	}
+
+	static FName GetObjectRegionName(TEXT("ObjectTransferQueue_RecvObject"));
+	UE::RemoteExecutor::OnRegionBeginDelegate.ExecuteIfBound(GetObjectRegionName);
+
 	check(!AutoRTFM::IsClosed());
 
 	// Notify that object data has been migrated from a remote server
@@ -908,6 +924,8 @@ UObject* MigrateObjectFromObjectData(FRemoteObjectData& ObjectData, const FUObje
 		}
 	}
 
+	UE::RemoteExecutor::OnRegionEndDelegate.ExecuteIfBound(FString::Format(TEXT("{0} {1} ({2} bytes)"), {*ReceivedObjects[0]->GetName(), *RootObjectRemoteId.ToString(), ByteSize}));
+
 	return ReceivedObjects[RequestedObjectIndex];
 } 
 
diff --git a/Engine/Source/Runtime/CoreUObject/Public/UObject/RemoteExecutor.h b/Engine/Source/Runtime/CoreUObject/Public/UObject/RemoteExecutor.h
--- a/Engine/Source/Runtime/CoreUObject/Public/UObject/RemoteExecutor.h
+++ b/Engine/Source/Runtime/CoreUObject/Public/UObject/RemoteExecutor.h
@@ -210,14 +210,17 @@ namespace UE::RemoteExecutor
 	extern COREUOBJECT_API TMulticastDelegate<void(FRemoteTransactionId, FName)> OnTransactionStartingDelegate;
 
 	/** Delegate executed when a ExecuteTransactional observes a completed AutoRTFM transaction */
-	extern COREUOBJECT_API TMulticastDelegate<void(FRemoteTransactionId)> OnTransactionCompletedDelegate;
+	extern COREUOBJECT_API TMulticastDelegate<void(FRemoteTransactionId, uint32)> OnTransactionCompletedDelegate;
 
 	/** Delegate executed when a ExecuteTransactional observes an aborted AutoRTFM transaction */
-	extern COREUOBJECT_API TMulticastDelegate<void(FRemoteTransactionId)> OnTransactionAbortedDelegate;
+	extern COREUOBJECT_API TMulticastDelegate<void(FRemoteTransactionId, uint32, const FString&)> OnTransactionAbortedDelegate;
 
 	/** Delegate executed when a transaction is released (execution is finished and the work is removed from the executor) */
 	extern COREUOBJECT_API TMulticastDelegate<void(FRemoteTransactionId)> OnTransactionReleasedDelegate;
 
+	extern COREUOBJECT_API TDelegate<void(FName /*Text*/)> OnRegionBeginDelegate;
+	extern COREUOBJECT_API TDelegate<void(const FString& /*Text*/)> OnRegionEndDelegate;
+
 	COREUOBJECT_API void RegisterRemoteSubsystem(FRemoteSubsystemBase* Subsystem);
 
 	COREUOBJECT_API void ExecuteTransactional(FName WorkName, const TFunctionRef<void(void)>& Work);
diff --git a/Engine/Source/Runtime/Engine/Classes/Engine/EngineBaseTypes.h b/Engine/Source/Runtime/Engine/Classes/Engine/EngineBaseTypes.h
--- a/Engine/Source/Runtime/Engine/Classes/Engine/EngineBaseTypes.h
+++ b/Engine/Source/Runtime/Engine/Classes/Engine/EngineBaseTypes.h
@@ -480,6 +480,11 @@ public:
 	 **/
 	virtual void ExecuteTick(float DeltaTime, ELevelTick TickType, ENamedThreads::Type CurrentThread, const FGraphEventRef& MyCompletionGraphEvent) PURE_VIRTUAL(,);
 	
+#if UE_WITH_REMOTE_OBJECT_HANDLE
+	FName CachedDiagnosticContext;
+	FString CachedDiagnosticMessage;
+#endif
+
 	/** Abstract function to describe this tick. Used to print messages about illegal cycles in the dependency graph **/
 	virtual FString DiagnosticMessage() PURE_VIRTUAL(, return TEXT("DiagnosticMessage() not implemented"););
 	
diff --git a/Engine/Source/Runtime/Engine/Private/Actor.cpp b/Engine/Source/Runtime/Engine/Private/Actor.cpp
--- a/Engine/Source/Runtime/Engine/Private/Actor.cpp
+++ b/Engine/Source/Runtime/Engine/Private/Actor.cpp
@@ -374,6 +374,14 @@ void FActorTickFunction::ExecuteTick(float DeltaTime, enum ELevelTick TickType,
 			FScopeCycleCounterUObject ActorScope(Target);
 			Target->TickActor(DeltaTime*Target->CustomTimeDilation, TickType, *this);
 		}
+
+#if UE_WITH_REMOTE_OBJECT_HANDLE
+		UE_AUTORTFM_OPEN
+		{
+			CachedDiagnosticContext = DiagnosticContext(true);
+			CachedDiagnosticMessage = DiagnosticMessage();
+		};
+#endif
 	}
 }
 
diff --git a/Engine/Source/Runtime/Engine/Private/Components/ActorComponent.cpp b/Engine/Source/Runtime/Engine/Private/Components/ActorComponent.cpp
--- a/Engine/Source/Runtime/Engine/Private/Components/ActorComponent.cpp
+++ b/Engine/Source/Runtime/Engine/Private/Components/ActorComponent.cpp
@@ -1593,6 +1593,14 @@ void FActorComponentTickFunction::ExecuteTick(float DeltaTime, enum ELevelTick T
 	{
 		Target->TickComponent(DilatedTime, TickType, this);
 	});
+
+#if UE_WITH_REMOTE_OBJECT_HANDLE
+	UE_AUTORTFM_OPEN
+	{
+		CachedDiagnosticContext = DiagnosticContext(true);
+		CachedDiagnosticMessage = DiagnosticMessage();
+	};
+#endif
 }
 
 FString FActorComponentTickFunction::DiagnosticMessage()
diff --git a/Engine/Source/Runtime/Engine/Private/TickTaskManager.cpp b/Engine/Source/Runtime/Engine/Private/TickTaskManager.cpp
--- a/Engine/Source/Runtime/Engine/Private/TickTaskManager.cpp
+++ b/Engine/Source/Runtime/Engine/Private/TickTaskManager.cpp
@@ -338,8 +338,7 @@ public:
 
 			if (Target->bRunTransactionally)
 			{
-				static FName TransactionalWorkName(TEXT("FTickFunction"));
-				UE::RemoteExecutor::ExecuteTransactional(TransactionalWorkName, ExecuteTickWork);
+				UE::RemoteExecutor::ExecuteTransactional(Target->CachedDiagnosticContext, ExecuteTickWork);
 			}
 			else
 			{
@@ -430,8 +429,7 @@ public:
 
 				if (Target->bRunTransactionally)
 				{
-					static FName TransactionalWorkName(TEXT("FTickFunction"));
-					UE::RemoteExecutor::EnqueueTransactional(TransactionalWorkName, ExecuteTickWork);
+					UE::RemoteExecutor::EnqueueTransactional(Target->CachedDiagnosticContext, ExecuteTickWork);
 				}
 				else
 				{
