Object Duplication: added a way to differentiate between default and null duplication seeds in FDuplicatedObject by mapping duplicated objects to INVALID_OBJECT.
This allows adding entries in the duplication seed for specific objects when passing it through StaticDuplicateObject that will not be duplicated and nulled out.

diff --git a/Engine/Source/Runtime/CoreUObject/Public/Serialization/DuplicatedObject.h b/Engine/Source/Runtime/CoreUObject/Public/Serialization/DuplicatedObject.h
--- a/Engine/Source/Runtime/CoreUObject/Public/Serialization/DuplicatedObject.h
+++ b/Engine/Source/Runtime/CoreUObject/Public/Serialization/DuplicatedObject.h
@@ -11,15 +11,20 @@
  */
 struct FDuplicatedObject
 {
+	/** The state of this object */
+	bool bIsDefault;
+
 	/** The duplicated object */
 	TWeakObjectPtr<UObject> DuplicatedObject;
 
 	FDuplicatedObject()
+		: bIsDefault(true)
 	{
 	}
 
 	FDuplicatedObject( UObject* InDuplicatedObject )
-		: DuplicatedObject( InDuplicatedObject )
+		: bIsDefault( !InDuplicatedObject )
+		, DuplicatedObject( InDuplicatedObject != INVALID_OBJECT ? InDuplicatedObject : nullptr )
 	{
 	}
 
@@ -28,7 +33,7 @@ struct FDuplicatedObject
 	 */
 	FORCEINLINE bool IsDefault()
 	{
-		return DuplicatedObject.IsExplicitlyNull();
+		return bIsDefault;
 	}
 };
 
