bContainsInnerProperties is now false for container when they do not have any elements the PropertyVisitor logic

diff --git a/Engine/Source/Runtime/CoreUObject/Private/UObject/PropertyArray.cpp b/Engine/Source/Runtime/CoreUObject/Private/UObject/PropertyArray.cpp
--- a/Engine/Source/Runtime/CoreUObject/Private/UObject/PropertyArray.cpp
+++ b/Engine/Source/Runtime/CoreUObject/Private/UObject/PropertyArray.cpp
@@ -1426,19 +1426,20 @@ bool FArrayProperty::CanSerializeFromTypeName(UE::FPropertyTypeName Type) const
 
 EPropertyVisitorControlFlow FArrayProperty::Visit(FPropertyVisitorContext& Context, const TFunctionRef<EPropertyVisitorControlFlow(const FPropertyVisitorContext& /*Conmtext*/)> InFunc) const
 {
+	FScriptArrayHelper ArrayHelper(this, Context.Data.PropertyData);
+	const int32 ArrayNum = ArrayHelper.Num();
+
 	// Indicate in the path that this property contains inner properties
-	Context.Path.Top().bContainsInnerProperties = true;
+	Context.Path.Top().bContainsInnerProperties = ArrayNum > 0;
 
 	EPropertyVisitorControlFlow RetVal = Super::Visit(Context, InFunc);
 
 	if (RetVal == EPropertyVisitorControlFlow::StepInto)
 	{
 		checkf(Inner, TEXT("Expecting a valid inner property type"));
-		FScriptArrayHelper ArrayHelper(this, Context.Data.PropertyData);
 
 		FPropertyVisitorScope Scope(Context.Path, FPropertyVisitorInfo(Inner));
 
-		const int32 ArrayNum = ArrayHelper.Num();
 		for (int32 ContainerIndex = 0; ContainerIndex < ArrayNum; ContainerIndex++)
 		{
 			Context.Path.Top().SetIndex(ContainerIndex, EPropertyVisitorInfoType::ContainerIndex);
diff --git a/Engine/Source/Runtime/CoreUObject/Private/UObject/PropertyMap.cpp b/Engine/Source/Runtime/CoreUObject/Private/UObject/PropertyMap.cpp
--- a/Engine/Source/Runtime/CoreUObject/Private/UObject/PropertyMap.cpp
+++ b/Engine/Source/Runtime/CoreUObject/Private/UObject/PropertyMap.cpp
@@ -1792,15 +1792,16 @@ bool FMapProperty::CanSerializeFromTypeName(UE::FPropertyTypeName Type) const
 
 EPropertyVisitorControlFlow FMapProperty::Visit(FPropertyVisitorContext& Context, const TFunctionRef<EPropertyVisitorControlFlow(const FPropertyVisitorContext& /*Conmtext*/)> InFunc) const
 {
+	FScriptMapHelper MapHelper(this, Context.Data.PropertyData);
+
 	// Indicate in the path that this property contains inner properties
-	Context.Path.Top().bContainsInnerProperties = true;
+	Context.Path.Top().bContainsInnerProperties = MapHelper.Num() > 0;
 
 	EPropertyVisitorControlFlow RetVal = Super::Visit(Context, InFunc);
 
 	if (RetVal == EPropertyVisitorControlFlow::StepInto)
 	{
 		checkf(KeyProp && ValueProp, TEXT("Expecting a valid inner property type"));
-		FScriptMapHelper MapHelper(this, Context.Data.PropertyData);
 
 		for (FScriptMapHelper::FIterator It(MapHelper); It; ++It)
 		{
diff --git a/Engine/Source/Runtime/CoreUObject/Private/UObject/PropertyOptional.cpp b/Engine/Source/Runtime/CoreUObject/Private/UObject/PropertyOptional.cpp
--- a/Engine/Source/Runtime/CoreUObject/Private/UObject/PropertyOptional.cpp
+++ b/Engine/Source/Runtime/CoreUObject/Private/UObject/PropertyOptional.cpp
@@ -693,7 +693,7 @@ bool FOptionalProperty::CanSerializeFromTypeName(UE::FPropertyTypeName Type) con
 EPropertyVisitorControlFlow FOptionalProperty::Visit(FPropertyVisitorContext& Context, const TFunctionRef<EPropertyVisitorControlFlow(const FPropertyVisitorContext& /*Context*/)> InFunc) const
 {
 	// Indicate in the path that this property contains inner properties
-	Context.Path.Top().bContainsInnerProperties = true;
+	Context.Path.Top().bContainsInnerProperties = IsSet(Context.Data.PropertyData);
 
 	EPropertyVisitorControlFlow RetVal = Super::Visit(Context, InFunc);
 
diff --git a/Engine/Source/Runtime/CoreUObject/Private/UObject/PropertySet.cpp b/Engine/Source/Runtime/CoreUObject/Private/UObject/PropertySet.cpp
--- a/Engine/Source/Runtime/CoreUObject/Private/UObject/PropertySet.cpp
+++ b/Engine/Source/Runtime/CoreUObject/Private/UObject/PropertySet.cpp
@@ -1123,15 +1123,16 @@ bool FSetProperty::CanSerializeFromTypeName(UE::FPropertyTypeName Type) const
 
 EPropertyVisitorControlFlow FSetProperty::Visit(FPropertyVisitorContext& Context, const TFunctionRef<EPropertyVisitorControlFlow(const FPropertyVisitorContext& /*Conmtext*/)> InFunc) const
 {
+	FScriptSetHelper SetHelper(this, Context.Data.PropertyData);
+
 	// Indicate in the path that this property contains inner properties
-	Context.Path.Top().bContainsInnerProperties = true;
+	Context.Path.Top().bContainsInnerProperties = SetHelper.Num() > 0;
 
 	EPropertyVisitorControlFlow RetVal = Super::Visit(Context, InFunc);
 
 	if (RetVal == EPropertyVisitorControlFlow::StepInto)
 	{
 		checkf(ElementProp, TEXT("Expecting a valid inner property type"));
-		FScriptSetHelper SetHelper(this, Context.Data.PropertyData);
 
 		FPropertyVisitorScope Scope(Context.Path, FPropertyVisitorInfo(ElementProp));
 		for (FScriptSetHelper::FIterator It(SetHelper); It; ++It)
