[ZenLoader]
- Add support for going wide with some nodes (requires s.AllowMultithreadedLoading = 1)
- Move the thread-unsafe part of ProcessPackageSummary into ProcessDependencies that will be sequential
- Make ProcessPackageSummary and LinkerLoadPreload Parallel when the option is enabled
- Fix a missing redirector Preload
- Add a new gate to wait until all dependencies are preloaded before moving to postload
- Add support for sync flush from parallel loading thread

diff --git a/Engine/Source/Runtime/CoreUObject/Private/Serialization/AsyncLoading2.cpp b/Engine/Source/Runtime/CoreUObject/Private/Serialization/AsyncLoading2.cpp
--- a/Engine/Source/Runtime/CoreUObject/Private/Serialization/AsyncLoading2.cpp
+++ b/Engine/Source/Runtime/CoreUObject/Private/Serialization/AsyncLoading2.cpp
@@ -18,6 +18,7 @@
 #include "HAL/PlatformMisc.h"
 #include "Misc/ScopeLock.h"
 #include "Stats/StatsMisc.h"
+#include "Containers/StripedMap.h"
 #include "Misc/CoreStats.h"
 #include "HAL/IConsoleManager.h"
 #include "Misc/PackageAccessTrackingOps.h"
@@ -42,8 +43,10 @@
 #include "UObject/UObjectHash.h"
 #include "Templates/Casts.h"
 #include "Templates/UniquePtr.h"
+#include "Templates/SharedPointer.h"
 #include "Serialization/BufferReader.h"
 #include "Async/TaskGraphInterfaces.h"
+#include "Async/UniqueLock.h"
 #include "Blueprint/BlueprintSupport.h"
 #include "HAL/LowLevelMemTracker.h"
 #include "ProfilingDebugging/CsvProfiler.h"
@@ -51,6 +54,7 @@
 #include "Misc/MTAccessDetector.h"
 #include "UObject/GarbageCollectionInternal.h"
 #include "ProfilingDebugging/MiscTrace.h"
+#include "StructUtils/UserDefinedStruct.h"
 #include "Serialization/LoadTimeTracePrivate.h"
 #include "ProfilingDebugging/CountersTrace.h"
 #include "Serialization/AsyncPackage.h"
@@ -67,7 +71,6 @@
 #include "UObject/LinkerLoadImportBehavior.h"
 #include "UObject/ObjectHandlePrivate.h"
 #include "UObject/ObjectSerializeAccessScope.h"
-#include "ProfilingDebugging/CountersTrace.h"
 #include "ProfilingDebugging/AssetMetadataTrace.h"
 #include "Async/Async.h"
 #include "Sanitizer/RaceDetector.h"
@@ -79,11 +82,9 @@
 #include "HAL/IPlatformFileOpenLogWrapper.h"
 #include "Modules/ModuleManager.h"
 #include "Containers/MpscQueue.h"
-#include "Containers/SpscQueue.h"
 #include "Misc/AssetRegistryInterface.h"
 #include "Misc/PathViews.h"
 #include "UObject/LinkerLoad.h"
-#include "Containers/SpscQueue.h"
 #include "IO/IoPriorityQueue.h"
 #include "UObject/CoreRedirects.h"
 #include "Serialization/ZenPackageHeader.h"
@@ -473,6 +474,14 @@ static FAutoConsoleVariableRef CVarGVerifyObjectLoadFlags(
 	ECVF_Default
 	);
 
+static bool GAllowMultithreadedLoading = false;
+static FAutoConsoleVariableRef CVarGAllowMultithreadedLoading(
+	TEXT("s.AllowMultithreadedLoading"),
+	GAllowMultithreadedLoading,
+	TEXT("Allow some loading steps to be fully multithreaded."),
+	ECVF_Default
+	);
+
 static bool GUsePostloadGroupsWithIOStore = true;
 static FAutoConsoleVariableRef CVarGUsePostloadGroupsWithIOStore(
 	TEXT("s.UsePostloadGroupsWithIOStore"),
@@ -1601,6 +1610,7 @@ private:
 	FAsyncLoadingThread2& AsyncLoadingThread;
 	// Packages in active loading or completely loaded packages, with Desc.UPackageId as key.
 	// Does not track temp packages with custom UPackage names, since they are never imorted by other packages.
+	mutable FRWLock PackagesLock;
 	TMap<FPackageId, FLoadedPackageRef> Packages;
 	// Packages in active loading that export VCells. Stored separately from FLoadedPackageRefs to save space.
 	TMap<FPackageId, FLoadedPackageCellsRef> PackageCells;
@@ -1633,6 +1643,7 @@ public:
 
 	int32 GetStoredPackagesCount() const
 	{
+		FReadScopeLock ScopeLock(PackagesLock);
 		return Packages.Num();
 	}
 
@@ -1655,6 +1666,7 @@ public:
 
 	inline FLoadedPackageRef* FindPackageRef(FPackageId PackageId)
 	{
+		FReadScopeLock ScopeLock(PackagesLock);
 		return Packages.Find(PackageId);
 	}
 
@@ -1670,6 +1682,7 @@ public:
 	{
 		LLM_SCOPE_BYNAME(TEXT("AsyncLoadPackageStore"));
 
+		FWriteScopeLock ScopeLock(PackagesLock);
 		FLoadedPackageRef& PackageRef = Packages.FindOrAdd(PackageId);
 
 		// Detect package that have been renamed but are still referenced and attempt to get rid of the references
@@ -1857,6 +1870,7 @@ public:
 	void VerifyLoadedPackages()
 	{
 		TRACE_CPUPROFILER_EVENT_SCOPE(VerifyLoadedPackages);
+		FReadScopeLock ScopeLock(PackagesLock);
 		for (TPair<FPackageId, FLoadedPackageRef>& Pair : Packages)
 		{
 			FPackageId& PackageId = Pair.Key;
@@ -1923,7 +1937,7 @@ public:
 	void RemovePackage(FPackageId PackageId)
 	{
 		UE_ASYNC_PACKAGEID_DEBUG(PackageId);
-
+		FWriteScopeLock ScopeLock(PackagesLock);
 		FLoadedPackageRef PackageRef;
 		bool bRemoved = Packages.RemoveAndCopyValue(PackageId, PackageRef);
 
@@ -2270,6 +2284,8 @@ struct FPackageImportStore
 #if WITH_EDITOR
 			if (UObjectRedirector* Redirector = dynamic_cast<UObjectRedirector*>(Object))
 			{
+				Redirector->ConditionalPreload();
+				// Make sure we're not trying to access a redirector that has yet to be loaded
 				Object = Redirector->DestinationObject;
 			}
 #endif
@@ -2873,6 +2889,7 @@ enum class EAsyncPackageLoadingState2 : uint8
 	NewPackage,
 	WaitingForIo,
 	ProcessPackageSummary,
+	ProcessDependencies,
 	WaitingForDependencies,
 	DependenciesReady,
 #if ALT2_ENABLE_LINKERLOAD_SUPPORT
@@ -2904,6 +2921,7 @@ const TCHAR* LexToString(EAsyncPackageLoadingState2 AsyncPackageLoadingState)
 	case EAsyncPackageLoadingState2::NewPackage: return TEXT("NewPackage");
 	case EAsyncPackageLoadingState2::WaitingForIo: return TEXT("WaitingForIo");
 	case EAsyncPackageLoadingState2::ProcessPackageSummary: return TEXT("ProcessPackageSummary");
+	case EAsyncPackageLoadingState2::ProcessDependencies: return TEXT("ProcessDependencies");
 	case EAsyncPackageLoadingState2::WaitingForDependencies: return TEXT("WaitingForDependencies");
 	case EAsyncPackageLoadingState2::DependenciesReady: return TEXT("DependenciesReady");
 #if ALT2_ENABLE_LINKERLOAD_SUPPORT
@@ -2967,13 +2985,11 @@ private:
 	friend struct FAsyncPackage2;
 	friend TIoPriorityQueue<FEventLoadNode2>;
 
-	FEventLoadNode2()
-	{
-
-	}
+	FEventLoadNode2() = default;
 
 	void ProcessDependencies(FAsyncLoadingThreadState2& ThreadState);
 	void Fire(FAsyncLoadingThreadState2* ThreadState);
+	void ParallelLoadingWork();
 
 	const FAsyncLoadEventSpec* Spec = nullptr;
 	FAsyncPackage2* Package = nullptr;
@@ -3098,12 +3114,19 @@ private:
 	int32 ExecuteSyncLoadEventsCallCounter = 0;
 };
 
+enum class EExecutionType
+{
+	Normal,
+	Immediate,
+	Parallel,
+};
+
 struct FAsyncLoadEventSpec
 {
 	typedef EEventLoadNodeExecutionResult(*FAsyncLoadEventFunc)(FAsyncLoadingThreadState2&, FAsyncPackage2*, int32);
 	FAsyncLoadEventFunc Func = nullptr;
 	FAsyncLoadEventQueue2* EventQueue = nullptr;
-	bool bExecuteImmediately = false;
+	EExecutionType ExecutionType = EExecutionType::Normal;
 	const TCHAR* Name = nullptr;
 };
 
@@ -3145,7 +3168,7 @@ public:
 	TArray<FAsyncPackage2*, TInlineAllocator<4>> RequestedPackages;
 	FAsyncPackage2* RequestingPackage = nullptr;
 	std::atomic<bool> bHasFoundRequestedPackages { false };
-
+	UE::FManualResetEvent DoneEvent;
 private:
 	std::atomic<int32> RefCount = 1;
 
@@ -3197,15 +3220,16 @@ struct FAsyncLoadingThreadState2
 		FPlatformTLS::SetTlsValue(TlsSlot, State);
 	}
 
-	static FAsyncLoadingThreadState2* Get()
+	static void Reset()
 	{
 		check(FPlatformTLS::IsValidTlsSlot(TlsSlot));
-		return static_cast<FAsyncLoadingThreadState2*>(FPlatformTLS::GetTlsValue(TlsSlot));
+		FPlatformTLS::SetTlsValue(TlsSlot, nullptr);
 	}
 
-	FAsyncLoadingThreadState2(FAsyncLoadEventGraphAllocator& InGraphAllocator, FIoDispatcher& InIoDispatcher)
-		: GraphAllocator(InGraphAllocator)
+	static FAsyncLoadingThreadState2* Get()
 	{
+		check(FPlatformTLS::IsValidTlsSlot(TlsSlot));
+		return static_cast<FAsyncLoadingThreadState2*>(FPlatformTLS::GetTlsValue(TlsSlot));
 	}
 
 	bool IsTimeLimitExceeded(const TCHAR* InLastTypeOfWorkPerformed = nullptr, UObject* InLastObjectWorkWasPerformedOn = nullptr)
@@ -3230,7 +3254,7 @@ struct FAsyncLoadingThreadState2
 			const bool bGCIsWaiting = IsGarbageCollectionWaiting();
 			
 			// Do not time out sync loads due to GC 
-			if (bGCIsWaiting && SyncLoadContextStack.IsEmpty())
+			if (bGCIsWaiting && SyncLoadContextStack.IsEmpty() && !IsInParallelLoadingThread())
 			{
 				bTimeLimitExceeded = true;
 				UE_LOG(LogStreaming, Verbose, TEXT("Timing out async loading (last work %s) due to Garbage Collection request"), InLastTypeOfWorkPerformed);
@@ -3260,15 +3284,14 @@ struct FAsyncLoadingThreadState2
 		return bIsActive.load(std::memory_order_relaxed);
 	}
 
-	FAsyncLoadEventGraphAllocator& GraphAllocator;
 	TArray<FEventLoadNode2*> NodesToFire;
 	TArray<FEventLoadNode2*> CurrentlyExecutingEventNodeStack;
 	TArray<FAsyncLoadingSyncLoadContext*> SyncLoadContextStack;
 	TArray<FAsyncPackage2*> PackagesOnStack;
-	TSpscQueue<FAsyncLoadingSyncLoadContext*> SyncLoadContextsCreatedOnGameThread;
-	TSpscQueue<FAsyncPackage2*> PackagesToReprioritize;
+	TMpscQueue<FAsyncLoadingSyncLoadContext*> SyncLoadContextsCreatedOnOtherThreads;
+	TMpscQueue<FAsyncPackage2*> PackagesToReprioritize;
 	bool bIsAsyncLoadingThread = false;
-	bool bCanAccessAsyncLoadingThreadData = true;
+	bool bCanAccessAsyncLoadingThreadData = false;
 	// used to probe activity for stall detection by the game thread
 	// we use relaxed memory ordering with a simple store to avoid any form of costly interlocked operations.
 	std::atomic<bool> bIsActive = false;
@@ -3288,6 +3311,7 @@ uint32 FAsyncLoadingThreadState2::TlsSlot;
 enum EEventLoadNode2 : uint8
 {
 	Package_ProcessSummary,
+	Package_ProcessDependencies,
 	Package_DependenciesReady,
 #if ALT2_ENABLE_LINKERLOAD_SUPPORT
 	Package_CreateLinkerLoadExports,
@@ -3629,6 +3653,12 @@ struct FAsyncPackage2
 
 	void AddConstructedObject(UObject* Object, bool bSubObjectThatAlreadyExists)
 	{
+		UE::TUniqueLock ScopeLock(ConstructedObjectsMutex);
+
+		// this is required even with the lock above since other
+		// parts are using ConstructedObjects without any locking
+		// but we are expecting that only this part runs concurrently.
+
 		UE_MT_SCOPED_WRITE_ACCESS(ConstructedObjectsAccessDetector);
 
 		if (bSubObjectThatAlreadyExists)
@@ -3843,6 +3873,7 @@ private:
 #if ALT2_ENABLE_LINKERLOAD_SUPPORT
 	FAllDependenciesState		AllDependenciesImportState;
 #endif
+	FAllDependenciesState		AllDependenciesPostLoadReadyState;
 	FAllDependenciesState		AllDependenciesFullyLoadedState;
 
 	/** True if our load has failed */
@@ -3851,6 +3882,8 @@ private:
 	bool						bCreatedLinkerRoot = false;
 	/** List of all request handles. */
 	TArray<int32> RequestIDs;
+	
+	UE::FMutex ConstructedObjectsMutex;
 	/** List of ConstructedObjects = Exports + UPackage + ObjectsCreatedFromExports */
 	TArray<UObject*> ConstructedObjects;
 	/** Detects if the constructed objects are improperly accessed by different threads at the same time. */
@@ -3884,11 +3917,14 @@ public:
 	static EEventLoadNodeExecutionResult Event_PreloadLinkerLoadExports(FAsyncLoadingThreadState2& ThreadState, FAsyncPackage2* Package, int32);
 #endif
 	static EEventLoadNodeExecutionResult Event_ProcessPackageSummary(FAsyncLoadingThreadState2& ThreadState, FAsyncPackage2* Package, int32);
+	static EEventLoadNodeExecutionResult Event_ProcessDependencies(FAsyncLoadingThreadState2& ThreadState, FAsyncPackage2* Package, int32);
 	static EEventLoadNodeExecutionResult Event_DependenciesReady(FAsyncLoadingThreadState2& ThreadState, FAsyncPackage2* Package, int32);
 	static EEventLoadNodeExecutionResult Event_ExportsDone(FAsyncLoadingThreadState2& ThreadState, FAsyncPackage2* Package, int32);
 	static EEventLoadNodeExecutionResult Event_PostLoadExportBundle(FAsyncLoadingThreadState2& ThreadState, FAsyncPackage2* Package, int32 ExportBundleIndex);
 	static EEventLoadNodeExecutionResult Event_DeferredPostLoadExportBundle(FAsyncLoadingThreadState2& ThreadState, FAsyncPackage2* Package, int32 ExportBundleIndex);
 
+	bool CanRunNodeAsync(const FAsyncLoadEventSpec* Spec) const;
+
 	void EventDrivenCreateExport(const FAsyncPackageHeaderData& Header, int32 LocalExportIndex);
 	bool EventDrivenSerializeExport(const FAsyncPackageHeaderData& Header, int32 LocalExportIndex, FExportArchive* Ar);
 
@@ -3946,6 +3982,7 @@ private:
 	void WaitForAllDependenciesToReachState(FAsyncLoadingThreadState2& ThreadState, FAllDependenciesState FAsyncPackage2::* StateMemberPtr, EAsyncPackageLoadingState2 WaitForPackageState, uint32& CurrentTickVariable, TFunctionRef<void(FAsyncPackage2*)> OnStateReached);
 	void ConditionalBeginProcessPackageExports(FAsyncLoadingThreadState2& ThreadState);
 	void ConditionalFinishLoading(FAsyncLoadingThreadState2& ThreadState);
+	void ConditionalBeginPostLoadPhase(FAsyncLoadingThreadState2& ThreadState);
 	void ConditionalReleasePartialRequests(FAsyncLoadingThreadState2& ThreadState);
 
 #if ALT2_ENABLE_LINKERLOAD_SUPPORT
@@ -4160,9 +4197,8 @@ private:
 	TArray<FCompletedPackageRequest> FailedPackageRequests;
 	FCriticalSection FailedPackageRequestsCritical;
 
-	FCriticalSection AsyncPackagesCritical;
 	/** Packages in active loading with GetAsyncPackageId() as key */
-	TMap<FPackageId, FAsyncPackage2*> AsyncPackageLookup;
+	TStripedMap<32, FPackageId, FAsyncPackage2*> AsyncPackageLookup;
 
 	TMpscQueue<FAsyncPackage2*> ExternalReadQueue;
 	TAtomic<int32> PendingIoRequestsCounter{ 0 };
@@ -4238,6 +4274,7 @@ private:
 
 	uint32 ConditionalBeginProcessExportsTick = 0;
 	uint32 ConditionalBeginResolveImportsTick = 0;
+	uint32 ConditionalBeginPostLoadTick = 0;
 	uint32 ConditionalFinishLoadingTick = 0;
 
 public:
@@ -4330,7 +4367,6 @@ public:
 
 	FORCEINLINE FAsyncPackage2* FindAsyncPackage(FPackageId PackageId)
 	{
-		FScopeLock LockAsyncPackages(&AsyncPackagesCritical);
 		//checkSlow(IsInAsyncLoadThread());
 		return AsyncPackageLookup.FindRef(PackageId);
 	}
@@ -4338,7 +4374,6 @@ public:
 	FORCEINLINE FAsyncPackage2* GetAsyncPackage(const FPackageId& PackageId)
 	{
 		// TRACE_CPUPROFILER_EVENT_SCOPE(GetAsyncPackage);
-		FScopeLock LockAsyncPackages(&AsyncPackagesCritical);
 		return AsyncPackageLookup.FindRef(PackageId);
 	}
 
@@ -4428,6 +4463,7 @@ public:
 	virtual void FlushLoading(TConstArrayView<int32> RequestIds) override;
 
 	void FlushLoadingFromLoadingThread(FAsyncLoadingThreadState2& ThreadState, TConstArrayView<int32> RequestIds);
+	void FlushLoadingFromParallelLoadingThread(FAsyncLoadingThreadState2& ThreadState, TConstArrayView<int32> RequestIds);
 	void WarnAboutPotentialSyncLoadStall(FAsyncLoadingSyncLoadContext* SyncLoadContext);
 
 	virtual int32 GetNumQueuedPackages() override
@@ -4888,9 +4924,16 @@ private:
 			CellExport = FExportCell();
 		}
 		Data.bCellExportsInitialized.store(true, std::memory_order_release);
+	}
 
+	void FinishImportingPackages(FAsyncLoadingThreadState2& ThreadState, FAsyncPackage2* AsyncPackage)
+	{
 		if (!AsyncPackage->bHasStartedImportingPackages)
 		{
+#if WITH_EDITOR
+			UE::Core::Private::FPlayInEditorLoadingScope PlayInEditorIDScope(AsyncPackage ? AsyncPackage->Desc.PIEInstanceID : INDEX_NONE);
+#endif
+
 			FIoBatch IoBatch = IoDispatcher.NewBatch();
 			{
 				FPackageStoreReadScope _(PackageStore);
@@ -4946,17 +4989,18 @@ public:
 		// Packages keep all their cell exports alive until they are done loading.
 		// This plays the role of ConstructedObjects and the KeepFlags aspect of EInternalObjectFlags::Async.
 		{
-			FScopeLock LockAsyncPackages(&AsyncLoadingThread->AsyncPackagesCritical);
-			for (TPair<FPackageId, FAsyncPackage2*>& Pair : AsyncLoadingThread->AsyncPackageLookup)
-			{
-				if (Pair.Value->Data.bCellExportsInitialized.load(std::memory_order_acquire))
+			AsyncLoadingThread->AsyncPackageLookup.ForEach(
+				[&Visitor](TPair<FPackageId, FAsyncPackage2*>& Pair)
 				{
-					for (FExportCell& CellExport : Pair.Value->Data.CellExports)
+					if (Pair.Value->Data.bCellExportsInitialized.load(std::memory_order_acquire))
 					{
-						Visitor.Visit(CellExport.Cell, TEXT(""));
+						for (FExportCell& CellExport : Pair.Value->Data.CellExports)
+						{
+							Visitor.Visit(CellExport.Cell, TEXT(""));
+						}
 					}
 				}
-			}
+			);
 		}
 
 		// Packages keep all their public cell exports alive while they are referenced by a loading package.
@@ -5102,7 +5146,6 @@ void FAsyncLoadingThread2::UpdatePackagePriority(FAsyncLoadingThreadState2& Thre
 {
 	TRACE_CPUPROFILER_EVENT_SCOPE(UpdatePackagePriority);
 	EAsyncPackageLoadingState2 LoadingState = Package->AsyncPackageLoadingState;
-	check(ThreadState.bCanAccessAsyncLoadingThreadData);
 
 	// TODO: Shader requests are not tracked so they won't be reprioritized correctly here
 	if (LoadingState <= EAsyncPackageLoadingState2::WaitingForIo)
@@ -5121,7 +5164,7 @@ void FAsyncLoadingThread2::UpdatePackagePriority(FAsyncLoadingThreadState2& Thre
 	}
 	if (LoadingState == EAsyncPackageLoadingState2::DeferredPostLoad)
 	{
-		if (ThreadState.bIsAsyncLoadingThread)
+		if (ThreadState.bIsAsyncLoadingThread || IsInParallelLoadingThread())
 		{
 			if (Package->TryAddRef())
 			{
@@ -5144,7 +5187,18 @@ void FAsyncLoadingThread2::UpdatePackagePriorityRecursive(FAsyncLoadingThreadSta
 	{
 		return;
 	}
+
 	Package->Desc.Priority = NewPriority;
+	UpdatePackagePriority(ThreadState, Package);
+
+	// For LinkerLoad, dependencies are loaded in ProcessPackageSummary so if this is called before,
+	// early out to avoid races and we'll reprocess them in the ProcessDependencies step
+	if (Package->LinkerLoadState.IsSet() && 
+		Package->AsyncPackageLoadingState <= EAsyncPackageLoadingState2::ProcessPackageSummary)
+	{
+		return;
+	}
+
 	for (FAsyncPackage2* ImportedPackage : Package->Data.ImportedAsyncPackages)
 	{
 		if (ImportedPackage)
@@ -5159,8 +5213,6 @@ void FAsyncLoadingThread2::UpdatePackagePriorityRecursive(FAsyncLoadingThreadSta
 		UpdatePackagePriorityRecursive(ThreadState, ImportedPackage, NewPriority);
 	}
 #endif
-
-	UpdatePackagePriority(ThreadState, Package);
 }
 
 void FAsyncLoadingThread2::ConditionalBeginPostLoad(FAsyncLoadingThreadState2& ThreadState, FAsyncLoadingPostLoadGroup* PostLoadGroup)
@@ -5170,6 +5222,8 @@ void FAsyncLoadingThread2::ConditionalBeginPostLoad(FAsyncLoadingThreadState2& T
 	check(ThreadState.bCanAccessAsyncLoadingThreadData);
 	if (PostLoadGroup->PackagesWithExportsToSerializeCount == 0)
 	{
+		PostLoadGroup->PackagesWithExportsToPostLoadCount = PostLoadGroup->Packages.Num();
+
 		// Release the post load node of packages in the post load group in reverse order that they were added to the group
 		// This usually means that dependencies will be post load first, similarly to how they are also serialized first
 		for (int32 Index = PostLoadGroup->Packages.Num() - 1; Index >= 0; --Index)
@@ -5184,12 +5238,9 @@ void FAsyncLoadingThread2::ConditionalBeginPostLoad(FAsyncLoadingThreadState2& T
 			
 			Package->AsyncPackageLoadingState = EAsyncPackageLoadingState2::PostLoad;
 			Package->ConditionalReleasePartialRequests(ThreadState);
-			for (int32 ExportBundleIndex = 0; ExportBundleIndex < Package->Data.TotalExportBundleCount; ++ExportBundleIndex)
-			{
-				Package->GetExportBundleNode(EEventLoadNode2::ExportBundle_PostLoad, ExportBundleIndex).ReleaseBarrier(&ThreadState);
-			}
+
+			Package->ConditionalBeginPostLoadPhase(ThreadState);
 		}
-		PostLoadGroup->PackagesWithExportsToPostLoadCount = PostLoadGroup->Packages.Num();
 	}
 }
 
@@ -5270,59 +5321,62 @@ void FAsyncLoadingThread2::MergePostLoadGroups(FAsyncLoadingThreadState2& Thread
 FAsyncPackage2* FAsyncLoadingThread2::FindOrInsertPackage(FAsyncLoadingThreadState2& ThreadState, FAsyncPackageDesc2& Desc, bool& bInserted, FAsyncPackage2* ImportedByPackage, TUniquePtr<FLoadPackageAsyncDelegate>&& PackageLoadedDelegate, TSharedPtr<FLoadPackageAsyncProgressDelegate> PackageProgressDelegate)
 {
 	TRACE_CPUPROFILER_EVENT_SCOPE(FindOrInsertPackage);
-	check(ThreadState.bCanAccessAsyncLoadingThreadData);
 	FAsyncPackage2* Package = nullptr;
 	bInserted = false;
 	{
-		FScopeLock LockAsyncPackages(&AsyncPackagesCritical);
-		Package = AsyncPackageLookup.FindRef(Desc.UPackageId);
-		if (!Package)
-		{
-			Package = CreateAsyncPackage(ThreadState, Desc);
-			checkf(Package, TEXT("Failed to create async package %s"), *Desc.UPackageName.ToString());
-			Package->AddRef();
-			++LoadingPackagesCounter;
-			TRACE_COUNTER_SET(AsyncLoadingLoadingPackages, LoadingPackagesCounter);
-			AsyncPackageLookup.Add(Desc.UPackageId, Package);
-			bInserted = true;
-		}
-		else
-		{
-			if (Desc.RequestID > 0)
-			{
-				Package->AddRequestID(ThreadState, Desc.RequestID);
-			}
-			if (Desc.Priority > Package->Desc.Priority)
+		AsyncPackageLookup.FindOrProduceAndApplyForWrite(Desc.UPackageId,
+			[&]() -> FAsyncPackage2*
 			{
-				TRACE_CPUPROFILER_EVENT_SCOPE(UpdatePackagePriority);
-				UpdatePackagePriorityRecursive(ThreadState, Package, Desc.Priority);
-			}
-		}
-		if (PackageLoadedDelegate.IsValid())
-		{
-			Package->AddCompletionCallback(MoveTemp(PackageLoadedDelegate));
-		}
-		if (PackageProgressDelegate.IsValid())
-		{
-			// To avoid missing progress callbacks for package already in-flight,
-			// call any state that the package might already have reached.
-			if (Package->AsyncPackageLoadingState >= EAsyncPackageLoadingState2::WaitingForIo)
+				Package = CreateAsyncPackage(ThreadState, Desc);
+				checkf(Package, TEXT("Failed to create async package %s"), *Desc.UPackageName.ToString());
+				Package->AddRef();
+				++LoadingPackagesCounter;
+				TRACE_COUNTER_SET(AsyncLoadingLoadingPackages, LoadingPackagesCounter);
+				bInserted = true;
+				return Package;
+			},
+			[&](FAsyncPackage2* InPackage)
 			{
-				Package->CallProgressCallback(*PackageProgressDelegate, EAsyncLoadingProgress::Started);
-			}
+				Package = InPackage;
 
-			if (Package->AsyncPackageLoadingState >= EAsyncPackageLoadingState2::ProcessPackageSummary)
-			{
-				Package->CallProgressCallback(*PackageProgressDelegate, EAsyncLoadingProgress::Read);
-			}
+				if (Desc.RequestID > 0)
+				{
+					Package->AddRequestID(ThreadState, Desc.RequestID);
+				}
+				if (Desc.Priority > Package->Desc.Priority)
+				{
+					TRACE_CPUPROFILER_EVENT_SCOPE(UpdatePackagePriority);
+					UpdatePackagePriorityRecursive(ThreadState, Package, Desc.Priority);
+				}
 
-			if (Package->AsyncPackageLoadingState >= EAsyncPackageLoadingState2::ExportsDone)
-			{
-				Package->CallProgressCallback(*PackageProgressDelegate, EAsyncLoadingProgress::Serialized);
-			}
+				if (PackageLoadedDelegate.IsValid())
+				{
+					Package->AddCompletionCallback(MoveTemp(PackageLoadedDelegate));
+				}
 
-			Package->AddProgressCallback(PackageProgressDelegate);
-		}
+				if (PackageProgressDelegate.IsValid())
+				{
+					// To avoid missing progress callbacks for package already in-flight,
+					// call any state that the package might already have reached.
+					if (Package->AsyncPackageLoadingState >= EAsyncPackageLoadingState2::WaitingForIo)
+					{
+						Package->CallProgressCallback(*PackageProgressDelegate, EAsyncLoadingProgress::Started);
+					}
+
+					if (Package->AsyncPackageLoadingState >= EAsyncPackageLoadingState2::ProcessPackageSummary)
+					{
+						Package->CallProgressCallback(*PackageProgressDelegate, EAsyncLoadingProgress::Read);
+					}
+
+					if (Package->AsyncPackageLoadingState >= EAsyncPackageLoadingState2::ExportsDone)
+					{
+						Package->CallProgressCallback(*PackageProgressDelegate, EAsyncLoadingProgress::Serialized);
+					}
+
+					Package->AddProgressCallback(PackageProgressDelegate);
+				}
+			}
+		);
 	}
 
 	// PostLoadGroup are to protect some race conditions happening in runtime when ALT is active but
@@ -5387,6 +5441,12 @@ void FAsyncLoadingThread2::IncludePackageInSyncLoadContextRecursive(FAsyncLoadin
 
 	Package->SyncLoadContextId = ContextId;
 
+	if (Package->Desc.Priority < MAX_int32)
+	{
+		Package->Desc.Priority = MAX_int32;
+		UpdatePackagePriority(ThreadState, Package);
+	}
+
 	FAsyncLoadingPostLoadGroup* PostLoadGroup = Package->PostLoadGroup ? Package->PostLoadGroup : Package->DeferredPostLoadGroup;
 	if (PostLoadGroup && PostLoadGroup->SyncLoadContextId < ContextId)
 	{
@@ -5400,6 +5460,14 @@ void FAsyncLoadingThread2::IncludePackageInSyncLoadContextRecursive(FAsyncLoadin
 		}
 	}
 
+	// For LinkerLoad, dependencies are loaded in ProcessPackageSummary so if this is called before,
+	// early out to avoid races and we'll reprocess them in the ProcessDependencies step
+	if (Package->LinkerLoadState.IsSet() &&
+		Package->AsyncPackageLoadingState <= EAsyncPackageLoadingState2::ProcessPackageSummary)
+	{
+		return;
+	}
+
 	for (FAsyncPackage2* ImportedPackage : Package->Data.ImportedAsyncPackages)
 	{
 		if (ImportedPackage && ImportedPackage->SyncLoadContextId < ContextId)
@@ -5417,12 +5485,6 @@ void FAsyncLoadingThread2::IncludePackageInSyncLoadContextRecursive(FAsyncLoadin
 		}
 	}
 #endif
-
-	if (Package->Desc.Priority < MAX_int32)
-	{
-		Package->Desc.Priority = MAX_int32;
-		UpdatePackagePriority(ThreadState, Package);
-	}
 }
 
 UE_TRACE_EVENT_BEGIN(CUSTOM_LOADTIMER_LOG, CreateAsyncPackage, NoSync)
@@ -5731,21 +5793,85 @@ void FEventLoadNode2::ReleaseBarrier(FAsyncLoadingThreadState2* ThreadState)
 	}
 }
 
+class FAsyncLoadingThreadState2Scope
+{
+public:
+	FAsyncLoadingThreadState2Scope(FAsyncLoadingThreadState2& ThreadState)
+	{
+		check(FAsyncLoadingThreadState2::Get() == nullptr);
+		FAsyncLoadingThreadState2::Set(&ThreadState);
+	}
+
+	~FAsyncLoadingThreadState2Scope()
+	{
+		check(FAsyncLoadingThreadState2::Get()->CurrentlyExecutingEventNodeStack.Num() == 0);
+		FAsyncLoadingThreadState2::Reset();
+	}
+
+	FAsyncLoadingThreadState2& Get()
+	{
+		return *FAsyncLoadingThreadState2::Get();
+	}
+};
+
+void FEventLoadNode2::ParallelLoadingWork()
+{
+	FTaskTagScope Scope(ETaskTag::EParallelLoadingThread);
+	FAsyncLoadingThreadState2 LocalThreadState;
+	FAsyncLoadingThreadState2Scope WorkerStateScope(LocalThreadState);
+
+	FUObjectThreadContext& ThreadContext = FUObjectThreadContext::Get();
+	// Firing a node can come from anywhere, needs to remove any current visibility filter
+	// A proper visibility filter is expected to be setup properly in each different nodes
+	TGuardValue GuardVisibilityFilter(ThreadContext.AsyncVisibilityFilter, EInternalObjectFlags::None);
+
+	// Used in NotifyConstructedDuringAsyncLoading to dispatch any created objects
+	// to this package if they don't belong to any other package.
+	ThreadContext.AsyncPackage = Package;
+
+	EEventLoadNodeExecutionResult Result = Execute(LocalThreadState);
+	check(Result == EEventLoadNodeExecutionResult::Complete);
+	
+	ThreadContext.AsyncPackage = nullptr;
+
+	// This is never used for zen loader, everything is handled in NotifyConstructedDuringAsyncLoading
+	// Reset to avoid accumulation and shutdown error complaining about objects being present in the context.
+	ThreadContext.GetSerializeContext()->PRIVATE_GetObjectsLoadedInternalUseOnly().Reset();
+}
+
 void FEventLoadNode2::Fire(FAsyncLoadingThreadState2* ThreadState)
 {
-	//TRACE_CPUPROFILER_EVENT_SCOPE(Fire);
-	if (Spec->bExecuteImmediately && ThreadState)
+	if (Spec->ExecutionType == EExecutionType::Immediate && ThreadState)
 	{
 		// Firing a node can come from anywhere, needs to remove any current visibility filter
 		TGuardValue GuardVisibilityFilter(FUObjectThreadContext::Get().AsyncVisibilityFilter, EInternalObjectFlags::None);
 
 		EEventLoadNodeExecutionResult Result = Execute(*ThreadState);
 		check(Result == EEventLoadNodeExecutionResult::Complete);
+
+		return;
 	}
-	else
+	else if (GAllowMultithreadedLoading && Spec->ExecutionType == EExecutionType::Parallel && Package->CanRunNodeAsync(Spec))
 	{
-		Spec->EventQueue->Push(ThreadState, this);
+		// We need the GC lock to go wide, otherwise we just queue the task normally since a GC is imminent
+		TSharedPtr<FGCScopeTryGuard> WorkerGCGuard = MakeShared<FGCScopeTryGuard>();
+		if (WorkerGCGuard->LockSucceeded())
+		{
+			UE::Tasks::Launch(
+				TEXT("ParallelLoadingWork"),
+				[this, WorkerGCGuard = MoveTemp(WorkerGCGuard)]()
+				{
+					ParallelLoadingWork();
+				}
+			);
+
+			return;
+		}
+
+		// fallthrough to push the task to the queue
 	}
+	
+	Spec->EventQueue->Push(ThreadState, this);
 }
 
 EEventLoadNodeExecutionResult FEventLoadNode2::Execute(FAsyncLoadingThreadState2& ThreadState)
@@ -6288,9 +6414,6 @@ void FAsyncPackage2::ImportPackagesRecursiveInner(FAsyncLoadingThreadState2& Thr
 	{
 		return;
 	}
-
-	TRACE_CPUPROFILER_EVENT_SCOPE(ImportPackagesRecursiveInner);
-
 	TArray<FName>& ImportedPackageNames = Header.ImportedPackageNames;
 	bool bHasImportedPackageNames = !ImportedPackageNames.IsEmpty();
 	check(ImportedPackageNames.Num() == 0 || ImportedPackageNames.Num() == ImportedPackageCount);
@@ -6659,6 +6782,13 @@ void FAsyncPackage2::ImportPackagesRecursive(FAsyncLoadingThreadState2& ThreadSt
 				AsyncLoadingThread.IncludePackageInSyncLoadContextRecursive(ThreadState, SyncLoadContextId, ImportedPackage);
 			}
 		}
+
+#if WITH_PARTIAL_REQUEST_DURING_RECURSION
+		for (FAsyncPackage2* ImportedPackage : AdditionalImportedAsyncPackages)
+		{
+			AsyncLoadingThread.IncludePackageInSyncLoadContextRecursive(ThreadState, SyncLoadContextId, ImportedPackage);
+		}
+#endif
 	}
 
 	UE_ASYNC_PACKAGE_LOG_VERBOSE(VeryVerbose, Desc, TEXT("ImportPackages: ImportsDone"),
@@ -7027,6 +7157,7 @@ EEventLoadNodeExecutionResult FAsyncPackage2::ProcessLinkerLoadPackageSummary(FA
 			PackageRef.PreInsertPublicExports({ LinkerLoadState->LinkerLoadHeaderData.ExportMap });
 
 #if WITH_METADATA
+			UE::TUniqueLock LinkerLock(LinkerLoadState->Linker->Mutex);
 			// Create metadata object, this needs to happen before any other package wants to use our exports
 			LinkerLoadState->MetaDataIndex = LinkerLoadState->Linker->LoadMetaDataFromExportMap(false);
 			if (LinkerLoadState->MetaDataIndex >= 0)
@@ -7047,18 +7178,17 @@ EEventLoadNodeExecutionResult FAsyncPackage2::ProcessLinkerLoadPackageSummary(FA
 		}
 	}
 
-	AsyncPackageLoadingState = EAsyncPackageLoadingState2::WaitingForDependencies;
-	if (!AsyncLoadingThread.bHasRegisteredAllScriptObjects)
-	{
-		SetupScriptDependencies();
-	}
-	GetPackageNode(Package_DependenciesReady).ReleaseBarrier(&ThreadState);
+	AsyncPackageLoadingState = EAsyncPackageLoadingState2::ProcessDependencies;
+	GetPackageNode(Package_ProcessDependencies).ReleaseBarrier(&ThreadState);
+
 	return EEventLoadNodeExecutionResult::Complete;
 }
 
 bool FAsyncPackage2::PreloadLinkerLoadExports(FAsyncLoadingThreadState2& ThreadState)
 {
 	// Serialize exports
+	UE::TUniqueLock LinkerLock(LinkerLoadState->Linker->Mutex);
+
 	const int32 ExportCount = LinkerLoadState->Linker->ExportMap.Num();
 	check(LinkerLoadState->Linker->ExportMap.Num() == Data.Exports.Num());
 	while (LinkerLoadState->SerializeExportIndex < ExportCount)
@@ -7199,6 +7329,8 @@ bool FAsyncPackage2::CreateLinkerLoadExports(FAsyncLoadingThreadState2& ThreadSt
 	check (!LinkerLoadState->bIsCurrentlyCreatingExports);
 	TGuardValue<bool> GuardIsCurrentlyCreatingExports(LinkerLoadState->bIsCurrentlyCreatingExports, true);
 
+	UE::TUniqueLock LinkerLock(LinkerLoadState->Linker->Mutex);
+
 	// Create exports
 	const int32 ExportCount = LinkerLoadState->Linker->ExportMap.Num();
 	while (LinkerLoadState->CreateExportIndex < ExportCount)
@@ -7441,6 +7573,25 @@ void FAsyncPackage2::ConditionalBeginResolveLinkerLoadImports(FAsyncLoadingThrea
 	);
 }
 
+void FAsyncPackage2::ConditionalBeginPostLoadPhase(FAsyncLoadingThreadState2& ThreadState)
+{
+	TRACE_CPUPROFILER_EVENT_SCOPE(ConditionalBeginPostLoadPhase);
+
+	// This phase is important so that all dependencies have finished serialization before we start postloading. Otherwise we could end up with RF_NeedLoad
+	// objects or worse, race conditions between postload trying to access properties, and serialize writing to those properties.
+	WaitForAllDependenciesToReachState(ThreadState, &FAsyncPackage2::AllDependenciesPostLoadReadyState, EAsyncPackageLoadingState2::PostLoad, AsyncLoadingThread.ConditionalBeginPostLoadTick,
+		[&ThreadState](FAsyncPackage2* Package)
+		{
+			check(Package->AsyncPackageLoadingState == EAsyncPackageLoadingState2::PostLoad);
+
+			for (int32 ExportBundleIndex = 0; ExportBundleIndex < Package->Data.TotalExportBundleCount; ++ExportBundleIndex)
+			{
+				Package->GetExportBundleNode(EEventLoadNode2::ExportBundle_PostLoad, ExportBundleIndex).ReleaseBarrier(&ThreadState);
+			}
+		}
+	);
+}
+
 EEventLoadNodeExecutionResult FAsyncPackage2::Event_ResolveLinkerLoadImports(FAsyncLoadingThreadState2& ThreadState, FAsyncPackage2* Package, int32)
 {
 	TRACE_CPUPROFILER_EVENT_SCOPE(Event_ResolveLinkerLoadImports);
@@ -7606,7 +7757,18 @@ EEventLoadNodeExecutionResult FAsyncPackage2::Event_ProcessPackageSummary(FAsync
 
 	Package->AsyncLoadingThread.FinishInitializeAsyncPackage(ThreadState, Package);
 
+	Package->AsyncPackageLoadingState = EAsyncPackageLoadingState2::ProcessDependencies;
+	Package->GetPackageNode(Package_ProcessDependencies).ReleaseBarrier(&ThreadState);
+
+	return EEventLoadNodeExecutionResult::Complete;
+}
+
+EEventLoadNodeExecutionResult FAsyncPackage2::Event_ProcessDependencies(FAsyncLoadingThreadState2& ThreadState, FAsyncPackage2* Package, int32)
+{
 	Package->AsyncPackageLoadingState = EAsyncPackageLoadingState2::WaitingForDependencies;
+
+	Package->AsyncLoadingThread.FinishImportingPackages(ThreadState, Package);
+
 	if (!Package->AsyncLoadingThread.bHasRegisteredAllScriptObjects)
 	{
 		Package->SetupScriptDependencies();
@@ -8181,6 +8343,28 @@ FString FAsyncPackage2::GetNameFromPackageObjectIndex(FPackageObjectIndex Index)
 	return Out;
 }
 
+bool FAsyncPackage2::CanRunNodeAsync(const FAsyncLoadEventSpec* Spec) const
+{
+	if (Spec->Func == &FAsyncPackage2::Event_PreloadLinkerLoadExports)
+	{
+		// UserDefinedStruct often try to load softobject ptr from text while
+		// they are preloading. This is very deadlock prone so exclude those
+		// types from going wide for now.
+		const int32 ExportCount = Data.Exports.Num();
+		for (int32 ExportIndex = 0; ExportIndex < ExportCount; ++ExportIndex)
+		{
+			FExportObject& Export = Data.Exports[ExportIndex];
+
+			if (UUserDefinedStruct* Struct = Cast<UUserDefinedStruct>(Export.Object))
+			{
+				return false;
+			}
+		}
+	}
+
+	return true;
+}
+
 void FAsyncPackage2::EventDrivenCreateExport(const FAsyncPackageHeaderData& Header, int32 LocalExportIndex)
 {
 	TRACE_CPUPROFILER_EVENT_SCOPE(CreateExport);
@@ -8632,57 +8816,63 @@ FAsyncPackage2* FAsyncPackage2::GetCurrentlyExecutingPackage(FAsyncLoadingThread
 void FAsyncPackage2::ConditionalReleasePartialRequests(FAsyncLoadingThreadState2& ThreadState)
 {
 #if WITH_PARTIAL_REQUEST_DURING_RECURSION
-	FAsyncLoadingSyncLoadContext* CurrentSyncLoadContext = ThreadState.SyncLoadContextStack.Num() ? ThreadState.SyncLoadContextStack.Top() : nullptr;
-	if (CurrentSyncLoadContext && CurrentSyncLoadContext->RequestingPackage)
-	{
-		EAsyncPackageLoadingState2 RequesterState = CurrentSyncLoadContext->RequestingPackage->AsyncPackageLoadingState;
+	TRACE_CPUPROFILER_EVENT_SCOPE(ConditionalReleasePartialRequests);
 
-		if (AsyncPackageLoadingState > RequesterState)
+	// We need to look at all syncload context now that its possible for any threads to push a context as they are not strictly recursive anymore.
+	// This means the request we're trying to satisfy is not necessarily on top of the sync context stack anymore.
+	for (FAsyncLoadingSyncLoadContext* SyncLoadContext : ThreadState.SyncLoadContextStack)
+	{
+		if (SyncLoadContext->RequestingPackage)
 		{
-			int32 Index = CurrentSyncLoadContext->RequestedPackages.Find(this);
-			if (Index != INDEX_NONE)
-			{
-				int32 RequestId = CurrentSyncLoadContext->RequestIDs[Index];
-				// Release the sync context request tied to our package, allowing the current flush to exit.
-				UE_LOG(LogStreaming, Log, TEXT("Package %s has reached state %s > %s, releasing request %d to allow recursive sync load to finish"),
-					*Desc.UPackageName.ToString(),
-					LexToString(AsyncPackageLoadingState),
-					LexToString(RequesterState),
-					RequestId
-				);
-				AsyncLoadingThread.RemovePendingRequests(ThreadState, { RequestId });
-			}
+			EAsyncPackageLoadingState2 RequesterState = SyncLoadContext->RequestingPackage->AsyncPackageLoadingState;
 
-			switch (AsyncPackageLoadingState)
+			if (AsyncPackageLoadingState > RequesterState)
 			{
-				case EAsyncPackageLoadingState2::ExportsDone:
+				int32 Index = SyncLoadContext->RequestedPackages.Find(this);
+				if (Index != INDEX_NONE)
 				{
-					// As a general rule, we want any synchronous load during deserialization to act like an import in order
-					// to make sure the requested package is part of the loading of the requesting package even if the
-					// sync load was released to avoid a deadlock. This will make sure that by the time the requesting package
-					// finishes loading, the requested package will be done too.
-					bool bIsAlreadyInSet = false;
-					CurrentSyncLoadContext->RequestingPackage->AdditionalImportedAsyncPackages.FindOrAdd(this, &bIsAlreadyInSet);
-					if (!bIsAlreadyInSet)
+					int32 RequestId = SyncLoadContext->RequestIDs[Index];
+					// Release the sync context request tied to our package, allowing the current flush to exit.
+					UE_LOG(LogStreaming, Log, TEXT("Package %s has reached state %s > %s, releasing request %d to allow recursive sync load to finish"),
+						*Desc.UPackageName.ToString(),
+						LexToString(AsyncPackageLoadingState),
+						LexToString(RequesterState),
+						RequestId
+					);
+					AsyncLoadingThread.RemovePendingRequests(ThreadState, { RequestId });
+
+					switch (AsyncPackageLoadingState)
 					{
-						this->AddRef();
+					case EAsyncPackageLoadingState2::ExportsDone:
+					{
+						// As a general rule, we want any synchronous load during deserialization to act like an import in order
+						// to make sure the requested package is part of the loading of the requesting package even if the
+						// sync load was released to avoid a deadlock. This will make sure that by the time the requesting package
+						// finishes loading, the requested package will be done too.
+						bool bIsAlreadyInSet = false;
+						SyncLoadContext->RequestingPackage->AdditionalImportedAsyncPackages.FindOrAdd(this, &bIsAlreadyInSet);
+						if (!bIsAlreadyInSet)
+						{
+							this->AddRef();
 
-						UE_LOG(LogStreaming, Display, TEXT("Package %s is adding a dynamic import to package %s because of a recursive sync load"),
-							*CurrentSyncLoadContext->RequestingPackage->Desc.UPackageName.ToString(),
-							*Desc.UPackageName.ToString()
-						);
+							UE_LOG(LogStreaming, Display, TEXT("Package %s is adding a dynamic import to package %s because of a recursive sync load"),
+								*SyncLoadContext->RequestingPackage->Desc.UPackageName.ToString(),
+								*Desc.UPackageName.ToString()
+							);
+						}
 					}
-				}
-				break;
-				case EAsyncPackageLoadingState2::PostLoad:
+					break;
+					case EAsyncPackageLoadingState2::PostLoad:
 					[[fallthrough]];
-				case EAsyncPackageLoadingState2::DeferredPostLoad:
+					case EAsyncPackageLoadingState2::DeferredPostLoad:
 					// We don't need to add a dynamic import if the package is already in postload since it wouldn't provide any benefits...
-				break;
-				default:
+					break;
+					default:
 					// If anything else begins to call ConditionalReleasePartialRequests, we'll need to consider what needs to be done here
 					// checkNoEntry();
-				break;
+					break;
+					}
+				}
 			}
 		}
 	}
@@ -9297,12 +9487,12 @@ void FAsyncLoadingThread2::UpdateSyncLoadContext(FAsyncLoadingThreadState2& Thre
 {
 	if (ThreadState.bIsAsyncLoadingThread && bAutoHandleSyncLoadContext)
 	{
-		FAsyncLoadingSyncLoadContext* CreatedOnMainThread;
-		while (ThreadState.SyncLoadContextsCreatedOnGameThread.Dequeue(CreatedOnMainThread))
+		FAsyncLoadingSyncLoadContext* CreatedOnOtherThread;
+		while (ThreadState.SyncLoadContextsCreatedOnOtherThreads.Dequeue(CreatedOnOtherThread))
 		{
-			ThreadState.SyncLoadContextStack.Push(CreatedOnMainThread);
+			ThreadState.SyncLoadContextStack.Push(CreatedOnOtherThread);
 
-			UE_LOG(LogStreaming, VeryVerbose, TEXT("Pushing ALT SyncLoadContext %" UINT64_FMT), CreatedOnMainThread->ContextId);
+			UE_LOG(LogStreaming, VeryVerbose, TEXT("Pushing ALT SyncLoadContext %" UINT64_FMT), CreatedOnOtherThread->ContextId);
 		}
 	}
 	if (ThreadState.SyncLoadContextStack.IsEmpty())
@@ -9317,6 +9507,7 @@ void FAsyncLoadingThread2::UpdateSyncLoadContext(FAsyncLoadingThreadState2& Thre
 		{
 			UE_LOG(LogStreaming, VeryVerbose, TEXT("Popping ALT SyncLoadContext %" UINT64_FMT), SyncLoadContext->ContextId);
 
+			SyncLoadContext->DoneEvent.Notify();
 			SyncLoadContext->ReleaseRef();
 			ThreadState.SyncLoadContextStack.Pop();
 			if (ThreadState.SyncLoadContextStack.IsEmpty())
@@ -9629,12 +9820,7 @@ EAsyncPackageState::Type FAsyncLoadingThread2::ProcessLoadedPackagesFromGameThre
 			Package->FinishUPackage();
 
 			{
-				FScopeLock LockAsyncPackages(&AsyncPackagesCritical);
 				AsyncPackageLookup.Remove(Package->Desc.UPackageId);
-				if (AsyncPackageLookup.IsEmpty())
-				{
-					AsyncPackageLookup.Empty(DefaultAsyncPackagesReserveCount);
-				}
 
 #if WITH_EDITOR
 				if (!Package->HasLoadFailed())
@@ -9838,23 +10024,23 @@ FAsyncLoadingThread2::FAsyncLoadingThread2(FIoDispatcher& InIoDispatcher, IAsync
 {
 	EventQueue.SetZenaphore(&AltZenaphore);
 
-	AsyncPackageLookup.Reserve(DefaultAsyncPackagesReserveCount);
 	PendingPackages.Reserve(DefaultAsyncPackagesReserveCount);
 	RequestIdToPackageMap.Reserve(DefaultAsyncPackagesReserveCount);
 
 	EventSpecs.AddDefaulted(EEventLoadNode2::Package_NumPhases + EEventLoadNode2::ExportBundle_NumPhases);
-	EventSpecs[EEventLoadNode2::Package_ProcessSummary] = { &FAsyncPackage2::Event_ProcessPackageSummary, &EventQueue, false, TEXT("ProcessSummary") };
-	EventSpecs[EEventLoadNode2::Package_DependenciesReady] = { &FAsyncPackage2::Event_DependenciesReady, &EventQueue, false, TEXT("DependenciesReady")};
+	EventSpecs[EEventLoadNode2::Package_ProcessSummary] = { &FAsyncPackage2::Event_ProcessPackageSummary, &EventQueue, EExecutionType::Parallel, TEXT("ProcessSummary") };
+	EventSpecs[EEventLoadNode2::Package_ProcessDependencies] = { &FAsyncPackage2::Event_ProcessDependencies, &EventQueue, EExecutionType::Normal, TEXT("ProcessDependencies") };
+	EventSpecs[EEventLoadNode2::Package_DependenciesReady] = { &FAsyncPackage2::Event_DependenciesReady, &EventQueue, EExecutionType::Normal, TEXT("DependenciesReady")};
 #if ALT2_ENABLE_LINKERLOAD_SUPPORT
-	EventSpecs[EEventLoadNode2::Package_CreateLinkerLoadExports] = { &FAsyncPackage2::Event_CreateLinkerLoadExports, &EventQueue, false, TEXT("CreateLinkerLoadExports") };
-	EventSpecs[EEventLoadNode2::Package_ResolveLinkerLoadImports] = { &FAsyncPackage2::Event_ResolveLinkerLoadImports, &EventQueue, false, TEXT("ResolveLinkerLoadImports") };
-	EventSpecs[EEventLoadNode2::Package_PreloadLinkerLoadExports] = { &FAsyncPackage2::Event_PreloadLinkerLoadExports, &EventQueue, false, TEXT("PreloadLinkerLoadExports") };
+	EventSpecs[EEventLoadNode2::Package_CreateLinkerLoadExports] = { &FAsyncPackage2::Event_CreateLinkerLoadExports, &EventQueue, EExecutionType::Normal, TEXT("CreateLinkerLoadExports") };
+	EventSpecs[EEventLoadNode2::Package_ResolveLinkerLoadImports] = { &FAsyncPackage2::Event_ResolveLinkerLoadImports, &EventQueue, EExecutionType::Normal, TEXT("ResolveLinkerLoadImports") };
+	EventSpecs[EEventLoadNode2::Package_PreloadLinkerLoadExports] = { &FAsyncPackage2::Event_PreloadLinkerLoadExports, &EventQueue, EExecutionType::Parallel, TEXT("PreloadLinkerLoadExports") };
 #endif
-	EventSpecs[EEventLoadNode2::Package_ExportsSerialized] = { &FAsyncPackage2::Event_ExportsDone, &EventQueue, true, TEXT("ExportsSerialized") };
+	EventSpecs[EEventLoadNode2::Package_ExportsSerialized] = { &FAsyncPackage2::Event_ExportsDone, &EventQueue, EExecutionType::Normal, TEXT("ExportsSerialized") };
 
-	EventSpecs[EEventLoadNode2::Package_NumPhases + EEventLoadNode2::ExportBundle_Process] = { &FAsyncPackage2::Event_ProcessExportBundle, &EventQueue, false, TEXT("ExportBundle_Process") };
-	EventSpecs[EEventLoadNode2::Package_NumPhases + EEventLoadNode2::ExportBundle_PostLoad] = { &FAsyncPackage2::Event_PostLoadExportBundle, &EventQueue, false, TEXT("ExportBundle_PostLoad") };
-	EventSpecs[EEventLoadNode2::Package_NumPhases + EEventLoadNode2::ExportBundle_DeferredPostLoad] = { &FAsyncPackage2::Event_DeferredPostLoadExportBundle, &MainThreadEventQueue, false, TEXT("ExportBundle_DeferredPostLoad") };
+	EventSpecs[EEventLoadNode2::Package_NumPhases + EEventLoadNode2::ExportBundle_Process] = { &FAsyncPackage2::Event_ProcessExportBundle, &EventQueue, EExecutionType::Normal, TEXT("ExportBundle_Process") };
+	EventSpecs[EEventLoadNode2::Package_NumPhases + EEventLoadNode2::ExportBundle_PostLoad] = { &FAsyncPackage2::Event_PostLoadExportBundle, &EventQueue, EExecutionType::Normal, TEXT("ExportBundle_PostLoad") };
+	EventSpecs[EEventLoadNode2::Package_NumPhases + EEventLoadNode2::ExportBundle_DeferredPostLoad] = { &FAsyncPackage2::Event_DeferredPostLoadExportBundle, &MainThreadEventQueue, EExecutionType::Normal, TEXT("ExportBundle_DeferredPostLoad") };
 
 	CancelLoadingEvent = FPlatformProcess::GetSynchEventFromPool();
 	ThreadSuspendedEvent = FPlatformProcess::GetSynchEventFromPool();
@@ -9864,7 +10050,8 @@ FAsyncLoadingThread2::FAsyncLoadingThread2(FIoDispatcher& InIoDispatcher, IAsync
 	FCoreUObjectDelegates::GetGarbageCollectStartedDelegate().AddRaw(this, &FAsyncLoadingThread2::OnGarbageCollectStarted);
 
 	FAsyncLoadingThreadState2::TlsSlot = FPlatformTLS::AllocTlsSlot();
-	GameThreadState = MakeUnique<FAsyncLoadingThreadState2>(GraphAllocator, IoDispatcher);
+	GameThreadState = MakeUnique<FAsyncLoadingThreadState2>();
+	GameThreadState->bCanAccessAsyncLoadingThreadData = true;
 	EventQueue.SetOwnerThread(GameThreadState.Get());
 	MainThreadEventQueue.SetOwnerThread(GameThreadState.Get());
 	FAsyncLoadingThreadState2::Set(GameThreadState.Get());
@@ -9923,7 +10110,7 @@ void FAsyncLoadingThread2::StartThread()
 
 	if (FAsyncLoadingThreadSettings::Get().bAsyncLoadingThreadEnabled && !Thread)
 	{
-		AsyncLoadingThreadState = MakeUnique<FAsyncLoadingThreadState2>(GraphAllocator, IoDispatcher);
+		AsyncLoadingThreadState = MakeUnique<FAsyncLoadingThreadState2>();
 		EventQueue.SetOwnerThread(AsyncLoadingThreadState.Get());
 
 		// When using ALT, we want to wake the main thread ASAP in case it's sleeping for lack of something to do during flush.
@@ -10503,14 +10690,21 @@ void FAsyncLoadingThread2::NotifyConstructedDuringAsyncLoading(UObject* Object,
 	const FPackageId ObjectPackageId = ObjectPackage->GetPackageId();
 	if (ObjectPackageId.IsValid())
 	{
-		FScopeLock LockAsyncPackages(&AsyncPackagesCritical);
-		if (FAsyncPackage2* AsyncPackage = AsyncPackageLookup.FindRef(ObjectPackageId))
-		{
-			if (AsyncPackage->AsyncPackageLoadingState < EAsyncPackageLoadingState2::DeferredPostLoad)
+		bool bReturn = false;
+		AsyncPackageLookup.FindAndApply(ObjectPackageId,
+			[&](FAsyncPackage2* AsyncPackage)
 			{
-				AsyncPackage->AddConstructedObject(Object, bSubObjectThatAlreadyExists);
-				return;
+				if (AsyncPackage->AsyncPackageLoadingState < EAsyncPackageLoadingState2::DeferredPostLoad)
+				{
+					AsyncPackage->AddConstructedObject(Object, bSubObjectThatAlreadyExists);
+					bReturn = true;
+				}
 			}
+		);
+
+		if (bReturn)
+		{
+			return;
 		}
 	}
 #endif
@@ -11296,6 +11490,38 @@ void FAsyncLoadingThread2::WarnAboutPotentialSyncLoadStall(FAsyncLoadingSyncLoad
 	}
 }
 
+void FAsyncLoadingThread2::FlushLoadingFromParallelLoadingThread(FAsyncLoadingThreadState2& ThreadState, TConstArrayView<int32> RequestIDs)
+{
+	TRACE_CPUPROFILER_EVENT_SCOPE(FlushLoadingFromParallelLoadingThread);
+
+	if (RequestIDs.IsEmpty())
+	{
+		return;
+	}
+
+	FAsyncLoadingSyncLoadContext* SyncLoadContext = new FAsyncLoadingSyncLoadContext(RequestIDs);
+
+	FAsyncPackage2* GetCurrentlyExecutingPackage = FAsyncPackage2::GetCurrentlyExecutingPackage(ThreadState);
+	SyncLoadContext->RequestingPackage = GetCurrentlyExecutingPackage;
+
+	UE_LOG(LogStreaming, VeryVerbose, TEXT("Pushing ALT Worker SyncLoadContext %" UINT64_FMT), SyncLoadContext->ContextId);
+	ThreadState.SyncLoadContextStack.Push(SyncLoadContext);
+
+	SyncLoadContext->AddRef();
+	AsyncLoadingThreadState->SyncLoadContextsCreatedOnOtherThreads.Enqueue(SyncLoadContext);
+	AltZenaphore.NotifyOne();
+
+	TRACE_CPUPROFILER_EVENT_FLUSH();
+	SyncLoadContext->DoneEvent.Wait();
+	check(!ContainsAnyRequestID(SyncLoadContext->RequestIDs));
+
+	check(ThreadState.SyncLoadContextStack.Top() == SyncLoadContext);
+	UE_LOG(LogStreaming, VeryVerbose, TEXT("Popping ALT Worker SyncLoadContext %" UINT64_FMT), ThreadState.SyncLoadContextStack.Top()->ContextId);
+	ThreadState.SyncLoadContextStack.Pop();
+	SyncLoadContext->ReleaseRef();
+	AltZenaphore.NotifyOne();
+}
+
 void FAsyncLoadingThread2::FlushLoadingFromLoadingThread(FAsyncLoadingThreadState2& ThreadState, TConstArrayView<int32> RequestIDs)
 {
 	TRACE_CPUPROFILER_EVENT_SCOPE(FlushLoadingFromLoadingThread);
@@ -11352,7 +11578,7 @@ void FAsyncLoadingThread2::FlushLoading(TConstArrayView<int32> RequestIDs)
 
 		// We can't possibly support flushing from async loading thread unless we have the partial request support active.
 #if WITH_PARTIAL_REQUEST_DURING_RECURSION
-		const bool bIsFlushSupportedOnCurrentThread = IsInGameThread() || IsInAsyncLoadingThread();
+		const bool bIsFlushSupportedOnCurrentThread = IsInGameThread() || IsInAsyncLoadingThread() || IsInParallelLoadingThread();
 #else
 		const bool bIsFlushSupportedOnCurrentThread = IsInGameThread();
 #endif
@@ -11391,6 +11617,12 @@ void FAsyncLoadingThread2::FlushLoading(TConstArrayView<int32> RequestIDs)
 			return;
 		}
 
+		if (ThreadState && IsInParallelLoadingThread())
+		{
+			FlushLoadingFromParallelLoadingThread(*ThreadState, RequestIDs);
+			return;
+		}
+
 #if WITH_EDITOR
 		// In the editor loading cannot be part of a transaction as it cannot be undone, and may result in recording half-loaded objects. So we suppress any active transaction while in this stack, and set the editor loading flag
 		TGuardValue<ITransaction*> SuppressTransaction(GUndo, nullptr);
@@ -11460,7 +11692,7 @@ void FAsyncLoadingThread2::FlushLoading(TConstArrayView<int32> RequestIDs)
 			if (AsyncLoadingThreadState)
 			{
 				SyncLoadContext->AddRef();
-				AsyncLoadingThreadState->SyncLoadContextsCreatedOnGameThread.Enqueue(SyncLoadContext);
+				AsyncLoadingThreadState->SyncLoadContextsCreatedOnOtherThreads.Enqueue(SyncLoadContext);
 				AltZenaphore.NotifyOne();
 			}
 		}
@@ -11504,8 +11736,6 @@ void FAsyncLoadingThread2::FlushLoading(TConstArrayView<int32> RequestIDs)
 					{
 #if ALT2_DUMP_STATE_ON_HANG
 						{
-							FScopeLock ScopeLock(&AsyncPackagesCritical);
-
 							if (GameThreadState->CurrentlyExecutingEventNodeStack.Num() > 0)
 							{
 								UE_LOG(LogStreaming, Warning, TEXT("============ Currently executing nodes on stack ============="));
@@ -11531,13 +11761,15 @@ void FAsyncLoadingThread2::FlushLoading(TConstArrayView<int32> RequestIDs)
 							}
 							else
 							{
-								for (auto& Pair : AsyncPackageLookup)
-								{
-									if (FAsyncPackage2* Package = Pair.Value)
+								AsyncPackageLookup.ForEach(
+									[](auto& Pair)
 									{
-										Package->DumpState();
+										if (FAsyncPackage2* Package = Pair.Value)
+										{
+											Package->DumpState();
+										}
 									}
-								}
+								);
 							}
 						}
 #endif // ALT2_DUMP_STATE_ON_HANG
@@ -11752,42 +11984,44 @@ void FAsyncPackage2::DumpStateImpl(TSet<FAsyncPackage2*>& Set, int32 Indent, TMu
 	{
 		MappedNodes = &LocalMappedNode;
 
-		FScopeLock ScopeLock(&AsyncLoadingThread.AsyncPackagesCritical);
-		for (auto& Pair : AsyncLoadingThread.AsyncPackageLookup)
-		{
-			if (FAsyncPackage2* Package = Pair.Value)
+		AsyncLoadingThread.AsyncPackageLookup.ForEach(
+			[&](auto& Pair)
 			{
-				for (uint8 PhaseIndex = 0; PhaseIndex < Package_NumPhases; ++PhaseIndex)
-				{
-					FEventLoadNode2& Node = Package->GetPackageNode((EEventLoadNode2)PhaseIndex);
-					VisitNodes(Node, [&](FEventLoadNode2* Dependent) { LocalMappedNode.AddUnique(Dependent, &Node); });
-				}
-
-				for (int32 ExportBundleIndex = 0; ExportBundleIndex < Package->Data.TotalExportBundleCount; ++ExportBundleIndex)
+				if (FAsyncPackage2* Package = Pair.Value)
 				{
-					for (uint8 PhaseIndex = 0; PhaseIndex < ExportBundle_NumPhases; ++PhaseIndex)
+					for (uint8 PhaseIndex = 0; PhaseIndex < Package_NumPhases; ++PhaseIndex)
 					{
-						EEventLoadNode2 EventNode = (EEventLoadNode2)PhaseIndex;
-						FEventLoadNode2& Node = Package->GetExportBundleNode(EventNode, ExportBundleIndex);
+						FEventLoadNode2& Node = Package->GetPackageNode((EEventLoadNode2)PhaseIndex);
 						VisitNodes(Node, [&](FEventLoadNode2* Dependent) { LocalMappedNode.AddUnique(Dependent, &Node); });
 					}
+
+					for (int32 ExportBundleIndex = 0; ExportBundleIndex < Package->Data.TotalExportBundleCount; ++ExportBundleIndex)
+					{
+						for (uint8 PhaseIndex = 0; PhaseIndex < ExportBundle_NumPhases; ++PhaseIndex)
+						{
+							EEventLoadNode2 EventNode = (EEventLoadNode2)PhaseIndex;
+							FEventLoadNode2& Node = Package->GetExportBundleNode(EventNode, ExportBundleIndex);
+							VisitNodes(Node, [&](FEventLoadNode2* Dependent) { LocalMappedNode.AddUnique(Dependent, &Node); });
+						}
+					}
 				}
 			}
-		}
+		);
 	}
 
 	{
-		FScopeLock ScopeLock(&AsyncLoadingThread.AsyncPackagesCritical);
-		for (auto& Pair : AsyncLoadingThread.AsyncPackageLookup)
-		{
-			if (FAsyncPackage2* Package = Pair.Value)
+		AsyncLoadingThread.AsyncPackageLookup.ForEach(
+			[&](auto& Pair)
 			{
-				if (Package->HeaderData.ImportedPackageIds.Contains(this->Desc.UPackageId))
+				if (FAsyncPackage2* Package = Pair.Value)
 				{
-					UE_LOG(LogStreaming, Warning, TEXT(" %s referenced by package %s"), FCString::Spc(Indent), *FormatPackage(Package));
+					if (Package->HeaderData.ImportedPackageIds.Contains(this->Desc.UPackageId))
+					{
+						UE_LOG(LogStreaming, Warning, TEXT(" %s referenced by package %s"), FCString::Spc(Indent), *FormatPackage(Package));
+					}
 				}
 			}
-		}
+		);
 	}
 
 	for (uint8 PhaseIndex = 0; PhaseIndex < Package_NumPhases; ++PhaseIndex)
@@ -11838,6 +12072,11 @@ void FAsyncPackage2::DumpStateImpl(TSet<FAsyncPackage2*>& Set, int32 Indent, TMu
 		UE_LOG(LogStreaming, Warning, TEXT(" %s AllDependenciesImportState is waiting on %s"), FCString::Spc(Indent), *FormatPackage(AllDependenciesImportState.WaitingForPackage));
 	}
 #endif
+	if (AllDependenciesPostLoadReadyState.WaitingForPackage)
+	{
+		UE_LOG(LogStreaming, Warning, TEXT(" %s AllDependenciesPostLoadReadyState is waiting on %s"), FCString::Spc(Indent), *FormatPackage(AllDependenciesPostLoadReadyState.WaitingForPackage));
+	}
+
 	if (AllDependenciesFullyLoadedState.WaitingForPackage)
 	{
 		UE_LOG(LogStreaming, Warning, TEXT(" %s AllDependenciesFullyLoadedState is waiting on %s"), FCString::Spc(Indent), *FormatPackage(AllDependenciesFullyLoadedState.WaitingForPackage));
@@ -12307,4 +12546,4 @@ bool FPublicExportMapPreInsertPublicExportsTest::RunTest(const FString& Paramete
 	return true;
 }
 
-#endif // WITH_LOW_LEVEL_TESTS
+#endif // WITH_LOW_LEVEL_TESTS
\ No newline at end of file
