[ZenLoader]
- Fix asymmetric async GC locks caused by reserving and releasing in different threads

diff --git a/Engine/Source/Runtime/CoreUObject/Private/Serialization/AsyncLoading2.cpp b/Engine/Source/Runtime/CoreUObject/Private/Serialization/AsyncLoading2.cpp
--- a/Engine/Source/Runtime/CoreUObject/Private/Serialization/AsyncLoading2.cpp
+++ b/Engine/Source/Runtime/CoreUObject/Private/Serialization/AsyncLoading2.cpp
@@ -5854,14 +5854,19 @@ void FEventLoadNode2::Fire(FAsyncLoadingThreadState2* ThreadState)
 	else if (GAllowMultithreadedLoading && Spec->ExecutionType == EExecutionType::Parallel && Package->CanRunNodeAsync(Spec))
 	{
 		// We need the GC lock to go wide, otherwise we just queue the task normally since a GC is imminent
-		TSharedPtr<FGCScopeTryGuard> WorkerGCGuard = MakeShared<FGCScopeTryGuard>();
-		if (WorkerGCGuard->LockSucceeded())
+
+		// Always take the async lock to avoid imbalance since we're releasing in a different thread
+		constexpr bool bShouldLock = true;
+
+		if (FGCCSyncObject::Get().TryLockAsync(bShouldLock))
 		{
 			UE::Tasks::Launch(
 				TEXT("ParallelLoadingWork"),
-				[this, WorkerGCGuard = MoveTemp(WorkerGCGuard)]()
+				[this]()
 				{
 					ParallelLoadingWork();
+
+					FGCCSyncObject::Get().UnlockAsync(bShouldLock);
 				}
 			);
 
diff --git a/Engine/Source/Runtime/CoreUObject/Private/UObject/GCScopeLock.h b/Engine/Source/Runtime/CoreUObject/Private/UObject/GCScopeLock.h
--- a/Engine/Source/Runtime/CoreUObject/Private/UObject/GCScopeLock.h
+++ b/Engine/Source/Runtime/CoreUObject/Private/UObject/GCScopeLock.h
@@ -77,10 +77,10 @@ public:
 	}
 
 	/** Attempt async Lock **/
-	bool TryLockAsync()
+	bool TryLockAsync(bool bShouldLock = !IsInGameThread())
 	{
 		bool bLocked = true;
-		if (!IsInGameThread())
+		if (bShouldLock)
 		{
 			bLocked = false;
 			if (GCCounter.GetValue() == 0)
@@ -98,9 +98,9 @@ public:
 	}
 
 	/** Release lock from non-game thread */
-	void UnlockAsync()
+	void UnlockAsync(bool bShouldUnlock = !IsInGameThread())
 	{
-		if (!IsInGameThread())
+		if (bShouldUnlock)
 		{
 			AsyncCounter.Decrement();
 		}
