Document easy footgun of FGCObject (not trivially relocatable)

#rnx

diff --git a/Engine/Source/Runtime/CoreUObject/Public/UObject/GCObject.h b/Engine/Source/Runtime/CoreUObject/Public/UObject/GCObject.h
--- a/Engine/Source/Runtime/CoreUObject/Public/UObject/GCObject.h
+++ b/Engine/Source/Runtime/CoreUObject/Public/UObject/GCObject.h
@@ -109,6 +109,15 @@ public:
  * This class provides common registration for garbage collection for
  * non-UObject classes. It is an abstract base class requiring you to implement
  * the AddReferencedObjects() method.
+ *
+ * Note: FGCObject instances are not trivially relocatable, and are therefore not
+ * suited for being stored in Unreal container types by value.
+ * For example, this won't work:
+ * class FMyClass : public FGCObject {};
+ * struct FSomeStruct {
+ *     TArray<FMyClass> Instances; // Storing anything here might crash GC at 
+ *                                 // reachability analysis
+ * };
  */
 class FGCObject
 {
