Adding UE Insights Events to help profile Remote Object functions.

diff --git a/Engine/Source/Runtime/CoreUObject/Private/UObject/RemoteExecutor.cpp b/Engine/Source/Runtime/CoreUObject/Private/UObject/RemoteExecutor.cpp
--- a/Engine/Source/Runtime/CoreUObject/Private/UObject/RemoteExecutor.cpp
+++ b/Engine/Source/Runtime/CoreUObject/Private/UObject/RemoteExecutor.cpp
@@ -123,6 +123,8 @@ public:
 
 		for (int32 LocalIterationNumber = 0; ; LocalIterationNumber++)
 		{
+			TRACE_CPUPROFILER_EVENT_SCOPE(ResumeExecutePendingWork)
+
 			double Now = FPlatformTime::Seconds();
 			bool bPrintStall = false;
 
@@ -299,6 +301,8 @@ public:
 
 					ExecutingWork->ExecutionAttempts++;
 
+					TRACE_CPUPROFILER_EVENT_SCOPE_TEXT(ExecutingWork->Name)
+
 					FRemoteTransactionId RequestId = ExecutingWork->RequestId;
 					UE::RemoteExecutor::OnTransactionStartingDelegate.Broadcast(RequestId, ExecutingWork->Name);
 					AutoRTFM::ETransactionResult WorkTransactResult =
@@ -323,6 +327,8 @@ public:
 
 								if (ExecutingWork->bRequiresMultiServerCommit)
 								{
+									TRACE_CPUPROFILER_EVENT_SCOPE(MultiServerCommit)
+
 									// we're done with the work and about to commit - 
 									// first we need to send borrowed objects back
 									// to their owner and ask them if they are ready
diff --git a/Engine/Source/Runtime/CoreUObject/Private/UObject/RemoteObjectSerialization.cpp b/Engine/Source/Runtime/CoreUObject/Private/UObject/RemoteObjectSerialization.cpp
--- a/Engine/Source/Runtime/CoreUObject/Private/UObject/RemoteObjectSerialization.cpp
+++ b/Engine/Source/Runtime/CoreUObject/Private/UObject/RemoteObjectSerialization.cpp
@@ -1214,6 +1214,7 @@ int32 DeserializeObjectData(FRemoteObjectData& ObjectData,
 {
 	using namespace UE::RemoteObject::Handle;
 	using namespace UE::RemoteObject::Private;
+	TRACE_CPUPROFILER_EVENT_SCOPE(DeserializeObjectData);
 
 	int32 Version = 0;
 	FRemoteObjectId RequestedObjectId;
@@ -1273,13 +1274,17 @@ int32 DeserializeObjectData(FRemoteObjectData& ObjectData,
 		}
 
 		// Construct all objects first
-		FRemoteObjectConstructionOverrides ConstructionOverrides(ObjectHeaders);
-		OutObjectRemoteIds.Reserve(ObjectHeaders.Num());
-		OutObjects.Reserve(ObjectHeaders.Num());
-		for (int32 ObjectIndex = 0; ObjectIndex < ObjectHeaders.Num(); ++ObjectIndex)
 		{
-			OutObjectRemoteIds.Add(ObjectHeaders[ObjectIndex].RemoteId);
-			OutObjects.Add(ConstructRemoteObject(ObjectHeaders[ObjectIndex], &ConstructionOverrides, DeserializeFlags));
+			TRACE_CPUPROFILER_EVENT_SCOPE(ConstructRemoteObjects);
+
+			FRemoteObjectConstructionOverrides ConstructionOverrides(ObjectHeaders);
+			OutObjectRemoteIds.Reserve(ObjectHeaders.Num());
+			OutObjects.Reserve(ObjectHeaders.Num());
+			for (int32 ObjectIndex = 0; ObjectIndex < ObjectHeaders.Num(); ++ObjectIndex)
+			{
+				OutObjectRemoteIds.Add(ObjectHeaders[ObjectIndex].RemoteId);
+				OutObjects.Add(ConstructRemoteObject(ObjectHeaders[ObjectIndex], &ConstructionOverrides, DeserializeFlags));
+			}
 		}
 
 		UObject* RootObject = OutObjects.Num() ? OutObjects[0] : nullptr;
@@ -1308,27 +1313,33 @@ int32 DeserializeObjectData(FRemoteObjectData& ObjectData,
 		}
 
 		// Deserialize all objects
-		Ar.SetRootObject(OutObjects.Num() ? OutObjects[0] : nullptr);
-		for (int32 ObjectIndex = 0; ObjectIndex < OutObjects.Num(); ++ObjectIndex)
 		{
-			if (UObject* Object = OutObjects[ObjectIndex])
+			TRACE_CPUPROFILER_EVENT_SCOPE(DeserializeIntoCreatedObjects);
+
+			Ar.SetRootObject(OutObjects.Num() ? OutObjects[0] : nullptr);
+			for (int32 ObjectIndex = 0; ObjectIndex < OutObjects.Num(); ++ObjectIndex)
 			{
-				FRemoteObjectReaderScope Scope(Ar, Object);
-				Object->Serialize(Ar);
-				if (RequestedObjectIndex == -1 && RequestedObjectId == FRemoteObjectId(Object))
+				if (UObject* Object = OutObjects[ObjectIndex])
 				{
-					RequestedObjectIndex = ObjectIndex;
+					FRemoteObjectReaderScope Scope(Ar, Object);
+					Object->Serialize(Ar);
+					if (RequestedObjectIndex == -1 && RequestedObjectId == FRemoteObjectId(Object))
+					{
+						RequestedObjectIndex = ObjectIndex;
+					}
+				}
+				else
+				{
+					Ar.Seek(ObjectHeaders[ObjectIndex].NextOffset);
 				}
-			}
-			else
-			{
-				Ar.Seek(ObjectHeaders[ObjectIndex].NextOffset);
 			}
 		}
 
 		// Patch any unresolved remote references (skip if we're resetting an object to its CDO state)
 		if (!bResetting)
 		{
+			TRACE_CPUPROFILER_EVENT_SCOPE(FArchiveRemoteReferencePatcher);
+
 			FArchiveRemoteReferencePatcher PatchAr(OutObjects);
 			for (UObject* Object : OutObjects)
 			{
diff --git a/Engine/Source/Runtime/CoreUObject/Private/UObject/RemoteObjectTransfer.cpp b/Engine/Source/Runtime/CoreUObject/Private/UObject/RemoteObjectTransfer.cpp
--- a/Engine/Source/Runtime/CoreUObject/Private/UObject/RemoteObjectTransfer.cpp
+++ b/Engine/Source/Runtime/CoreUObject/Private/UObject/RemoteObjectTransfer.cpp
@@ -822,6 +822,8 @@ UObject* MigrateObjectFromObjectData(FRemoteObjectData& ObjectData, const FUObje
 	checkf(ReceivedObjects.Num() > 0, TEXT("PeerId:%s unable to deserialize object data (%d bytes)"), *GetGlobalServerId().ToString(), ObjectData.GetNumBytes());
 	if (ReceivedObjects.Num() > 0)
 	{
+		TRACE_CPUPROFILER_EVENT_SCOPE(CheckForNewOrRemovedSubObjects);
+
 		if (UE_LOG_ACTIVE(LogRemoteObject, VeryVerbose))
 		{
 			UE_LOG(LogRemoteObject, VeryVerbose, TEXT("PeerId:%s deserialized object data (%d bytes)"), *GetGlobalServerId().ToString(), ObjectData.GetNumBytes());
@@ -883,9 +885,13 @@ UObject* MigrateObjectFromObjectData(FRemoteObjectData& ObjectData, const FUObje
 			MarkAsBorrowed(ReceivedObject);
 		}
 	}
-	for (UObject* ReceivedObject : ReceivedObjects)
 	{
-		ReceivedObject->PostMigrate(MigrationContext);
+		TRACE_CPUPROFILER_EVENT_SCOPE(PostMigrateObjects);
+
+		for (UObject* ReceivedObject : ReceivedObjects)
+		{
+			ReceivedObject->PostMigrate(MigrationContext);
+		}
 	}
 #endif
 
diff --git a/Engine/Source/Runtime/Engine/Private/DataReplication.cpp b/Engine/Source/Runtime/Engine/Private/DataReplication.cpp
--- a/Engine/Source/Runtime/Engine/Private/DataReplication.cpp
+++ b/Engine/Source/Runtime/Engine/Private/DataReplication.cpp
@@ -1315,6 +1315,8 @@ bool FObjectReplicator::ReceivedRPC(FNetBitReader& Reader, const FReplicationFla
 
 	if (RPCDoS != nullptr && bIsServer && RPCDoS->IsRPCDoSDetectionEnabled() && !Connection->IsReplay())
 	{
+		TRACE_CPUPROFILER_EVENT_SCOPE(RPCDosDetection);
+
 		ScopedTracker.RPCDoS = RPCDoS;
 
 		if (UNLIKELY(RPCDoS->ShouldMonitorReceivedRPC()))
@@ -1342,6 +1344,8 @@ bool FObjectReplicator::ReceivedRPC(FNetBitReader& Reader, const FReplicationFla
 
 	if (bCanExecute)
 	{
+		TRACE_CPUPROFILER_EVENT_SCOPE(EvaluateCanExecute);
+
 		// Only delay if reliable and CVar is enabled
 		bool bCanDelayUnmapped = bCanDelayRPC && Function->FunctionFlags & FUNC_NetReliable;
 
diff --git a/Engine/Source/Runtime/Engine/Private/NetDriver.cpp b/Engine/Source/Runtime/Engine/Private/NetDriver.cpp
--- a/Engine/Source/Runtime/Engine/Private/NetDriver.cpp
+++ b/Engine/Source/Runtime/Engine/Private/NetDriver.cpp
@@ -216,6 +216,7 @@ namespace UE::Net::Private
 
 	TArray<UE::Net::Private::FAutoDestructProperty> CopyOutParametersToLocalParameters(UFunction* Function, FOutParmRec* OutParms, void* LocalParms, UObject* TargetObj)
 	{
+		TRACE_CPUPROFILER_EVENT_SCOPE(CopyOutParametersToLocalParameters);
 		TArray<UE::Net::Private::FAutoDestructProperty> CopiedProperties;
 
 		// Look for CPF_OutParm's, we'll need to copy these into the local parameter memory manually
diff --git a/Engine/Source/Runtime/Experimental/Iris/Core/Private/Iris/ReplicationSystem/ReplicationSystem.cpp b/Engine/Source/Runtime/Experimental/Iris/Core/Private/Iris/ReplicationSystem/ReplicationSystem.cpp
--- a/Engine/Source/Runtime/Experimental/Iris/Core/Private/Iris/ReplicationSystem/ReplicationSystem.cpp
+++ b/Engine/Source/Runtime/Experimental/Iris/Core/Private/Iris/ReplicationSystem/ReplicationSystem.cpp
@@ -1289,6 +1289,7 @@ bool UReplicationSystem::SendRPC(const UObject* RootObject, const UObject* SubOb
 {
 	using namespace UE::Net;
 	using namespace UE::Net::Private;
+	TRACE_CPUPROFILER_EVENT_SCOPE(SendRPC_SendMulticastRPC);
 
 	ENetObjectAttachmentSendPolicyFlags SendFlags = ENetObjectAttachmentSendPolicyFlags::None;
 	if (ReplicationSystemCVars::bAllowAttachmentSendPolicyFlags)
@@ -1308,6 +1309,8 @@ bool UReplicationSystem::SendRPC(const UObject* RootObject, const UObject* SubOb
 bool UReplicationSystem::SendRPC(uint32 ConnectionId, const UObject* RootObject, const UObject* SubObject, const UFunction* Function, const void* Parameters)
 {
 	using namespace UE::Net::Private;
+	TRACE_CPUPROFILER_EVENT_SCOPE(SendRPC_SendUnicastRPC);
+	
 	FNetBlobManager::FSendRPCContext RPCContext = { .RootObject = RootObject, .SubObject = SubObject, .Function = Function };
 	FNetBlobManager& NetBlobManager = Impl->ReplicationSystemInternal.GetNetBlobManager();
 	return NetBlobManager.SendUnicastRPC(ConnectionId, RPCContext, Parameters);
