[GCLock]
- Add support to take a lock even on the game thread in case we expect to release it on another thread

diff --git a/Engine/Source/Runtime/CoreUObject/Private/UObject/GCScopeLock.h b/Engine/Source/Runtime/CoreUObject/Private/UObject/GCScopeLock.h
--- a/Engine/Source/Runtime/CoreUObject/Private/UObject/GCScopeLock.h
+++ b/Engine/Source/Runtime/CoreUObject/Private/UObject/GCScopeLock.h
@@ -77,10 +77,10 @@ public:
 	}
 
 	/** Attempt async Lock **/
-	bool TryLockAsync()
+	bool TryLockAsync(bool bShouldLock = !IsInGameThread())
 	{
 		bool bLocked = true;
-		if (!IsInGameThread())
+		if (bShouldLock)
 		{
 			bLocked = false;
 			if (GCCounter.GetValue() == 0)
@@ -98,9 +98,9 @@ public:
 	}
 
 	/** Release lock from non-game thread */
-	void UnlockAsync()
+	void UnlockAsync(bool bShouldUnlock = !IsInGameThread())
 	{
-		if (!IsInGameThread())
+		if (bShouldUnlock)
 		{
 			AsyncCounter.Decrement();
 		}
