[Object/Rename]
- Replace FApp::IsGame() which can dynamically change when starting or exiting PIE to a more viable and deterministic per-object approach

diff --git a/Engine/Source/Runtime/CoreUObject/Private/UObject/Obj.cpp b/Engine/Source/Runtime/CoreUObject/Private/UObject/Obj.cpp
--- a/Engine/Source/Runtime/CoreUObject/Private/UObject/Obj.cpp
+++ b/Engine/Source/Runtime/CoreUObject/Private/UObject/Obj.cpp
@@ -388,11 +388,16 @@ bool UObject::Rename(const TCHAR* InName, UObject* NewOuter, ERenameFlags Flags)
 		{
 			const bool bUniquePathChanged = ((NewOuter != NULL && OldOuter != NewOuter) || (OldName != NewName));
 			const bool bRootPackage = GetClass() == UPackage::StaticClass() && OldOuter == NULL;
-			const bool bRedirectionAllowed = !FApp::IsGame() && ((Flags & REN_DontCreateRedirectors) == 0);
+			const bool bRedirectionAllowed = ((Flags & REN_DontCreateRedirectors) == 0) && GIsEditor && !IsRunningGame();
+			auto IsGameObject =
+				[this]() -> bool
+				{
+					return GetPackage()->HasAnyPackageFlags(PKG_PlayInEditor);
+				};
 
 			// We need to create a redirector if we changed the Outer or Name of an object that can be referenced from other packages
 			// [i.e. has the RF_Public flag] so that references to this object are not broken.
-			bCreateRedirector = bRootPackage == false && bUniquePathChanged == true && bRedirectionAllowed == true && bIsCaseOnlyChange == false;
+			bCreateRedirector = bRootPackage == false && bUniquePathChanged == true && bRedirectionAllowed == true && bIsCaseOnlyChange == false && IsGameObject() == false;
 		}
 
 #if UE_WITH_OBJECT_HANDLE_LATE_RESOLVE
