Add a special closed-only codepath for `FArrayProperty` when inner is `FObjectProperty` for a 4.3x speed boost.

#rnx

diff --git a/Engine/Source/Runtime/CoreUObject/Private/UObject/PropertyArray.cpp b/Engine/Source/Runtime/CoreUObject/Private/UObject/PropertyArray.cpp
--- a/Engine/Source/Runtime/CoreUObject/Private/UObject/PropertyArray.cpp
+++ b/Engine/Source/Runtime/CoreUObject/Private/UObject/PropertyArray.cpp
@@ -1162,6 +1162,7 @@ void FArrayProperty::CopyValuesInternal(void* Dest, void const* Src, int32 Count
 	{
 		DestArrayHelper.EmptyAndAddUninitializedValues(Num);
 	}
+
 	if (Num)
 	{
 		size_t Size = Inner->GetElementSize();
@@ -1169,9 +1170,37 @@ void FArrayProperty::CopyValuesInternal(void* Dest, void const* Src, int32 Count
 		uint8* DestData = (uint8*)DestArrayHelper.GetRawPtr();
 		if (!(Inner->PropertyFlags & CPF_IsPlainOldData))
 		{
-			for (int32 i = 0; i < Num; i++)
+			// If we are in the closed body of a transaction, we have a special case for `FObjectProperty` that vastly
+			// improves performance when copying arrays around (in the order of 4x faster). We **do not want** to do
+			// this optimization with anything that derives from `FObjectProperty` though, so we **do not use** `IsA`
+			// as that would all catch classes that inherited from `FObjectProperty`!
+			if (AutoRTFM::IsClosed() && (Inner->GetClass() == FObjectProperty::StaticClass()))
 			{
-				Inner->CopyCompleteValue(DestData + i * Size, SrcData + i * Size);
+				// Since we've already written to each location that we are going to be copying over below, we know that
+				// all of the record-writes that the transactional system would have inserted are no-ops. Also, each copy
+				// of a `TObjectPtr` will go into the open to do mark the object as reachable, so by doing a single open
+				// here we can avoid doing N calls to record-write and N calls to open for each object in the array. Big
+				// saving on large arrays!
+				UE_AUTORTFM_OPEN_NO_VALIDATION
+				{
+					// This optimization assumes that `FObjectProperty` is just backed by a `TObjectPtr` which when copying
+					// does not do any heap allocations, just pulls over the pointer to the underlying `UObject` and then
+					// marks it as reachable for incremental reachability analysis. Lets check that this is still the case
+					// here so that if anyone ever changes that in future this line will shout at them!
+					static_assert(std::is_same_v<decltype(FObjectProperty::GetDefaultPropertyValue()), TObjectPtr<UObject>>);
+
+					for (int32 i = 0; i < Num; i++)
+					{
+						Inner->CopyCompleteValue(DestData + i * Size, SrcData + i * Size);
+					}
+				};
+			}
+			else
+			{
+				for (int32 i = 0; i < Num; i++)
+				{
+					Inner->CopyCompleteValue(DestData + i * Size, SrcData + i * Size);
+				}
 			}
 		}
 		else
