[Docs] Update UObjectBaseUtility::IsA documentation to specify that the supplied class cannot resolve to nullptr when assigned to a UClass* variable

diff --git a/Engine/Source/Runtime/CoreUObject/Public/UObject/UObjectBaseUtility.h b/Engine/Source/Runtime/CoreUObject/Public/UObject/UObjectBaseUtility.h
--- a/Engine/Source/Runtime/CoreUObject/Public/UObject/UObjectBaseUtility.h
+++ b/Engine/Source/Runtime/CoreUObject/Public/UObject/UObjectBaseUtility.h
@@ -610,7 +610,11 @@ private:
 	}
 
 public:
-	/** Returns true if this object is of the specified type. */
+	/**
+	 * Returns true if this object is of the specified type.
+	 *
+	 * @param	SomeBase	The base class to check against, UClass* assignable. Cannot be null.
+	 */
 	template <typename OtherClassType>
 	FORCEINLINE bool IsA( OtherClassType SomeBase ) const
 	{
