fix the object migration abort message to include the object id and expected server id

diff --git a/Engine/Source/Runtime/CoreUObject/Private/UObject/RemoteObjectTransfer.cpp b/Engine/Source/Runtime/CoreUObject/Private/UObject/RemoteObjectTransfer.cpp
--- a/Engine/Source/Runtime/CoreUObject/Private/UObject/RemoteObjectTransfer.cpp
+++ b/Engine/Source/Runtime/CoreUObject/Private/UObject/RemoteObjectTransfer.cpp
@@ -72,14 +72,6 @@ public:
 	FRemoteObjectId ObjectId;
 	FRemoteServerId DestinationServerId;
 	FRemoteWorkPriority RequestPriority;
-
-	FString ToString() const
-	{
-		return FString::Format(TEXT("[obj %s dest %s %s]"),
-			{*ObjectId.ToString(),
-			*DestinationServerId.ToString(),
-			*RequestPriority.ToString()});
-	}
 };
 
 class FObjectMigrationRequests
@@ -990,7 +982,10 @@ void MigrateObjectFromRemoteServer(FRemoteObjectId ObjectId, FRemoteServerId Cur
 	// TODO: DO we want to log when we're pulling over an object we've NEVER known about?!
 	// 
 	// abort so the outer handler can renegotiate object transfers and retry
-	UE::RemoteExecutor::AbortTransactionRequiresDependencies(TEXT("object migration"));
+	UE::RemoteExecutor::AbortTransactionRequiresDependencies(
+		FString::Format(TEXT("object {0} required from server {1}"),
+			{*ObjectId.ToString(),
+			*CurrentOwnerServerId.ToString()}));
 }
 
 void TouchResidentObject(UObject* Object)
