SavePackage: Remove the addition of spurious ManagementPropagationDependencies in editor saves; the autoadded builddependencies were only supposed to be added for cook saves; when added to editorsaves they have a different undesired meaning.
#rnx

diff --git a/Engine/Source/Runtime/CoreUObject/Private/UObject/SavePackage/SaveContext.cpp b/Engine/Source/Runtime/CoreUObject/Private/UObject/SavePackage/SaveContext.cpp
--- a/Engine/Source/Runtime/CoreUObject/Private/UObject/SavePackage/SaveContext.cpp
+++ b/Engine/Source/Runtime/CoreUObject/Private/UObject/SavePackage/SaveContext.cpp
@@ -154,7 +154,15 @@ UObject* FSaveContext::ResolveForSave(TObjectPtr<UObject> InObject, UE::SavePack
 {
 	using namespace UE::Cook;
 
-	if (!Status.IsInSavePackage(InObject, GetPackage()))
+
+	if (
+		// SavePackage Build dependencies are only added when cooking. Build dependencies in editor saves mean
+		// management propagation dependencies and are only added from system-specific OnCookEvent overrides.
+		IsCooking()
+		// No need to consider build dependencies when resolving an export; the package and class of the export
+		// are already present as dependencies.
+		&& !Status.IsInSavePackage(InObject, GetPackage())
+		)
 	{
 		UObject* Object = nullptr;
 		{
@@ -176,7 +184,6 @@ UObject* FSaveContext::ResolveForSave(TObjectPtr<UObject> InObject, UE::SavePack
 	}
 	else
 	{
-		// No need to consider build dependencies when resolving an export
 		return InObject.Get();
 	}
 }
