DestroyUnversionedSchema when a UStruct is relinked in editor in order to prevent caching of stale FProperties in the UPS state

#rnx

diff --git a/Engine/Source/Runtime/CoreUObject/Private/UObject/Class.cpp b/Engine/Source/Runtime/CoreUObject/Private/UObject/Class.cpp
--- a/Engine/Source/Runtime/CoreUObject/Private/UObject/Class.cpp
+++ b/Engine/Source/Runtime/CoreUObject/Private/UObject/Class.cpp
@@ -862,6 +862,13 @@ void UStruct::PreloadChildren(FArchive& Ar)
 
 void UStruct::Link(FArchive& Ar, bool bRelinkExistingProperties)
 {
+#if WITH_EDITORONLY_DATA
+	// In the editor structures can be relinked to accomdate changes elsewhere in the struct hierarchy (mostly
+	// due to blueprint compilation, but also from package reload, live coding, and other blueprint like systems
+	// such as property bags, control rig, or third party scripting solutions)
+	DestroyUnversionedSchema(this);
+#endif
+
 	if (bRelinkExistingProperties)
 	{
 		// Preload everything before we calculate size, as the preload may end up recursively linking things
