PackageStore: FLinkerLoadPackageStoreBackend: Fix failure to load packages when the first attempt to load the package fails, either because it has PackageId but not PackageName, or because it occurs before the plugin containing the asset is mounted. Fixed for now by disabling the caching of missing packages.
#rnx

diff --git a/Engine/Source/Runtime/CoreUObject/Private/Serialization/PackageStore.cpp b/Engine/Source/Runtime/CoreUObject/Private/Serialization/PackageStore.cpp
--- a/Engine/Source/Runtime/CoreUObject/Private/Serialization/PackageStore.cpp
+++ b/Engine/Source/Runtime/CoreUObject/Private/Serialization/PackageStore.cpp
@@ -68,6 +68,10 @@ public:
 			return EPackageStoreEntryStatus::Missing;
 		}
 
+		if (PackageName.IsNone())
+		{
+			return EPackageStoreEntryStatus::Missing;
+		}
 
 #if WITH_EDITORONLY_DATA
 		// In editor, set MatchCaseOnDisk=true so that we set the capitalization of the Package's FName to match the
@@ -100,7 +104,9 @@ public:
 			return EPackageStoreEntryStatus::Ok;
 		}
 
-		PackageStoreCache.Add(PackageId, FPackageStoreEntry());
+		// Disable caching missing entries for now. They can be loaded too early, before the plugin containing the asset
+		// is mounted, and might become findable by FPackageName::DoesPackageExistEx later.
+		//PackageStoreCache.Add(PackageId, FPackageStoreEntry());
 		return EPackageStoreEntryStatus::Missing;
 	}
 
