Converted the OverrideObject to OverrideAllProperties
Implemented new method to retrieve object OverriddenState for each properties

diff --git a/Engine/Plugins/Experimental/EditorDataStorageFeatures/Source/TedsOutliner/Private/Widgets/OutlinerIconWidget.cpp b/Engine/Plugins/Experimental/EditorDataStorageFeatures/Source/TedsOutliner/Private/Widgets/OutlinerIconWidget.cpp
--- a/Engine/Plugins/Experimental/EditorDataStorageFeatures/Source/TedsOutliner/Private/Widgets/OutlinerIconWidget.cpp
+++ b/Engine/Plugins/Experimental/EditorDataStorageFeatures/Source/TedsOutliner/Private/Widgets/OutlinerIconWidget.cpp
@@ -21,7 +21,7 @@ namespace UE::OutlinerIconWidget::Local
 			return FAppStyle::GetBrush("SceneOutliner.OverrideAddedBase");
 
 		case EOverriddenState::AllOverridden:
-			// Not implemented yet
+			return FAppStyle::GetBrush("SceneOutliner.OverrideBase");
 			break;
 
 		case EOverriddenState::HasOverrides:
@@ -47,8 +47,7 @@ namespace UE::OutlinerIconWidget::Local
 			return FAppStyle::GetBrush("SceneOutliner.OverrideAdded");
 
 		case EOverriddenState::AllOverridden:
-			// Not implemented yet
-			break;
+			return FAppStyle::GetBrush("SceneOutliner.OverrideHere");
 
 		case EOverriddenState::HasOverrides:
 			return FAppStyle::GetBrush("SceneOutliner.OverrideInside");
@@ -73,8 +72,7 @@ namespace UE::OutlinerIconWidget::Local
 			return LOCTEXT("OverrideAddedTooltip", "This entity has been added.");
 
 		case EOverriddenState::AllOverridden:
-			// Not implemented yet
-			break;
+			return LOCTEXT("OverrideAddedTooltip", "This entity has all its properties overridden.");
 
 		case EOverriddenState::HasOverrides:
 			return LOCTEXT("OverrideInsideTooltip", "At least one property or child has an override.");
diff --git a/Engine/Source/Editor/PropertyEditor/Private/DetailsViewStyle.cpp b/Engine/Source/Editor/PropertyEditor/Private/DetailsViewStyle.cpp
--- a/Engine/Source/Editor/PropertyEditor/Private/DetailsViewStyle.cpp
+++ b/Engine/Source/Editor/PropertyEditor/Private/DetailsViewStyle.cpp
@@ -8,6 +8,7 @@
 #include "UObject/OverridableManager.h"
 #include "Layout/Visibility.h"
 #include "Styling/AppStyle.h"
+#include "Subsystems/PropertyVisibilityOverrideSubsystem.h"
 
 #define LOCTEXT_NAMESPACE "DetailsViewStyle"
 
@@ -46,7 +47,6 @@ void FOverridesWidgetStyleKeys::Initialize()
 const FOverridesWidgetStyleKey& FOverridesWidgetStyleKeys::Here()
 {
 	static const FOverridesWidgetStyleKey Here{"OverrideHere",
-		EOverriddenPropertyOperation::Replace,
 		EOverriddenState::AllOverridden,
 		false,
 		LOCTEXT("HereCategoryTooltip", "This component has been overridden."),
@@ -58,7 +58,6 @@ const FOverridesWidgetStyleKey& FOverridesWidgetStyleKeys::Here()
 const FOverridesWidgetStyleKey& FOverridesWidgetStyleKeys::Added()
 {
 	static const FOverridesWidgetStyleKey Added{"OverrideAdded",
-		EOverriddenPropertyOperation::Add,
 		EOverriddenState::Added,
 		TOptional<bool>(),
 		LOCTEXT("AddedCategoryTooltip", "This component has been added."),
@@ -70,7 +69,6 @@ const FOverridesWidgetStyleKey& FOverridesWidgetStyleKeys::Added()
 const FOverridesWidgetStyleKey& FOverridesWidgetStyleKeys::None()
 {
 	static const FOverridesWidgetStyleKey None{"OverrideNone",
-		EOverriddenPropertyOperation::None,
 		EOverriddenState::NoOverrides,
 		TOptional<bool>(),
 		LOCTEXT("NoneCategoryTooltip", "No override."),
@@ -88,7 +86,6 @@ const FOverridesWidgetStyleKey& FOverridesWidgetStyleKeys::Removed()
 const FOverridesWidgetStyleKey& FOverridesWidgetStyleKeys::Inside()
 {
 	static const FOverridesWidgetStyleKey Inside{"OverrideInside",
-		EOverriddenPropertyOperation::Modified,
 		EOverriddenState::HasOverrides,
 		TOptional<bool>(),
 		LOCTEXT("InsideCategoryTooltip", "At least one of this component's properties has been overridden."),
@@ -106,8 +103,7 @@ const FOverridesWidgetStyleKey& FOverridesWidgetStyleKeys::HereInside()
 const FOverridesWidgetStyleKey& FOverridesWidgetStyleKeys::Inherited()
 {
 	static const FOverridesWidgetStyleKey Inherited{"OverrideInherited",
-		EOverriddenPropertyOperation::Replace,
-		TOptional<EOverriddenState>(),
+		EOverriddenState::AllOverridden,
 		true,
 		LOCTEXT("InheritedCategoryTooltip", "This component's parent entity has been overridden."),
 		LOCTEXT("InheritedPropertyTooltip", "This property's parent component has been overridden.")
@@ -132,19 +128,17 @@ TAttribute<EVisibility> FOverridesWidgetStyleKey::GetVisibilityAttribute(const T
 	    {
 		    UObject* OverriddenObject = OverriddenObjectWeakPtr.Get();
 
-			bool bIsVisible = false;
+			FPropertyVisitorPath Path;
 			if (bIsProperty)
 			{
-				bool bInherited = false;
-				bIsVisible =  VisibleOverriddenPropertyOperation == Manager.GetOverriddenPropertyOperation(OverriddenObject, FPropertyChangedEvent(nullptr), *PropertyChain.Get(), &bInherited) &&
-							  (!bStateInherited.IsSet() || bStateInherited.GetValue() == bInherited);
-			}
-			else if (VisibleOverriddenState.IsSet())
-			{
-				bIsVisible = VisibleOverriddenState.GetValue() == Manager.GetOverriddenState(OverriddenObject);
+				Path = FPropertyVisitorPath(FPropertyChangedEvent(nullptr), *PropertyChain.Get());
 			}
 
-			if ( bIsVisible )
+			bool bInherited = false;
+			const bool bIsVisible = VisibleOverriddenState.GetValue() == Manager.GetOverriddenState(OverriddenObject, &UE::PropertyVisibility::ConsiderPropertyForOverriddenState, Path, &bInherited) &&
+									(!bStateInherited.IsSet() || bStateInherited.GetValue() == bInherited);
+
+			if (bIsVisible)
 			{
 				return EVisibility::Visible;
 			}
@@ -159,13 +153,11 @@ FOverridesWidgetStyleKey::FOverridesWidgetStyleKey(FName InName) : Name{InName}
 }
 
 FOverridesWidgetStyleKey::FOverridesWidgetStyleKey(FName InName,
-													EOverriddenPropertyOperation InOverriddenPropertyOperation,
 													const TOptional<EOverriddenState>& InOverriddenState,
 													const TOptional<bool>& bInStateInherited /*= TOptional<bool>()*/,
 													const FText& InCategoryTooltip,
 													const FText& InPropertyTooltip) :
    Name(InName),
-   VisibleOverriddenPropertyOperation(InOverriddenPropertyOperation),
    VisibleOverriddenState(InOverriddenState),
    bStateInherited(bInStateInherited),
    bCanBeVisible(true),
diff --git a/Engine/Source/Editor/PropertyEditor/Private/DetailsViewStyle.h b/Engine/Source/Editor/PropertyEditor/Private/DetailsViewStyle.h
--- a/Engine/Source/Editor/PropertyEditor/Private/DetailsViewStyle.h
+++ b/Engine/Source/Editor/PropertyEditor/Private/DetailsViewStyle.h
@@ -28,18 +28,15 @@ public:
 	PROPERTYEDITOR_API FOverridesWidgetStyleKey(FName InName);
 
 	/**
- * The constructor for the style key that initializes the Overridden state members
- *
- * @param InName the FName which is the name of the key
- * @param InOverriddenPropertyOperation the overridden property operation for which this style is visible
- * @param InOverriddenState the overridden state for components for which this style is visible
- * @param bInStateInherited take into account if the state was inherited or not, not set means anything
- * @param InCategoryTooltip tooltip for this widget if it is a category
- * @param InPropertyTooltip tooltip for this widget if it is a property
-
- */
+	 * The constructor for the style key that initializes the Overridden state members
+	 *
+	 * @param InName the FName which is the name of the key
+	 * @param InOverriddenState the overridden state for components for which this style is visible
+	 * @param bInStateInherited take into account if the state was inherited or not, not set means anything
+	 * @param InCategoryTooltip tooltip for this widget if it is a category
+	 * @param InPropertyTooltip tooltip for this widget if it is a property
+	 */
 	PROPERTYEDITOR_API FOverridesWidgetStyleKey(FName InName, 
-												EOverriddenPropertyOperation InOverriddenPropertyOperation,
 												const TOptional<EOverriddenState>& InOverriddenState,
 												const TOptional<bool>& bInStateInherited = TOptional<bool>(),
 												const FText& InCategoryTooltip = FText::GetEmpty(),
@@ -64,7 +61,7 @@ public:
 	 * Creates a TAttribute<EVisibility> which indicates the EVisibility for this widget style
 	 * 
 	 * @param PropertyChain points to the FEditPropertyChain for the Property, if this style is for a Property, else it is nullptr
-	 * @param OverriddenObject the UObject associated with this widget style
+	 * @param OverriddenObjectWeakPtr the UObject associated with this widget style
 	 * @return 
 	 */
 	TAttribute<EVisibility> GetVisibilityAttribute(const TSharedPtr<FEditPropertyChain>& PropertyChain, TWeakObjectPtr<UObject>& OverriddenObjectWeakPtr) const;
@@ -74,11 +71,6 @@ public:
 	 */
 	const FName Name;
 
-	/**
-	 * The EOverriddenPropertyOperation for which this style is visible for a property
-	 */
-	const TOptional<EOverriddenPropertyOperation> VisibleOverriddenPropertyOperation;
-
 	/**
 	 * The EOverriddenState for which this style is visible for a UObject
 	 */
diff --git a/Engine/Source/Editor/UnrealEd/Private/Subsystems/PropertyVisibilityOverrideSubsystem.cpp b/Engine/Source/Editor/UnrealEd/Private/Subsystems/PropertyVisibilityOverrideSubsystem.cpp
--- a/Engine/Source/Editor/UnrealEd/Private/Subsystems/PropertyVisibilityOverrideSubsystem.cpp
+++ b/Engine/Source/Editor/UnrealEd/Private/Subsystems/PropertyVisibilityOverrideSubsystem.cpp
@@ -3,8 +3,33 @@
 #include "Subsystems/PropertyVisibilityOverrideSubsystem.h"
 
 #include "Editor.h"
-#include "GameFramework/Actor.h"
-#include "Math/NumericLimits.h"
+#include "UObject/PropertyNames.h"
+
+bool UE::PropertyVisibility::ConsiderPropertyForOverriddenState(TNotNull<const FProperty*> Property)
+{
+	// @Todo: find a better way to do this.
+	// Hack for object array that are not CP_Edit but are editable via the card layout using a custom details panel like entities and components
+	constexpr EPropertyFlags MustHaveFlagsToBeConsidered = CPF_Edit | CPF_ExperimentalOverridableLogic;
+
+	// Check if the property or the owner property(containers) has any of the must-have flags
+	const FProperty* OwnerProperty = Property->GetOwnerProperty();
+	if (!Property->HasAnyPropertyFlags(MustHaveFlagsToBeConsidered) &&
+		(!OwnerProperty || !OwnerProperty->HasAnyPropertyFlags(MustHaveFlagsToBeConsidered)))
+	{
+		return false;
+	}
+
+	// Is the property filtered by user permission?
+	if (Property->GetBoolMetaData(PropertyNames::PropertyVisibilityOverrideName))
+	{
+		if (UPropertyVisibilityOverrideSubsystem::Get() && UPropertyVisibilityOverrideSubsystem::Get()->ShouldHideProperty(Property))
+		{
+			return false;
+		}
+	}
+
+	return true;
+}
 
 UPropertyVisibilityOverrideSubsystem* UPropertyVisibilityOverrideSubsystem::Get()
 {
diff --git a/Engine/Source/Editor/UnrealEd/Public/Subsystems/PropertyVisibilityOverrideSubsystem.h b/Engine/Source/Editor/UnrealEd/Public/Subsystems/PropertyVisibilityOverrideSubsystem.h
--- a/Engine/Source/Editor/UnrealEd/Public/Subsystems/PropertyVisibilityOverrideSubsystem.h
+++ b/Engine/Source/Editor/UnrealEd/Public/Subsystems/PropertyVisibilityOverrideSubsystem.h
@@ -9,6 +9,11 @@
 class FProperty;
 class UObject;
 
+namespace UE::PropertyVisibility
+{
+	UNREALED_API bool ConsiderPropertyForOverriddenState(TNotNull<const FProperty*> Property);
+}
+
 /** Bindable event for external objects to hook into ControlRig-level execution */
 DECLARE_DELEGATE_RetVal_OneParam(bool, FShouldHidePropertyDelegate, const FProperty*);
 
diff --git a/Engine/Source/Runtime/CoreUObject/Private/UObject/OverridableManager.cpp b/Engine/Source/Runtime/CoreUObject/Private/UObject/OverridableManager.cpp
--- a/Engine/Source/Runtime/CoreUObject/Private/UObject/OverridableManager.cpp
+++ b/Engine/Source/Runtime/CoreUObject/Private/UObject/OverridableManager.cpp
@@ -3,6 +3,7 @@
 #include "UObject/OverridableManager.h"
 #include "UObject/InstanceDataObjectUtils.h"
 #include "UObject/UObjectGlobals.h"
+#include "UObject/PropertyNames.h"
 
 /*
  *************************************************************************************
@@ -154,77 +155,129 @@ void FOverridableManager::RestoreOverrideState(TNotNull<const UObject*> OldObjec
 	}
 }
 
-EOverriddenState FOverridableManager::GetOverriddenState(TNotNull<UObject*> Object)
+EOverriddenState FOverridableManager::GetOverriddenState(TNotNull<UObject*> Object, 
+	TFunction<bool(TNotNull<const FProperty*>)> ConsiderPropertyForOverriddenState /*= TFunction<bool(TNotNull<const FProperty*>)>*/, 
+	const FPropertyVisitorPath& Path /*= FPropertyVisitorPath()*/, bool* bOutInheritedState /*= nullptr*/)
 {
 #if WITH_EDITORONLY_DATA
 	if (const FOverriddenPropertySet* OverriddenProperties = GetOverriddenProperties(Object))
 	{
-		// Transient should not affect the overridden state
-		if (Object->HasAnyFlags(RF_Transient))
+		if (Path.Num() == 0)
 		{
-			return EOverriddenState::NoOverrides;
-		}
+			// Transient should not affect the overridden state
+			if (Object->HasAnyFlags(RF_Transient))
+			{
+				return EOverriddenState::NoOverrides;
+			}
 
-		// Looking if the Archetype of the object was a CDO is not working in the case of a delete and readd.
-		// So now we explicitly use the bWasAdded to remember if an object was added.
-		if (OverriddenProperties->WasAdded())
-		{
-			return EOverriddenState::Added;
+			// Looking if the Archetype of the object was a CDO is not working in the case of a delete and readd.
+			// So now we explicitly use the bWasAdded to remember if an object was added.
+			if (OverriddenProperties->WasAdded())
+			{
+				return EOverriddenState::Added;
+			}
 		}
 
-		const EOverriddenPropertyOperation Operation = OverriddenProperties->GetOverriddenPropertyOperation((FArchiveSerializedPropertyChain*)nullptr, (FProperty*)nullptr);
-		switch (Operation)
+		TSet<UObject*> VisitedObject;
+		VisitedObject.Add(Object);
+		bool bAllOverridden = true;
+		bool bHasOverrides = false;
+		Object->GetClass()->Visit(Object, [this, Object, OverriddenProperties, Path, &VisitedObject, &bAllOverridden, &bHasOverrides, ConsiderPropertyForOverriddenState, &bOutInheritedState](const FPropertyVisitorContext& Context)
 		{
-			case EOverriddenPropertyOperation::Replace:
-				return EOverriddenState::AllOverridden;
-			case EOverriddenPropertyOperation::Modified:
-				return EOverriddenState::HasOverrides;
-			case EOverriddenPropertyOperation::None:
-				break;
-			default:
-				checkf(false, TEXT("Unsupported operation(%s) on object(%p:%s)"), *GetOverriddenOperationString(Operation), static_cast<UObject*>(Object), *Object->GetName());
-		}
+			const EOverriddenPropertyOperation Operation = OverriddenProperties->GetOverriddenPropertyOperation(Context.Path.GetRootIterator());
+			bool bIsEqual = false;
+			if (Context.Path.Num() <= Path.Num())
+			{
+				if (!Context.Path.Contained(Path, &bIsEqual))
+				{
+					return EPropertyVisitorControlFlow::StepOver;
+				}
 
-		auto GetSubObjectState = [this](TNotNull<UObject*> Object, auto& Recurse) -> EOverriddenState
-		{
-			// Need to check subobjects
-			for (TPropertyValueIterator<FObjectProperty> ObjPropIt(Object->GetClass(), Object); ObjPropIt; ++ObjPropIt)
+				if (Operation == EOverriddenPropertyOperation::Modified)
+				{
+					return EPropertyVisitorControlFlow::StepInto;
+				}
+
+				if (Operation == EOverriddenPropertyOperation::Replace)
+				{
+					if (bOutInheritedState)
+					{
+						*bOutInheritedState = !bIsEqual;
+					}
+				}
+				else
+				{
+					bAllOverridden = false;
+				}
+
+				return EPropertyVisitorControlFlow::Stop;
+			}
+
+			const FProperty* Property = Context.Path.Top().Property;
+
+			// Skip always overridden properties
+			if (Property->HasAnyPropertyFlags(CPF_ExperimentalAlwaysOverriden))
 			{
-				if (UObject* InstancedSubObject = ObjPropIt.Key()->GetObjectPropertyValue(ObjPropIt.Value()))
+				return EPropertyVisitorControlFlow::StepOver;
+			}
+
+			// Check if the property should be considered for the overridden state evaluation 
+			if (ConsiderPropertyForOverriddenState.IsSet() && !ConsiderPropertyForOverriddenState(Property))
+			{
+				return EPropertyVisitorControlFlow::StepOver;
+			}
+
+			// No need to dig deeper if this is a replace of an add
+			if (Operation == EOverriddenPropertyOperation::Replace || 
+				Operation == EOverriddenPropertyOperation::Add)
+			{
+				bHasOverrides = true;
+				return !bAllOverridden ? EPropertyVisitorControlFlow::Stop : EPropertyVisitorControlFlow::StepOver;
+			}
+
+			// Recurse if the property contains inner ones
+			if (Context.Path.Top().bContainsInnerProperties)
+			{
+				if (Operation != EOverriddenPropertyOperation::None)
+				{
+					bHasOverrides = true;
+				}
+
+				// Only recurse on instanced subobjects
+				if (const FObjectPropertyBase* ObjectProperty = CastField<FObjectPropertyBase>(Property))
 				{
-					if (InstancedSubObject->IsIn(Object))
+					if (UObject* InstancedSubObject = ObjectProperty->GetObjectPropertyValue(Context.Data.PropertyData))
 					{
-						if (const FOverriddenPropertySet* OverriddenProperties = GetOverriddenProperties(InstancedSubObject))
+						if (!VisitedObject.Contains(InstancedSubObject) && InstancedSubObject->IsIn(Object))
 						{
-							const EOverriddenPropertyOperation Operation = OverriddenProperties->GetOverriddenPropertyOperation((FArchiveSerializedPropertyChain*)nullptr, (FProperty*)nullptr);
-							if (Operation != EOverriddenPropertyOperation::None)
-							{
-								return EOverriddenState::HasOverrides;
-							}
-
-							if (Recurse(InstancedSubObject, Recurse) != EOverriddenState::NoOverrides)
-							{
-								return EOverriddenState::HasOverrides;
-							}
+							VisitedObject.Add(InstancedSubObject);
+							return EPropertyVisitorControlFlow::StepInto;
 						}
 					}
 				}
+				else
+				{
+					return EPropertyVisitorControlFlow::StepInto;
+				}
+			}
+
+			
+			if (Operation == EOverriddenPropertyOperation::None)
+			{
+				bAllOverridden = false;
+				return bHasOverrides ? EPropertyVisitorControlFlow::Stop : EPropertyVisitorControlFlow::StepOver;
 			}
 
-			return EOverriddenState::NoOverrides;
-		};
+			return EPropertyVisitorControlFlow::Stop;
+		});
 
-		EOverriddenState SubObjectState = GetSubObjectState(Object, GetSubObjectState);
-		if (SubObjectState != EOverriddenState::NoOverrides)
-		{
-			return SubObjectState;
-		}
+		return bAllOverridden ? EOverriddenState::AllOverridden : (bHasOverrides ? EOverriddenState::HasOverrides : EOverriddenState::NoOverrides);
 	}
 #endif // WITH_EDITORONLY_DATA
 	return EOverriddenState::NoOverrides;
 }
 
-void FOverridableManager::OverrideObject(TNotNull<UObject*> Object)
+void FOverridableManager::OverrideAllObjectProperties(TNotNull<UObject*> Object)
 {
 #if WITH_EDITORONLY_DATA
 	if (FOverriddenPropertySet* ThisObjectOverriddenProperties = Find(Object))
@@ -235,33 +288,6 @@ void FOverridableManager::OverrideObject(TNotNull<UObject*> Object)
 #endif // WITH_EDITORONLY_DATA
 }
 
-void FOverridableManager::OverrideInstancedSubObject(TNotNull<const UObject*> Object, TNotNull<UObject*> InstancedSubObject)
-{
-#if WITH_EDITORONLY_DATA
-	if (InstancedSubObject->IsIn(Object))
-	{
-		OverrideObject(InstancedSubObject);
-	}
-#endif // WITH_EDITORONLY_DATA
-}
-
-void FOverridableManager::PropagateOverrideToInstancedSubObjects(TNotNull<UObject*> Object)
-{
-#if WITH_EDITORONLY_DATA
-	// Do not use TPropertyValueIterator<FObjectProperty> to be consistent with PropagateClearOverridesToInstancedSubObjects
-	// which cannot use TPropertyValueIterator<FObjectProperty> because the object might not have 
-	// the object class setup correctly when this is called from PostInit
-	TArray<UObject*> InstancedSubObjects;
-	constexpr bool bIncludeNestedObjects = false;
-	GetObjectsWithOuter(Object, InstancedSubObjects, bIncludeNestedObjects);
-	for (UObject* InstancedSubObject : InstancedSubObjects)
-	{
-		checkf(InstancedSubObject, TEXT("Expecting non null SubObjects"));
-		OverrideInstancedSubObject(Object, InstancedSubObject);
-	}
-#endif // WITH_EDITORONLY_DATA
-}
-
 void FOverridableManager::OverrideProperty(TNotNull<UObject*> Object, const FPropertyVisitorPath& PropertyPath)
 {
 #if WITH_EDITORONLY_DATA
@@ -283,10 +309,10 @@ bool FOverridableManager::ClearOverriddenProperty(TNotNull<UObject*> Object, FPr
 	return false;
 }
 
-void FOverridableManager::PreOverrideProperty(TNotNull<UObject*> Object, const FPropertyVisitorPath& PropertyPath)
+void FOverridableManager::PreOverrideProperty(TNotNull<UObject*> Object, const FPropertyVisitorPath& PropertyPath, const EPropertyChangeType::Type ChangeType)
 {
 #if WITH_EDITORONLY_DATA
-	NotifyPropertyChange(EPropertyNotificationType::PreEdit, Object, PropertyPath.GetRootIterator(), EPropertyChangeType::Unspecified);
+	NotifyPropertyChange(EPropertyNotificationType::PreEdit, Object, PropertyPath.GetRootIterator(), ChangeType);
 #endif // WITH_EDITORONLY_DATA
 }
 
@@ -392,7 +418,7 @@ FOverridableManager::FOverridableManager()
 #if WITH_EDITOR
 	FCoreUObjectDelegates::OnObjectsReinstanced.AddRaw(this, &FOverridableManager::HandleObjectsReInstantiated);
 	FCoreUObjectDelegates::OnVerseDeadObjectReferences.AddRaw(this, &FOverridableManager::HandleDeadObjectReferences);
-	FCoreUObjectDelegates::OnPreObjectPropertyChanged.AddRaw(this, &FOverridableManager::PreOverrideProperty);
+	FCoreUObjectDelegates::OnPreObjectPropertyChanged.AddRaw(this, &FOverridableManager::PreOverridePropertyChain);
 	FCoreUObjectDelegates::OnPostObjectPropertyChanged.AddRaw(this, &FOverridableManager::PostOverrideProperty);
 #endif
 }
diff --git a/Engine/Source/Runtime/CoreUObject/Private/UObject/OverriddenPropertySet.cpp b/Engine/Source/Runtime/CoreUObject/Private/UObject/OverriddenPropertySet.cpp
--- a/Engine/Source/Runtime/CoreUObject/Private/UObject/OverriddenPropertySet.cpp
+++ b/Engine/Source/Runtime/CoreUObject/Private/UObject/OverriddenPropertySet.cpp
@@ -563,6 +563,13 @@ EOverriddenPropertyOperation FOverriddenPropertySet::GetOverriddenPropertyOperat
 		// Special handling for array of instanced subobjects 
 		else if (const FArrayProperty* ArrayProperty = CastField<FArrayProperty>(CurrentProperty))
 		{
+			// This code handles inner properties at the same time as the container property,
+			// skip the container one if the next one is equal to its inner property
+			if (NextPropertyIterator && NextPropertyIterator->Property == ArrayProperty->Inner)
+			{
+				++PropertyIterator;
+				++NextPropertyIterator;
+			}
 			ArrayIndex = PropertyIterator->Index;
 			checkf(ArrayIndex == INDEX_NONE || PropertyIterator->PropertyInfo == EPropertyVisitorInfoType::ContainerIndex, TEXT("Expecting a container index"));
 
@@ -600,8 +607,21 @@ EOverriddenPropertyOperation FOverriddenPropertySet::GetOverriddenPropertyOperat
 		// Special handling for maps and values of instance subobjects
 		else if (const FMapProperty* MapProperty = CastField<FMapProperty>(CurrentProperty))
 		{
+			// This code handles inner properties at the same time as the container property,
+			// skip the container one if the next one is equal to one of its inner properties
+			if (NextPropertyIterator && 
+			    (NextPropertyIterator->Property == MapProperty->KeyProp ||
+				 NextPropertyIterator->Property == MapProperty->ValueProp) )
+			{
+				++PropertyIterator;
+				++NextPropertyIterator;
+			}
+
 			ArrayIndex = PropertyIterator->Index;
-			checkf(ArrayIndex == INDEX_NONE || PropertyIterator->PropertyInfo == EPropertyVisitorInfoType::ContainerIndex, TEXT("Expecting a container index"));
+			checkf(ArrayIndex == INDEX_NONE || 
+				PropertyIterator->PropertyInfo == EPropertyVisitorInfoType::ContainerIndex || 
+				PropertyIterator->PropertyInfo == EPropertyVisitorInfoType::MapKey || 
+				PropertyIterator->PropertyInfo == EPropertyVisitorInfoType::MapValue, TEXT("Expecting a container type"));
 
 			checkf(MapProperty->ValueProp, TEXT("Expecting a value type for Maps"));
 			FScriptMapHelper MapHelper(MapProperty, SubValuePtr);
@@ -714,6 +734,13 @@ bool FOverriddenPropertySet::ClearOverriddenProperty(FOverriddenPropertyNode& Pa
 		// Special handling for array of instanced subobjects 
 		else if (const FArrayProperty* ArrayProperty = CastField<FArrayProperty>(CurrentProperty))
 		{
+			// This code handles inner properties at the same time as the container property,
+			// skip the container one if the next one is equal to its inner property
+			if (NextPropertyIterator && NextPropertyIterator->Property == ArrayProperty->Inner)
+			{
+				++PropertyIterator;
+				++NextPropertyIterator;
+			}
 			ArrayIndex = PropertyIterator->Index;
 			checkf(ArrayIndex == INDEX_NONE || PropertyIterator->PropertyInfo == EPropertyVisitorInfoType::ContainerIndex, TEXT("Expecting a container index"));
 
@@ -759,15 +786,28 @@ bool FOverriddenPropertySet::ClearOverriddenProperty(FOverriddenPropertyNode& Pa
 		// Special handling for maps and values of instance subobjects 
 		else if (const FMapProperty* MapProperty = CastField<FMapProperty>(CurrentProperty))
 		{
+			// This code handles inner properties at the same time as the container property,
+			// skip the container one if the next one is equal to one of its inner properties
+			if (NextPropertyIterator && 
+			    (NextPropertyIterator->Property == MapProperty->KeyProp ||
+				 NextPropertyIterator->Property == MapProperty->ValueProp) )
+			{
+				++PropertyIterator;
+				++NextPropertyIterator;
+			}
+
 			ArrayIndex = PropertyIterator->Index;
-			checkf(ArrayIndex == INDEX_NONE || PropertyIterator->PropertyInfo == EPropertyVisitorInfoType::ContainerIndex, TEXT("Expecting a container index"));
+			checkf(ArrayIndex == INDEX_NONE || 
+				PropertyIterator->PropertyInfo == EPropertyVisitorInfoType::ContainerIndex || 
+				PropertyIterator->PropertyInfo == EPropertyVisitorInfoType::MapKey || 
+				PropertyIterator->PropertyInfo == EPropertyVisitorInfoType::MapValue, TEXT("Expecting a container type"));
 
 			FScriptMapHelper MapHelper(MapProperty, SubValuePtr);
 
 			const int32 InternalMapIndex = ArrayIndex != INDEX_NONE ? MapHelper.FindInternalIndex(ArrayIndex) : INDEX_NONE;
 			const FObjectPropertyBase* ValueObjectProperty = CastField<FObjectPropertyBase>(MapProperty->ValueProp);
 
-			// If there is a next node, it is probably because the map value is holding a instanced subobject and the user is changing value on it.
+			// If there is a next node, it is probably because the map value is holding an instanced subobject and the user is changing value on it.
 			// So forward the call to the instanced subobject
 			if (NextPropertyIterator)
 			{
@@ -878,20 +918,9 @@ void FOverriddenPropertySet::NotifyPropertyChange(FOverriddenPropertyNode* Paren
 	FOverridableManager& OverridableManager = FOverridableManager::Get();
 	if (!PropertyIterator)
 	{
-		if (ParentPropertyNode && Notification == EPropertyNotificationType::PostEdit)
-		{
-			// Sub-property overrides are not needed from now on, so clear them
-			ParentPropertyNode->Reset();
-
-			// Replacing this entire property
-			ParentPropertyNode->Operation = EOverriddenPropertyOperation::Replace;
+		checkf(ParentPropertyNode == &RootNode, TEXT("Only expecting the root node in this code path"));
 
-			// If we are overriding the root node, need to propagate the overrides to all instanced sub object
-			if (&RootNode == ParentPropertyNode)
-			{
-				OverridableManager.PropagateOverrideToInstancedSubObjects(Owner);
-			}
-		}
+		NotifyAllPropertyChange(ParentPropertyNode, Notification, ChangeType, Owner->GetClass(), Owner);
 		return;
 	}
 
@@ -930,13 +959,20 @@ void FOverriddenPropertySet::NotifyPropertyChange(FOverriddenPropertyNode* Paren
 
 	// Because the PreEdit API doesn't provide Index info, we need to snapshot changed containers during pre-edit
 	// so we can intuit which element was removed during PostEdit. This is a map-like structure that stores the latest snapshot for each contianer property
+	struct FContainerInfo
+	{
+		const UObject* Owner = nullptr;
+		const FProperty* Property = nullptr;
+	};
 	static struct
 	{
-		uint8* Find(const FProperty* ContainerProperty) const
+		// There's not many elements so we're using an array of TPairs for cache friendliness. A TMap would work fine here as well.
+		uint8* Find(const UObject* ContainerOwner, const FProperty* ContainerProperty) const
 		{
-			for (const TPair<const FProperty*, uint8*>& Element : Data)
+			for (const TPair<FContainerInfo, uint8*>& Element : Data)
 			{
-				if (Element.Key == ContainerProperty)
+				if (Element.Key.Owner == ContainerOwner &&
+					Element.Key.Property == ContainerProperty)
 				{
 					return Element.Value;
 				}
@@ -944,11 +980,12 @@ void FOverriddenPropertySet::NotifyPropertyChange(FOverriddenPropertyNode* Paren
 			return nullptr;
 		}
 
-		void Free(const FProperty* ContainerProperty)
+		void Free(const UObject* ContainerOwner, const FProperty* ContainerProperty)
 		{
 			for (int32 I = 0; I < Data.Num(); ++I)
 			{
-				if (Data[I].Key == ContainerProperty)
+				if (Data[I].Key.Owner == ContainerOwner &&
+					Data[I].Key.Property == ContainerProperty)
 				{
 					ContainerProperty->DestroyValue(Data[I].Value);
 					FMemory::Free(Data[I].Value);
@@ -959,19 +996,19 @@ void FOverriddenPropertySet::NotifyPropertyChange(FOverriddenPropertyNode* Paren
 			checkf(false, TEXT("Expecting a matching property to the allocated memory"));
 		}
 
-		uint8* FindOrAdd(const FProperty* ContainerProperty)
+		uint8* FindOrAdd(const UObject* ContainerOwner, const FProperty* ContainerProperty)
 		{
-			if (uint8* Found = Find(ContainerProperty))
+			if (uint8* Found = Find(ContainerOwner, ContainerProperty))
 			{
 				return Found;
 			}
-			int32 I = Data.Add({ ContainerProperty, (uint8*)FMemory::Malloc(ContainerProperty->GetSize(), ContainerProperty->GetMinAlignment()) });
+			FContainerInfo Key = {ContainerOwner, ContainerProperty};
+			int32 I = Data.Add({ Key, (uint8*)FMemory::Malloc(ContainerProperty->GetSize(), ContainerProperty->GetMinAlignment()) });
 			ContainerProperty->InitializeValue(Data[I].Value);
 			return Data[I].Value;
 		}
 
-		// There's not many elements so we're using an array of TPairs for cache friendliness. A TMap would work fine here as well.
-		TArray<TPair<const FProperty*, uint8*>> Data;
+		TArray<TPair<FContainerInfo, uint8*>> Data;
 	} SavedPreEditContainers;
 
 	FPropertyVisitorPath::Iterator NextPropertyIterator = PropertyIterator+1;
@@ -981,7 +1018,15 @@ void FOverriddenPropertySet::NotifyPropertyChange(FOverriddenPropertyNode* Paren
 		if (FObjectPropertyBase* InnerObjectProperty = CastField<FObjectPropertyBase>(ArrayProperty->Inner))
 		{
 			FScriptArrayHelper ArrayHelper(ArrayProperty, SubValuePtr);
-			int32 ArrayIndex = PropertyIterator->Index;
+
+			// This code handles inner properties at the same time as the container property,
+			// skip the container one if the next one is equal to its inner property
+			if (NextPropertyIterator && NextPropertyIterator->Property == ArrayProperty->Inner)
+			{
+				++PropertyIterator;
+				++NextPropertyIterator;
+			}
+			int ArrayIndex = PropertyIterator->Index;
 			checkf(ArrayIndex == INDEX_NONE || PropertyIterator->PropertyInfo == EPropertyVisitorInfoType::ContainerIndex, TEXT("Expecting a container index"));
 
 			if (!NextPropertyIterator)
@@ -990,7 +1035,7 @@ void FOverriddenPropertySet::NotifyPropertyChange(FOverriddenPropertyNode* Paren
 
 				if (Notification == EPropertyNotificationType::PreEdit)
 				{
-					FScriptArrayHelper PreEditArrayHelper(ArrayProperty, SavedPreEditContainers.FindOrAdd(ArrayProperty));
+					FScriptArrayHelper PreEditArrayHelper(ArrayProperty, SavedPreEditContainers.FindOrAdd(Owner, ArrayProperty));
 					PreEditArrayHelper.EmptyAndAddValues(ArrayHelper.Num());
 					for (int32 i = 0; i < ArrayHelper.Num(); i++)
 					{
@@ -999,7 +1044,7 @@ void FOverriddenPropertySet::NotifyPropertyChange(FOverriddenPropertyNode* Paren
 					return;
 				}
 
-				FScriptArrayHelper PreEditArrayHelper(ArrayProperty, SavedPreEditContainers.Find(ArrayProperty));
+				FScriptArrayHelper PreEditArrayHelper(ArrayProperty, SavedPreEditContainers.Find(Owner, ArrayProperty));
 
 				auto ArrayReplace = [&]
 				{
@@ -1010,19 +1055,12 @@ void FOverriddenPropertySet::NotifyPropertyChange(FOverriddenPropertyNode* Paren
 					}
 
 					// This is a case of the entire array is overridden
-					// Need to loop through every sub object and setup them up as overridden
+					// Need to loop through every sub object and override all properties
 					for (int i = 0; i < ArrayHelper.Num(); ++i)
 					{
 						if (UObject* SubObject = TryGetInstancedSubObjectValue(InnerObjectProperty, ArrayHelper.GetElementPtr(i)))
 						{
-							if(SubPropertyNode)
-							{
-								const FOverriddenPropertyNodeID SubObjectID(SubObject);
-								FOverriddenPropertyNode& SubObjectNode = FindOrAddNode(*SubPropertyNode, SubObjectID);
-								SubObjectNode.Operation = EOverriddenPropertyOperation::Replace;
-							}
-
-							OverridableManager.OverrideInstancedSubObject(Owner, SubObject);
+							OverridableManager.OverrideAllObjectProperties(SubObject);
 						}
 					}
 				};
@@ -1198,15 +1236,29 @@ void FOverriddenPropertySet::NotifyPropertyChange(FOverriddenPropertyNode* Paren
 		FObjectPropertyBase* ValueObjectProperty = CastField<FObjectPropertyBase>(MapProperty->ValueProp);
 
 		FScriptMapHelper MapHelper(MapProperty, SubValuePtr);
+
+		// This code handles inner properties at the same time as the container property,
+		// skip the container one if the next one is equal to one of its inner properties
+		if (NextPropertyIterator && 
+		    (NextPropertyIterator->Property == MapProperty->KeyProp ||
+			 NextPropertyIterator->Property == MapProperty->ValueProp) )
+		{
+			++PropertyIterator;
+			++NextPropertyIterator;
+		}
+
 		int32 LogicalMapIndex = PropertyIterator->Index;
-		checkf(LogicalMapIndex == INDEX_NONE || PropertyIterator->PropertyInfo == EPropertyVisitorInfoType::ContainerIndex, TEXT("Expecting a container index type"));
+		checkf(LogicalMapIndex == INDEX_NONE || 
+			PropertyIterator->PropertyInfo == EPropertyVisitorInfoType::ContainerIndex || 
+			PropertyIterator->PropertyInfo == EPropertyVisitorInfoType::MapKey || 
+			PropertyIterator->PropertyInfo == EPropertyVisitorInfoType::MapValue, TEXT("Expecting a container type"));
 
 		int32 InternalMapIndex = LogicalMapIndex != INDEX_NONE ? MapHelper.FindInternalIndex(LogicalMapIndex) : INDEX_NONE;
 		if (!NextPropertyIterator)
 		{
 			if (Notification == EPropertyNotificationType::PreEdit)
 			{
-				FScriptMapHelper PreEditMapHelper(MapProperty, SavedPreEditContainers.FindOrAdd(MapProperty));
+				FScriptMapHelper PreEditMapHelper(MapProperty, SavedPreEditContainers.FindOrAdd(Owner, MapProperty));
 				PreEditMapHelper.EmptyValues();
 				for (FScriptMapHelper::FIterator It(MapHelper); It; ++It)
 				{
@@ -1215,7 +1267,7 @@ void FOverriddenPropertySet::NotifyPropertyChange(FOverriddenPropertyNode* Paren
 				return;
 			}
 
-			uint8* SavedPreEditMap = SavedPreEditContainers.Find(MapProperty);
+			uint8* SavedPreEditMap = SavedPreEditContainers.Find(Owner, MapProperty);
 			checkf(SavedPreEditMap, TEXT("Expecting the same property as the pre edit flow"));
 			FScriptMapHelper PreEditMapHelper(MapProperty, SavedPreEditMap);
 			// The logical should map directly to the pre edit map internal index as we skipped all of the invalid entries
@@ -1223,7 +1275,7 @@ void FOverriddenPropertySet::NotifyPropertyChange(FOverriddenPropertyNode* Paren
 
 			ON_SCOPE_EXIT
 			{
-				SavedPreEditContainers.Free(MapProperty);
+				SavedPreEditContainers.Free(Owner, MapProperty);
 			};
 
 			auto MapReplace = [&]()
@@ -1249,11 +1301,11 @@ void FOverriddenPropertySet::NotifyPropertyChange(FOverriddenPropertyNode* Paren
 					//if (UObject* KeySubObject = TryGetInstancedSubObjectValue(KeyObjectProperty, MapHelper.GetKeyPtr(It.GetInternalIndex())))
 					//{
 					//	checkf(false, TEXT("Keys as an instanced subobject is not supported yet"));
-					//	OverridableManager.OverrideInstancedSubObject(Owner, KeySubObject);
+					//	OverridableManager.OverrideAllObjectProperties(KeySubObject);
 					//}
 					if (UObject* ValueSubObject = TryGetInstancedSubObjectValue(ValueObjectProperty, MapHelper.GetValuePtr(It.GetInternalIndex())))
 					{
-						OverridableManager.OverrideInstancedSubObject(Owner, ValueSubObject);
+						OverridableManager.OverrideAllObjectProperties(ValueSubObject);
 					}
 				}
 			};
@@ -1392,15 +1444,20 @@ void FOverriddenPropertySet::NotifyPropertyChange(FOverriddenPropertyNode* Paren
 			}
 		}
 	}
-	else if (Property->IsA<FStructProperty>())
+	else if (const FStructProperty* StructProperty = CastField<FStructProperty>(Property))
 	{
-		if (!NextPropertyIterator)
+		// No individual property serialization from this point on
+		if (StructProperty->Struct->UseNativeSerialization())
 		{
 			if (Notification == EPropertyNotificationType::PostEdit && SubPropertyNode)
 			{
 				SubPropertyNode->Operation = EOverriddenPropertyOperation::Replace;
 			}
 		}
+		else if (!NextPropertyIterator)
+		{
+			NotifyAllPropertyChange(SubPropertyNode, Notification, ChangeType, StructProperty->Struct, SubValuePtr);
+		}
 		else
 		{
 			NotifyPropertyChange(SubPropertyNode, Notification, NextPropertyIterator, ChangeType, SubValuePtr, bNeedsCleanup);
@@ -1409,14 +1466,47 @@ void FOverriddenPropertySet::NotifyPropertyChange(FOverriddenPropertyNode* Paren
 	}
 	else if (const FObjectPropertyBase* ObjectProperty = CastField<FObjectPropertyBase>(Property))
 	{
+		UObject* SubObject = TryGetInstancedSubObjectValue(ObjectProperty, SubValuePtr);
+		if (Notification == EPropertyNotificationType::PreEdit)
+		{
+			uint8* PreEditValuePtr = SavedPreEditContainers.FindOrAdd(Owner, ObjectProperty);
+			ObjectProperty->SetObjectPtrPropertyValue(PreEditValuePtr, SubObject);
+		}
+
+		ON_SCOPE_EXIT
+		{
+			if (Notification == EPropertyNotificationType::PostEdit)
+			{
+				SavedPreEditContainers.Free(Owner, ObjectProperty);
+			}
+		};
+
 		if (!NextPropertyIterator)
 		{
-			if (Notification == EPropertyNotificationType::PostEdit && SubPropertyNode)
+			if (Notification == EPropertyNotificationType::PostEdit)
 			{
-				SubPropertyNode->Operation = EOverriddenPropertyOperation::Replace;
+				if (SubPropertyNode)
+				{
+					SubPropertyNode->Operation = EOverriddenPropertyOperation::Replace;
+				}
+
+				uint8* PreEditValuePtr = SavedPreEditContainers.Find(Owner, ObjectProperty);
+				checkf(PreEditValuePtr, TEXT("Expecting a previous stored subobjectproperty"));
+				UObject* OldSubObject = TryGetInstancedSubObjectValue(ObjectProperty, PreEditValuePtr);
+
+				// Handle the case where subobject was nulled out by clearing any previous edit as the reinstancer can still find this object as it might not be GC'd yet
+				if (OldSubObject && OldSubObject != SubObject)
+				{
+					OverridableManager.ClearOverrides(OldSubObject);
+				}
+				// If the object stayed the same, that means the users want to override all object properties
+				else if (SubObject && SubObject == OldSubObject)
+				{
+					OverridableManager.OverrideAllObjectProperties(SubObject);
+				}
 			}
 		}
-		else if (UObject* SubObject = TryGetInstancedSubObjectValue(ObjectProperty, SubValuePtr))
+		else if (SubObject)
 		{
 			// This should not be needed in the property grid, as it should already been called on the subobject.
 			OverridableManager.NotifyPropertyChange(Notification, SubObject, NextPropertyIterator, ChangeType);
@@ -1450,6 +1540,22 @@ void FOverriddenPropertySet::NotifyPropertyChange(FOverriddenPropertyNode* Paren
 	}
 }
 
+void FOverriddenPropertySet::NotifyAllPropertyChange(FOverriddenPropertyNode* ParentPropertyNode, const EPropertyNotificationType Notification, const EPropertyChangeType::Type ChangeType, TNotNull<const UStruct*> Struct, const void* Data)
+{
+	for (TFieldIterator<FProperty> It(Struct, EFieldIteratorFlags::IncludeSuper); It; ++It)
+	{
+		// Skip always overridden properties
+		if (It->HasAnyPropertyFlags(CPF_ExperimentalAlwaysOverriden))
+		{
+			continue;
+		}
+
+		bool bNeedsCleanup = false;
+		FPropertyVisitorPath Path(FPropertyVisitorInfo(It.operator*()));
+		NotifyPropertyChange(ParentPropertyNode, Notification, Path.GetRootIterator(), ChangeType, Data, bNeedsCleanup);
+	}
+}
+
 UObject* FOverriddenPropertySet::TryGetInstancedSubObjectValue(const FObjectPropertyBase* FromProperty, const void* ValuePtr) const
 {
 	// Property can be NULL - in that case there is no value.
diff --git a/Engine/Source/Runtime/CoreUObject/Private/UObject/OverrideEventHelper.cpp b/Engine/Source/Runtime/CoreUObject/Private/UObject/OverrideEventHelper.cpp
--- a/Engine/Source/Runtime/CoreUObject/Private/UObject/OverrideEventHelper.cpp
+++ b/Engine/Source/Runtime/CoreUObject/Private/UObject/OverrideEventHelper.cpp
@@ -57,7 +57,7 @@ namespace UE
 		TArray<TMap<FString, int32>> ArrayIndicesPerObject;
 	};
 
-	void SendOverrideObjectEvent(TNotNull<UObject*> Object)
+	void SendOverrideAllObjectPropertiesEvent(TNotNull<UObject*> Object)
 	{
 #if UE_EDITOR
 		// an event with no provided property path will override the entire object
@@ -130,7 +130,7 @@ namespace UE
 	void SendPreOverridePropertyEvent(TNotNull<UObject*> Object, const FEditPropertyChain& PropertyChain)
 	{
 #if UE_EDITOR
-		FOverridableManager::Get().PreOverrideProperty(Object, PropertyChain);
+		FOverridableManager::Get().PreOverridePropertyChain(Object, PropertyChain);
 #endif
 	}
 
diff --git a/Engine/Source/Runtime/CoreUObject/Public/UObject/OverridableManager.h b/Engine/Source/Runtime/CoreUObject/Public/UObject/OverridableManager.h
--- a/Engine/Source/Runtime/CoreUObject/Public/UObject/OverridableManager.h
+++ b/Engine/Source/Runtime/CoreUObject/Public/UObject/OverridableManager.h
@@ -154,26 +154,20 @@ public:
 	COREUOBJECT_API void RestoreOverrideState(TNotNull<const UObject*> OldObject, TNotNull<UObject*> NewObject);
 
 	/**
-	 * Retrieve the overridden state for the specified object
+	 * Retrieve the overridden state for the specified property of the specified object
 	 * @param Object to fetch the overridden properties
+	 * @param ConsiderPropertyForOverriddenState optional parameter to consider or not properties in the overridden state evaluation
+	 * @param Path to the property to know its overridden state, empty will return the object overridden state
+	 * @param bOutInheritedState optional parameter to know if the operation returned was inherited from a parent property
 	 * @return the overridden state if the object has overridable serialization enabled */
-	COREUOBJECT_API EOverriddenState GetOverriddenState(TNotNull<UObject*> Object);
+	COREUOBJECT_API EOverriddenState GetOverriddenState(TNotNull<UObject*> Object, 
+		TFunction<bool(TNotNull<const FProperty*>)> ConsiderPropertyForOverriddenState = TFunction<bool(TNotNull<const FProperty*>)>(),
+		const FPropertyVisitorPath& Path = FPropertyVisitorPath(), bool* bOutInheritedState = nullptr);
 
 	/**
 	 * Override the entire object properties and all its instanced subobjects
 	 * @param Object to override all it properties */
-	COREUOBJECT_API void OverrideObject(TNotNull<UObject*> Object);
-
-	/**
-	 * Override all the properties of the specified instanced subobject if it is really owned by the referencer
-	 * @param Object referencing this subobject
-	 * @param InstancedSubObject the sub object to override all its properties */
-	COREUOBJECT_API void OverrideInstancedSubObject(TNotNull<const UObject*> Object, TNotNull<UObject*> InstancedSubObject);
-
-	/**
-	 * Propagate override to all instanced sub object of the specified object
-	 * @param Object to propagate from */
-	COREUOBJECT_API void PropagateOverrideToInstancedSubObjects(TNotNull<UObject*> Object);
+	COREUOBJECT_API void OverrideAllObjectProperties(TNotNull<UObject*> Object);
 
 	/**
 	 * Clears all the overrides on the specified object
@@ -212,8 +206,9 @@ public:
 	 * To be called prior to override a property of the specified object
 	 * Note: Supports object that does not have overridable serialization enabled
 	 * @param Object owning the property
-	 * @param PropertyPath leading to the property about to be overridden */
-	COREUOBJECT_API void PreOverrideProperty(TNotNull<UObject*> Object, const FPropertyVisitorPath& PropertyPath);
+	 * @param PropertyPath leading to the property about to be overridden
+ 	 * @param ChangeType optional current operation */
+	COREUOBJECT_API void PreOverrideProperty(TNotNull<UObject*> Object, const FPropertyVisitorPath& PropertyPath, const EPropertyChangeType::Type ChangeType = EPropertyChangeType::Unspecified);
 
 	/**
 	 * To be called after the property of the specified object was overridden
@@ -267,11 +262,21 @@ public:
 	 * Note: Supports object that does not have overridable serialization enabled
 	 * @param Object owning the property
 	 * @param PropertyChain leading to the property about to be overridden */
-	void PreOverrideProperty(UObject* Object, const FEditPropertyChain& PropertyChain)
+	void PreOverridePropertyChain(UObject* Object, const FEditPropertyChain& PropertyChain)
 	{
 		PreOverrideProperty(Object, FPropertyVisitorPath(FPropertyChangedEvent(nullptr), PropertyChain));
 	}
 
+	/**
+	 * To be called prior to override a property of the specified object
+	 * Note: Supports object that does not have overridable serialization enabled
+	 * @param Object owning the property
+	 * @param PropertyChainEvent information about the type of change including any container item index */
+	void PreOverrideProperty(UObject* Object, const FPropertyChangedChainEvent& PropertyChainEvent)
+	{
+		PreOverrideProperty(Object, FPropertyVisitorPath(PropertyChainEvent, PropertyChainEvent.PropertyChain), PropertyChainEvent.ChangeType);
+	}
+
 	/**
 	 * To be called after the property was overridden of the specified object
 	 * Note: Supports object that does not have overridable serialization enabled
diff --git a/Engine/Source/Runtime/CoreUObject/Public/UObject/OverriddenPropertySet.h b/Engine/Source/Runtime/CoreUObject/Public/UObject/OverriddenPropertySet.h
--- a/Engine/Source/Runtime/CoreUObject/Public/UObject/OverriddenPropertySet.h
+++ b/Engine/Source/Runtime/CoreUObject/Public/UObject/OverriddenPropertySet.h
@@ -516,6 +516,7 @@ protected:
 	UE_API EOverriddenPropertyOperation GetOverriddenPropertyOperation(const FOverriddenPropertyNode* ParentPropertyNode, FPropertyVisitorPath::Iterator PropertyIterator, bool* bOutInheritedOperation, const void* Data) const;
 	UE_API bool ClearOverriddenProperty(FOverriddenPropertyNode& ParentPropertyNode, FPropertyVisitorPath::Iterator PropertyIterator, const void* Data);
 	UE_API void NotifyPropertyChange(FOverriddenPropertyNode* ParentPropertyNode, const EPropertyNotificationType Notification, FPropertyVisitorPath::Iterator PropertyIterator, const EPropertyChangeType::Type ChangeType, const void* Data, bool& bNeedsCleanup);
+	UE_API void NotifyAllPropertyChange(FOverriddenPropertyNode* ParentPropertyNode, const EPropertyNotificationType Notification, const EPropertyChangeType::Type ChangeType, TNotNull<const UStruct*> Struct, const void* Data);
 
 	UE_API EOverriddenPropertyOperation GetOverriddenPropertyOperation(const FOverriddenPropertyNode* ParentPropertyNode, const FArchiveSerializedPropertyChain* CurrentPropertyChain, FProperty* Property) const;
 	UE_API FOverriddenPropertyNode* SetOverriddenPropertyOperation(EOverriddenPropertyOperation Operation, FOverriddenPropertyNode& ParentPropertyNode, const FArchiveSerializedPropertyChain* CurrentPropertyChain, FProperty* Property);
diff --git a/Engine/Source/Runtime/CoreUObject/Public/UObject/OverrideEventHelper.h b/Engine/Source/Runtime/CoreUObject/Public/UObject/OverrideEventHelper.h
--- a/Engine/Source/Runtime/CoreUObject/Public/UObject/OverrideEventHelper.h
+++ b/Engine/Source/Runtime/CoreUObject/Public/UObject/OverrideEventHelper.h
@@ -11,7 +11,7 @@ namespace UE
 	 * Helper method that calls the pre and post change notification responsible for overriding an object.
 	 * This is the preferred way to override because it will notify listeners
 	 * @param Object to override all it's properties */
-	COREUOBJECT_API void SendOverrideObjectEvent(TNotNull<UObject*> Object);
+	COREUOBJECT_API void SendOverrideAllObjectPropertiesEvent(TNotNull<UObject*> Object);
 
 	/**
 	 * Helper method that calls the pre and post change notification responsible for clearing the overrides of an object.
