[UObjectGlobals]
- Add missing construction hook when going wide with the loader

diff --git a/Engine/Source/Runtime/CoreUObject/Private/UObject/UObjectGlobals.cpp b/Engine/Source/Runtime/CoreUObject/Private/UObject/UObjectGlobals.cpp
--- a/Engine/Source/Runtime/CoreUObject/Private/UObject/UObjectGlobals.cpp
+++ b/Engine/Source/Runtime/CoreUObject/Private/UObject/UObjectGlobals.cpp
@@ -3863,7 +3863,7 @@ UObject* StaticAllocateObject
 		}
 	}
 
-	if (IsInAsyncLoadingThread())
+	if (IsInAsyncLoadingThread() || IsInParallelLoadingThread())
 	{
 		FUObjectThreadContext& ThreadContext = FUObjectThreadContext::Get();
 		if (ThreadContext.AsyncPackageLoader)
