[UObjectGlobals]
- Make another delegate thread-safe by default

diff --git a/Engine/Source/Runtime/CoreUObject/Public/UObject/UObjectGlobals.h b/Engine/Source/Runtime/CoreUObject/Public/UObject/UObjectGlobals.h
--- a/Engine/Source/Runtime/CoreUObject/Public/UObject/UObjectGlobals.h
+++ b/Engine/Source/Runtime/CoreUObject/Public/UObject/UObjectGlobals.h
@@ -3314,7 +3314,7 @@ struct FCoreUObjectDelegates
 	 * This is called when the CDO and its associated class structure have been fully generated and populated, and allows the assignment of cached/derived data,
 	 * eg) caching the name/count of properties of a certain type, or inspecting the properties on the class and using their meta-data and CDO default values to derive game data.
 	 */
-	DECLARE_MULTICAST_DELEGATE_TwoParams(FOnObjectPostCDOCompiled, UObject*, const FObjectPostCDOCompiledContext&);
+	DECLARE_TS_MULTICAST_DELEGATE_TwoParams(FOnObjectPostCDOCompiled, UObject*, const FObjectPostCDOCompiledContext&);
 	static COREUOBJECT_API FOnObjectPostCDOCompiled OnObjectPostCDOCompiled;
 
 	/** Callback for when an asset is saved. This is called from UObject::PreSave before it is actually written to disk, for every object saved */
