ZenLoader: Re-enable GEnableScriptImportsDebugging by default (it was previously disabled in 42978094).
It no longer causes standalone programs (BuildStorageTool, CrashReportClient) to hang on startup.
Don't try and read the ScriptObjects chunk if it does not exist, helps in any scenario where the iodispatcher has not been fully initialized yet.
Compile out the whole feature for programs.

diff --git a/Engine/Source/Runtime/CoreUObject/Private/Serialization/AsyncLoading2.cpp b/Engine/Source/Runtime/CoreUObject/Private/Serialization/AsyncLoading2.cpp
--- a/Engine/Source/Runtime/CoreUObject/Private/Serialization/AsyncLoading2.cpp
+++ b/Engine/Source/Runtime/CoreUObject/Private/Serialization/AsyncLoading2.cpp
@@ -324,10 +324,10 @@ void FindAllRuntimeScriptPackages(FRuntimeScriptPackages& OutPackages)
 #endif
 
 #ifndef ALT2_ENABLE_SCRIPT_IMPORTS_DEBUGGING
-#define ALT2_ENABLE_SCRIPT_IMPORTS_DEBUGGING !UE_SHIPPING
+#define ALT2_ENABLE_SCRIPT_IMPORTS_DEBUGGING !UE_SHIPPING && !IS_PROGRAM
 #endif
 #if ALT2_ENABLE_SCRIPT_IMPORTS_DEBUGGING
-static bool GEnableScriptImportsDebugging = false;
+static bool GEnableScriptImportsDebugging = true;
 static FAutoConsoleVariableRef CEnableScriptImportsDebugging(
 	TEXT("s.EnableScriptImportsDebugging"),
 	GEnableScriptImportsDebugging,
@@ -6314,8 +6314,16 @@ void FCookedScriptObjectsDebug::LoadDebugData()
 	TRACE_CPUPROFILER_EVENT_SCOPE(FCookedScriptObjectsDebug::LoadDebugData);
 
 	FIoDispatcher& IoDispatcher = FIoDispatcher::Get();
+	FIoChunkId ChunkId = CreateIoChunkId(0, 0, EIoChunkType::ScriptObjects);
+	// Do a quick sync check if we are running in a valid scenario (cooked client or server) where we actually have the script objects chunk.
+	// If not, early out since the iodispatcher may not even be fully initialized yet and the read may never complete.
+	if (!IoDispatcher.DoesChunkExist(ChunkId))
+	{
+		return;
+	}
+
 	FIoBatch IoBatch = IoDispatcher.NewBatch();
-	FIoRequest Request = IoBatch.Read(CreateIoChunkId(0, 0, EIoChunkType::ScriptObjects), FIoReadOptions(), IoDispatcherPriority_Max);
+	FIoRequest Request = IoBatch.Read(ChunkId, FIoReadOptions(), IoDispatcherPriority_Max);
 	FEventRef Event;
 	IoBatch.IssueAndTriggerEvent(Event.Get());
 	Event->Wait();
@@ -12595,4 +12603,4 @@ bool FPublicExportMapPreInsertPublicExportsTest::RunTest(const FString& Paramete
 	return true;
 }
 
-#endif // WITH_LOW_LEVEL_TESTS
\ No newline at end of file
+#endif // WITH_LOW_LEVEL_TESTS
