It's no longer required to initialize the global server id during UObject system init. The Id can be assigned after when servers connect to each other.

Deprecated UE::RemoteObject::GetGlobalServerId() and replaced it with FRemoteServerId::GetLocalServerId() to better reflect the fact it's always the 'Local' id instead of the global id number.
- although GetLocalServerId() returns a constant "Local" id comarison operators and FRemoteServerId::GetIdNumber() use the actual  'global' id number
- added restrictions (checks) to when the comparison operators (and GetIdNumber()) can be used
Deprecated FRemoteServerId(uint32) and replaced it with FRemoteServerId::FromIdNumber(uint32) to prevent the abuse of the constructor and match it with FRemoteServerId::GetIdNumber()
Deprecated FRemoteServerId(const FString&) and replaced it with FRemoteServerId::FromString(const FString&)

diff --git a/Engine/Source/Runtime/CoreUObject/Private/UObject/Class.cpp b/Engine/Source/Runtime/CoreUObject/Private/UObject/Class.cpp
--- a/Engine/Source/Runtime/CoreUObject/Private/UObject/Class.cpp
+++ b/Engine/Source/Runtime/CoreUObject/Private/UObject/Class.cpp
@@ -7988,6 +7988,7 @@ UScriptStruct* TBaseStructure<FRemoteObjectId>::Get()
 	static UScriptStruct* ScriptStruct = StaticGetBaseStructureInternal(TEXT("RemoteObjectId"));
 	return ScriptStruct;
 }
+UE_IMPLEMENT_STRUCT("/Script/CoreUObject", RemoteObjectId);
 
 UScriptStruct* TBaseStructure<FRemoteObjectTables>::Get()
 {
@@ -8024,6 +8025,7 @@ UScriptStruct* TBaseStructure<FRemoteServerId>::Get()
 	static UScriptStruct* ScriptStruct = StaticGetBaseStructureInternal(TEXT("RemoteServerId"));
 	return ScriptStruct;
 }
+UE_IMPLEMENT_STRUCT("/Script/CoreUObject", RemoteServerId);
 
 UScriptStruct* TBaseStructure<FRemoteObjectReference>::Get()
 {
diff --git a/Engine/Source/Runtime/CoreUObject/Private/UObject/Obj.cpp b/Engine/Source/Runtime/CoreUObject/Private/UObject/Obj.cpp
--- a/Engine/Source/Runtime/CoreUObject/Private/UObject/Obj.cpp
+++ b/Engine/Source/Runtime/CoreUObject/Private/UObject/Obj.cpp
@@ -5603,9 +5603,7 @@ void StaticUObjectInit()
 		}
 	}
 
-#if UE_WITH_REMOTE_OBJECT_HANDLE
 	UE::RemoteObject::Private::InitRemoteObjects();
-#endif // UE_WITH_REMOTE_OBJECT_HANDLE
 
 	UE_LOG(LogInit, Log, TEXT("Object subsystem initialized") );
 }
diff --git a/Engine/Source/Runtime/CoreUObject/Private/UObject/RemoteExecutor.cpp b/Engine/Source/Runtime/CoreUObject/Private/UObject/RemoteExecutor.cpp
--- a/Engine/Source/Runtime/CoreUObject/Private/UObject/RemoteExecutor.cpp
+++ b/Engine/Source/Runtime/CoreUObject/Private/UObject/RemoteExecutor.cpp
@@ -888,7 +888,7 @@ void ReadyRemoteMultiServerCommit(FRemoteServerId ServerId, FRemoteTransactionId
 
 FRemoteWorkPriority CreateRootWorkPriority()
 {
-	FRemoteWorkPriority RootWorkPriority = FRemoteWorkPriority::CreateRootWorkPriority(UE::RemoteObject::GetGlobalServerId(), GRemoteExecutor.GenerateNextTransactionId());
+	FRemoteWorkPriority RootWorkPriority = FRemoteWorkPriority::CreateRootWorkPriority(FRemoteServerId::GetLocalServerId(), GRemoteExecutor.GenerateNextTransactionId());
 	return RootWorkPriority;
 }
 
@@ -908,7 +908,7 @@ void ExecuteTransactionalWithExplicitPriority(FName WorkName, FRemoteWorkPriorit
 
 void ExecuteTransactional(FName WorkName, const TFunctionRef<void(void)>& Work)
 {
-	FRemoteWorkPriority RootWorkPriority = FRemoteWorkPriority::CreateRootWorkPriority(UE::RemoteObject::GetGlobalServerId(), GRemoteExecutor.GenerateNextTransactionId());
+	FRemoteWorkPriority RootWorkPriority = FRemoteWorkPriority::CreateRootWorkPriority(FRemoteServerId::GetLocalServerId(), GRemoteExecutor.GenerateNextTransactionId());
 	ExecuteTransactionalWithExplicitPriority(WorkName, RootWorkPriority, Work);
 }
 
@@ -916,7 +916,7 @@ void EnqueueTransactional(FName WorkName, const TFunction<void(void)>& InWork)
 {
 #if UE_WITH_REMOTE_OBJECT_HANDLE
 	check(!AutoRTFM::IsClosed());
-	FRemoteWorkPriority RootWorkPriority = FRemoteWorkPriority::CreateRootWorkPriority(UE::RemoteObject::GetGlobalServerId(), GRemoteExecutor.GenerateNextTransactionId());
+	FRemoteWorkPriority RootWorkPriority = FRemoteWorkPriority::CreateRootWorkPriority(FRemoteServerId::GetLocalServerId(), GRemoteExecutor.GenerateNextTransactionId());
 	GRemoteExecutor.EnqueueWork(WorkName, RootWorkPriority, true, InWork);
 #else
 	InWork();
@@ -937,8 +937,7 @@ FRemoteWorkPriority FRemoteWorkPriority::CreateRootWorkPriority(FRemoteServerId
 {
 	FRemoteWorkPriority Result{};
 
-	uint32 RawServerId = 0;
-	FMemory::Memcpy(&RawServerId, &ServerId, sizeof(FRemoteServerId));
+	uint32 RawServerId = ServerId.GetIdNumber();
 
 	uint32 RawTransactionId = 0;
 	FMemory::Memcpy(&RawTransactionId, &TransactionId, sizeof(FRemoteTransactionId));
diff --git a/Engine/Source/Runtime/CoreUObject/Private/UObject/RemoteObject.cpp b/Engine/Source/Runtime/CoreUObject/Private/UObject/RemoteObject.cpp
--- a/Engine/Source/Runtime/CoreUObject/Private/UObject/RemoteObject.cpp
+++ b/Engine/Source/Runtime/CoreUObject/Private/UObject/RemoteObject.cpp
@@ -20,22 +20,76 @@
 #include "Misc/CoreDelegates.h"
 #include "String/LexFromString.h"
 #include "Hash/CityHash.h"
+#include "Modules/VisualizerDebuggingState.h"
 #include <atomic>
 
 DEFINE_LOG_CATEGORY(LogRemoteObject);
 
+int32 GInitRemoteServerIdBeforeUObjectInit = 1;
+static FAutoConsoleVariableRef CVarInitRemoteServerIdBeforeUObjectInit(
+	TEXT("ro.InitRemoteServerIdBeforeUObjectInit"),
+	GInitRemoteServerIdBeforeUObjectInit,
+	TEXT("Initializes remote server id before UObject system is initialized"));
+
+FRemoteServerId FRemoteServerId::GlobalServerId(ERemoteServerIdConstants::Local); // Here Local means uninitialized
+
+void AssignGlobalServerIdDebuggingState()
+{
+	// CF8B6D3D-3185-453C-AF12-88EB19245359 => cf8b6d3d3185453caf1288eb19245359
+	constexpr FGuid GGlobalServerIdGuid = FGuid(0xCF8B6D3D, 0x3185453C, 0xAF1288EB, 0x19245359);
+
+	(void)::UE::Core::FVisualizerDebuggingState::Assign(GGlobalServerIdGuid, &FRemoteServerId::GlobalServerId);
+}
+
+void FRemoteServerId::InitGlobalServerId(FRemoteServerId Id)
+{
+	checkf(GlobalServerId.Id == (uint32)ERemoteServerIdConstants::Local, TEXT("Global server id has already been initialized (%s)"), *GlobalServerId.ToString());
+	GlobalServerId = Id;
+}
+
+bool FRemoteServerId::IsGlobalServerIdInitialized()
+{
+	// This is the only place we check if the id is valid.
+	// We can be running a build with remote object support compiled in but disabled in which case we initialize GlobalServerId to an invalid value
+	// but we still want to return false in this case because of existing use cases that check this value to determine if remote objects are enabled
+	return GlobalServerId.Id != (uint32)ERemoteServerIdConstants::Local;
+}
+
+
 namespace UE::RemoteObject::Private
 {
 
-FRemoteServerId RemoteServerId;
 std::atomic<uint64> RemoteObjectSerialNumber(1);
 std::atomic<uint64> AssetObjectSerialNumber(1);
 int32 UnsafeToMigrateObjects = 0; // This should go into TLS
 
+class FRemoteObjectStubMap : private TMap<FRemoteObjectId, UE::RemoteObject::Handle::FRemoteObjectStub*>
+{
+	using Super = TMap<FRemoteObjectId, UE::RemoteObject::Handle::FRemoteObjectStub*>;
+
+public:
+
+	using Super::begin;
+	using Super::end;
+
+	UE::RemoteObject::Handle::FRemoteObjectStub*& FindOrAdd(FRemoteObjectId Id)
+	{		
+		return Super::FindOrAdd(UE::RemoteObject::Private::FRemoteIdLocalizationHelper::GetLocalized(Id));
+	}
+
+	UE::RemoteObject::Handle::FRemoteObjectStub* Find(FRemoteObjectId Id)
+	{
+		using namespace UE::RemoteObject::Handle;
+
+		UE::RemoteObject::Handle::FRemoteObjectStub** ExistingStub = Super::Find(UE::RemoteObject::Private::FRemoteIdLocalizationHelper::GetLocalized(Id));
+		return ExistingStub ? *ExistingStub : nullptr;
+	}
+};
+
 class FRemoteObjectMaps
 {
 	mutable FTransactionallySafeCriticalSection ObjectMapCritical;
-	TMap<FRemoteObjectId, UE::RemoteObject::Handle::FRemoteObjectStub*> RemoteObjects;
+	FRemoteObjectStubMap RemoteObjects;
 	TMap<FRemoteObjectId, FRemoteObjectPathName> AssetPaths;
 
 public:
@@ -45,11 +99,7 @@ public:
 		using namespace UE::RemoteObject::Handle;
 
 		UE::TScopeLock MapLock(ObjectMapCritical);
-		if (FRemoteObjectStub** ExistingStub = RemoteObjects.Find(Id))
-		{
-			return *ExistingStub;
-		}
-		return nullptr;
+		return RemoteObjects.Find(Id);
 	}
 
 	UE::RemoteObject::Handle::FRemoteObjectStub* FindOrAddRemoteObjectStub(FRemoteObjectId Id, FRemoteServerId ResidentServerId = FRemoteServerId())
@@ -67,7 +117,7 @@ public:
 			// if we are creating the stub, then this object's owner is deduced from its ID
 			// if the server ID is invalid then it's a local native object that was created before the local server had its ID assigned
 			FRemoteServerId ObjectServerId = Id.GetServerId();
-			Stub->OwningServerId = ObjectServerId.IsValid() ? ObjectServerId : GetGlobalServerId();
+			Stub->OwningServerId = ObjectServerId.IsValid() ? ObjectServerId : FRemoteServerId::GetLocalServerId();
 		}		
 		return Stub;
 	}
@@ -77,7 +127,7 @@ public:
 		using namespace UE::RemoteObject::Handle;
 
 		UE::TScopeLock MapLock(ObjectMapCritical);
-		FRemoteObjectStub*& Stub = RemoteObjects.FindOrAdd(UE::CoreUObject::Private::FObjectHandleUtils::GetRemoteId(Object));
+		FRemoteObjectStub*& Stub = RemoteObjects.FindOrAdd(FRemoteObjectId(Object));
 		if (!Stub)
 		{
 			Stub = new FRemoteObjectStub(Object);
@@ -85,7 +135,7 @@ public:
 			// if we are creating the stub, then this object's owner is deduced from its ID
 			// if the server ID is invalid then it's a local native object that was created before the local server had its ID assigned
 			FRemoteServerId ObjectServerId = Stub->Id.GetServerId();
-			Stub->OwningServerId = ObjectServerId.IsValid() ? ObjectServerId : GetGlobalServerId();
+			Stub->OwningServerId = ObjectServerId.IsValid() ? ObjectServerId : FRemoteServerId::GetLocalServerId();
 		}
 
 		FUObjectItem* ObjectItem = GUObjectArray.ObjectToObjectItem(Object);
@@ -165,6 +215,7 @@ FRemoteObjectMaps* ObjectMaps = nullptr;
 
 void InitServerId()
 {
+	FRemoteServerId GlobalServerId;
 	FString ServerId;
 	const TCHAR* CommandLine = FCommandLine::Get();
 	if (!FParse::Value(CommandLine, TEXT("MultiServerLocalId="), ServerId))
@@ -179,26 +230,35 @@ void InitServerId()
 					ServerId = FString::FromInt(ListenPort % 1000);
 				}
 			}
-			else
-			{
-				RemoteServerId = FRemoteServerId(ERemoteServerIdConstants::Invalid);
-			}
 		}
 	}
 	if (!ServerId.IsEmpty())
 	{
-		RemoteServerId = FRemoteServerId(ServerId);
-		checkf(RemoteServerId.IsValid(), TEXT("Remote ServerId is not valid"));
+		GlobalServerId = FRemoteServerId::FromString(ServerId);
+		checkf(GlobalServerId.IsValid(), TEXT("Remote ServerId is not valid"));
+
+		FRemoteServerId::InitGlobalServerId(GlobalServerId);
 	}
 
-	UE_LOG(LogRemoteObject, Display, TEXT("Remote ServerId: %s"), *RemoteServerId.ToString());
+	UE_LOG(LogRemoteObject, Display, TEXT("Global Server Id: %s"), *FRemoteServerId::GetLocalServerId().ToString());
 }
 
 void InitRemoteObjects()
 {
+	AssignGlobalServerIdDebuggingState();
+
+	if (!FRemoteObjectId::RemoteObjectSupportCompiledIn)
+	{
+		// Always init global server id debug visualizer support but early out if remote object support is not compiled in (UE_WITH_REMOTE_OBJECT_HANDLE is 0)
+		return;
+	}
+
 	ObjectMaps = new FRemoteObjectMaps();
 
-	InitServerId();
+	if (GInitRemoteServerIdBeforeUObjectInit)
+	{
+		InitServerId();
+	}
 
 	UE::RemoteObject::Transfer::InitRemoteObjectTransfer();
 
@@ -309,7 +369,7 @@ void MarkAsLocal(UObject* Object)
 	ObjectItem->ClearFlags(EInternalObjectFlags::Remote);
 	ObjectItem->SetFlags(EInternalObjectFlags::RemoteReference);
 #endif
-	ObjectMaps->FindOrAddRemoteObjectStub(Object, GetGlobalServerId());
+	ObjectMaps->FindOrAddRemoteObjectStub(Object, FRemoteServerId::GetLocalServerId());
 }
 
 void StoreAssetPath(UObject* Object)
@@ -336,9 +396,9 @@ FName GetServerBaseNameForUniqueName(const UClass* Class)
 	checkf(Class, TEXT("Unable to generate base name for a unique object name without the object's Class"));
 
 	// Packages follow different naming rules than other UObjects and ATM we're not migrating packages so fall back to Class->GetFName()
-	if (GetGlobalServerId().IsValid() && Class->GetFName() != NAME_Package)
+	if (FRemoteServerId::IsGlobalServerIdInitialized() && Class->GetFName() != NAME_Package)
 	{
-		return *FString::Printf(TEXT("%s_S%s"), *Class->GetFName().GetPlainNameString(), *GetGlobalServerId().ToString());
+		return *FString::Printf(TEXT("%s_S%s"), *Class->GetFName().GetPlainNameString(), *FRemoteServerId::GetLocalServerId().ToString());
 	}
 	return Class->GetFName();
 }
@@ -361,14 +421,6 @@ bool IsSafeToMigrateObjects()
 
 } // namespace UE::RemoteObject::Private
 
-namespace UE::RemoteObject
-{
-	FRemoteServerId GetGlobalServerId()
-	{
-		return UE::RemoteObject::Private::RemoteServerId;
-	}
-}
-
 namespace UE::RemoteObject::Handle
 {
 
@@ -404,7 +456,7 @@ bool IsRemote(FRemoteObjectId ObjectId)
 
 	const FRemoteServerId ServerId = ObjectId.GetServerId();
 	// Invalid server Id means local native classes which are created before a server has a chance to have an id assigned
-	return ServerId.IsValid() && ServerId != UE::RemoteObject::GetGlobalServerId();
+	return ServerId.IsValid() && !ServerId.IsLocal() && !ServerId.IsAsset();
 }
 
 bool IsRemote(const UObject* Object)
@@ -434,13 +486,13 @@ bool IsOwned(FRemoteObjectId ObjectId)
 	UE::RemoteObject::Handle::FRemoteObjectStub* RemoteStub = FindRemoteObjectStub(ObjectId);
 	if (RemoteStub)
 	{
-		bResult = (RemoteStub->OwningServerId == GetGlobalServerId() || RemoteStub->OwningServerId.IsAsset());
+		bResult = (RemoteStub->OwningServerId == FRemoteServerId::GetLocalServerId() || RemoteStub->OwningServerId.IsAsset());
 	}
 	else
 	{
 		const FRemoteServerId ServerId = ObjectId.GetServerId();
 		// Invalid server Id means local native objects which were created before the local server had a chance to have an id assigned
-		bResult = (!ServerId.IsValid() || ServerId.IsAsset() || ServerId == UE::RemoteObject::GetGlobalServerId());
+		bResult = (!ServerId.IsValid() || ServerId.IsAsset() || ServerId.IsLocal());
 	}
 #endif
 	return bResult;
@@ -461,7 +513,7 @@ FRemoteServerId GetOwnerServerId(const UObject* Object)
 	else
 	{
 		// if the object wasn't received or ever migrated, we own it locally
-		Result = GetGlobalServerId();
+		Result = FRemoteServerId::GetLocalServerId();
 	}
 #endif
 	return Result;
@@ -652,7 +704,7 @@ namespace UE::CoreUObject::Private
 		// ObjectItem may not exist when the UObject system hasn't been initialized yet but theoretically this function should only get called when
 		// something attempts to re-construct a default subobject that already exists so ObjectItem should always be valid
 		checkf(ObjectItem, TEXT("Attempting to change remote ID for an object that does not exist in the global UObject array (it's possible GUObjectArray is not initialized yet, ObjectIndex=%d)"), GUObjectArray.ObjectToIndex(Object));
-		ObjectItem->RemoteId = Id;
+		ObjectItem->SetRemoteId(Id);
 		HashObject(Object);
 	}
 
@@ -711,7 +763,7 @@ FRemoteObjectId FRemoteObjectId::Generate(UObjectBase* InObject, const TCHAR* In
 
 	bool bIsAsset = false;
 #if UE_WITH_REMOTE_ASSET_ID
-	if (GIsInitialLoad || !UE::RemoteObject::GetGlobalServerId().IsValid() || !!(InInitialFlags & EInternalObjectFlags::Native) || !!(InObject->GetFlags() & RF_ArchetypeObject))
+	if (GIsInitialLoad || !!(InInitialFlags & EInternalObjectFlags::Native) || !!(InObject->GetFlags() & RF_ArchetypeObject))
 	{
 		// Native objects (classes, CDOs etc) or objects loaded during initial load are always in memory and are considered assets any server can find locally
 		// Note that this first condition can not touch too much of UObject API because we might literarlly be constructing the first StaticClass() etc.
@@ -753,7 +805,7 @@ FRemoteObjectId FRemoteObjectId::Generate(UObjectBase* InObject, const TCHAR* In
 
 	if (!bIsAsset)
 	{
-		Result = FRemoteObjectId(RemoteServerId, RemoteObjectSerialNumber.fetch_add(1));
+		Result = FRemoteObjectId(FRemoteServerId(ERemoteServerIdConstants::Local), RemoteObjectSerialNumber.fetch_add(1));
 		UE_LOG(LogRemoteObject, Verbose, TEXT("FRemoteObjectId::Generate: Type=Dynamic Object=%s RemoteId=%s"), InObject ? *InObject->GetFName().ToString() : TEXT("None"), *Result.ToString());
 	}
 	else
@@ -797,21 +849,78 @@ FRemoteObjectId::FRemoteObjectId(const UObjectBase* Object)
 {
 }
 
+bool FRemoteObjectId::Serialize(FArchive& Ar)
+{
+	static_assert(sizeof(uint64) == sizeof(Id));
+
+	uint64 SerializedID = Id;
+	if (Ar.IsSaving())
+	{
+		FRemoteObjectId GlobalizedId = GetGlobalized();
+		SerializedID = GlobalizedId.Id;
+	}
+
+	Ar << SerializedID;
+	{
+		FRemoteObjectId LocalizedId;
+		LocalizedId.Id = SerializedID;
+		checkf(LocalizedId.GetServerId().Id != (uint32)ERemoteServerIdConstants::Local, TEXT("Local server id should never be serialized"));
+	}
+
+	if (Ar.IsLoading())
+	{
+		FRemoteObjectId LocalizedId;
+		LocalizedId.Id = SerializedID;
+		checkf(LocalizedId.GetServerId().Id != (uint32)ERemoteServerIdConstants::Local, TEXT("Local server id should never be serialized"));
+		*this = LocalizedId.GetLocalized();
+	}
+
+	return true;
+}
+
+bool FRemoteObjectId::NetSerialize(FArchive& Ar, class UPackageMap* Map, bool& bOutSuccess)
+{
+	return Serialize(Ar);
+}
+
 FArchive& operator<<(FArchive& Ar, FRemoteObjectId& Id)
 {
-	Ar << Id.Id;
+	Id.Serialize(Ar);
 	return Ar;
 }
 
+FRemoteServerId FRemoteServerId::FromIdNumber(uint32 InNumber)
+{
+	checkf(InNumber < (uint32)ERemoteServerIdConstants::FirstReserved, TEXT("Remote server id can not be greater than %u, got: %u"), (uint32)ERemoteServerIdConstants::FirstReserved - 1, InNumber);
+	FRemoteServerId Result;
+	Result.Id = InNumber;
+	return Result;
+}
+
 FRemoteServerId::FRemoteServerId(const FString& InText)
 {
-	if (InText == TEXT("Asset"))
+	*this = FromString(InText);
+}
+
+FRemoteServerId FRemoteServerId::FromString(const FString& InText)
+{
+	FRemoteServerId Result;
+
+	if (InText == TEXT("Invalid"))
+	{
+		Result.Id = (uint32)ERemoteServerIdConstants::Invalid;
+	}
+	else if (InText == TEXT("Asset"))
 	{
-		Id = (uint32)ERemoteServerIdConstants::Asset;
+		Result.Id = (uint32)ERemoteServerIdConstants::Asset;
 	}
 	else if (InText == TEXT("Database"))
 	{
-		Id = (uint32)ERemoteServerIdConstants::Database;
+		Result.Id = (uint32)ERemoteServerIdConstants::Database;
+	}
+	else if (InText == TEXT("Local"))
+	{
+		Result.Id = (uint32)ERemoteServerIdConstants::Local;
 	}
 	else
 	{
@@ -819,14 +928,16 @@ FRemoteServerId::FRemoteServerId(const FString& InText)
 		LexFromString(ServerIdNumber, InText);
 		if (ensureMsgf(ServerIdNumber <= (uint32)ERemoteServerIdConstants::Max, TEXT("Parsed Remote Server Id value %u that is bigger than allowed max %u"), ServerIdNumber, (uint32)ERemoteServerIdConstants::Max))
 		{
-			Id = ServerIdNumber;
+			Result.Id = ServerIdNumber;
 		}
 		else
 		{
 			UE_LOG(LogRemoteObject, Warning, TEXT("Clamping ServerId number %u to the maximum allowed %u"), ServerIdNumber, (uint32)ERemoteServerIdConstants::Max);
-			Id = (uint32)ERemoteServerIdConstants::Max;
+			Result.Id = (uint32)ERemoteServerIdConstants::Max;
 		}		
 	}
+
+	return Result;
 }
 
 FString FRemoteServerId::ToString() const
@@ -839,14 +950,46 @@ FString FRemoteServerId::ToString() const
 		case (uint32)ERemoteServerIdConstants::Database:
 			return TEXT("Database");
 
+		case (uint32)ERemoteServerIdConstants::Local:
+			return IsGlobalServerIdInitialized() ? GetGlobalized().ToString() : TEXT("Local");
+
 		default:
 			return FString::FromInt(Id);
 	}
 }
 
+bool FRemoteServerId::Serialize(FArchive& Ar)
+{
+	static_assert(sizeof(uint32) == sizeof(Id));
+
+	uint32 SerializedId = Id;
+	if (Ar.IsSaving())
+	{
+		FRemoteServerId GlobalizedId = GetGlobalized();
+		SerializedId = GlobalizedId.Id;
+	}
+
+	Ar << SerializedId;
+	checkf(SerializedId != (uint32)ERemoteServerIdConstants::Local, TEXT("Local server id should never be serialized"));
+
+	if (Ar.IsLoading())
+	{
+		FRemoteServerId LocalizedId;
+		LocalizedId.Id = SerializedId;
+		*this = LocalizedId.GetLocalized();
+	}	
+
+	return true;
+}
+
+bool FRemoteServerId::NetSerialize(FArchive& Ar, class UPackageMap* Map, bool& bOutSuccess)
+{
+	return Serialize(Ar);
+}
+
 FArchive& operator<<(FArchive& Ar, FRemoteServerId& Id)
 {
-	Ar << Id.Id;
+	Id.Serialize(Ar);
 	return Ar;
 }
 
@@ -887,7 +1030,7 @@ COREUOBJECT_API UE::RemoteObject::Handle::FRemoteObjectStub* DebugFindRemoteObje
  */
 COREUOBJECT_API UObject* DebugFindObjectLocallyFromRemoteId(uint16 ServerId, uint64 SerialNumber)
 {
-	return StaticFindObjectFastInternal(FRemoteObjectId(FRemoteServerId{ static_cast<uint32>(ServerId) }, SerialNumber));
+	return StaticFindObjectFastInternal(FRemoteObjectId(FRemoteServerId::FromIdNumber(static_cast<uint32>(ServerId)), SerialNumber));
 }
 
 /**
@@ -898,7 +1041,7 @@ COREUOBJECT_API UObject* DebugFindObjectLocallyFromRemoteId(uint16 ServerId, uin
 COREUOBJECT_API UObject* DebugFindObjectLocallyFromRemoteId(uint64 FullId)
 {
 	constexpr uint64 SerialBitMask = (1ull << 54ull) - 1ull;
-	FRemoteObjectId RemoteId(FRemoteServerId{ static_cast<uint32>((FullId >> 54ull) & 1023ull) }, FullId & SerialBitMask);
+	FRemoteObjectId RemoteId(FRemoteServerId::FromIdNumber(static_cast<uint32>((FullId >> 54ull) & 1023ull)), FullId & SerialBitMask);
 	ensure(RemoteId.GetIdNumber() == FullId);
 
 	return StaticFindObjectFastInternal(RemoteId);
diff --git a/Engine/Source/Runtime/CoreUObject/Private/UObject/RemoteObjectPrivate.h b/Engine/Source/Runtime/CoreUObject/Private/UObject/RemoteObjectPrivate.h
--- a/Engine/Source/Runtime/CoreUObject/Private/UObject/RemoteObjectPrivate.h
+++ b/Engine/Source/Runtime/CoreUObject/Private/UObject/RemoteObjectPrivate.h
@@ -84,6 +84,26 @@ namespace UE::RemoteObject::Private
 		FUnsafeToMigrateScope();
 		~FUnsafeToMigrateScope();
 	};
+
+	struct FRemoteIdLocalizationHelper
+	{
+		inline static FRemoteServerId GetLocalized(FRemoteServerId InId)
+		{
+			return InId.GetLocalized();
+		}
+		inline static FRemoteServerId GetGlobalized(FRemoteServerId InId)
+		{
+			return InId.GetGlobalized();
+		}
+		inline static FRemoteObjectId GetLocalized(FRemoteObjectId InId)
+		{
+			return InId.GetLocalized();
+		}
+		inline static FRemoteObjectId GetGlobalized(FRemoteObjectId InId)
+		{
+			return InId.GetGlobalized();
+		}
+	};
 }
 
 namespace UE::RemoteObject::Transfer::Private
diff --git a/Engine/Source/Runtime/CoreUObject/Private/UObject/RemoteObjectSerialization.cpp b/Engine/Source/Runtime/CoreUObject/Private/UObject/RemoteObjectSerialization.cpp
--- a/Engine/Source/Runtime/CoreUObject/Private/UObject/RemoteObjectSerialization.cpp
+++ b/Engine/Source/Runtime/CoreUObject/Private/UObject/RemoteObjectSerialization.cpp
@@ -83,7 +83,7 @@ void DeserializeNameTables(FArchive& Ar, FRemoteObjectData& OutObjectData)
 
 FString GenerateRemoteObjectFilename(FRemoteObjectId ObjectId, FRemoteServerId OwnerServerId)
 {
-	return FPaths::Combine(*FPaths::ProjectSavedDir(), *(FString::Printf(TEXT("%s-%s_%s.remote"), *UE::RemoteObject::GetGlobalServerId().ToString(), *ObjectId.ToString(), *OwnerServerId.ToString())));
+	return FPaths::Combine(*FPaths::ProjectSavedDir(), *(FString::Printf(TEXT("%s-%s_%s.remote"), *FRemoteServerId::GetLocalServerId().ToString(), *ObjectId.ToString(), *OwnerServerId.ToString())));
 }
 
 void LoadObjectFromDisk(const FUObjectMigrationContext& MigrationContext)
@@ -99,7 +99,7 @@ void LoadObjectFromDisk(const FUObjectMigrationContext& MigrationContext)
 	IFileManager::Get().Delete(*Filename, false, true, true);
 
 	// We have transferred ownership from the Database to the local server
-	FRemoteServerId LocalServerId = UE::RemoteObject::GetGlobalServerId();
+	FRemoteServerId LocalServerId = FRemoteServerId::GetLocalServerId();
 	const FRemoteServerId& DatabaseId = UE::RemoteObject::Transfer::DatabaseId;
 	UE::RemoteObject::Transfer::OnObjectDataReceived(LocalServerId, LocalServerId, 0, MigrationContext.ObjectId, DatabaseId, ObjectData);
 }
@@ -385,7 +385,7 @@ public:
 				Info.Object = ObjPtr.Get();
 				Info.bIsSubobject = Info.Object->IsIn(InRootObject);
 
-				if (Info.Object == InRootObject || Info.bIsSubobject || (Info.Id.GetServerId().IsValid() && Info.Id.GetServerId() != GetGlobalServerId() && !Info.Id.IsAsset()))
+				if (Info.Object == InRootObject || Info.bIsSubobject || (Info.Id.GetServerId().IsValid() && Info.Id.GetServerId() != FRemoteServerId::GetLocalServerId() && !Info.Id.IsAsset()))
 				{
 					Info.Type = ERemoteReferenceType::IdOnly;
 				}
@@ -1249,7 +1249,7 @@ int32 DeserializeObjectData(FRemoteObjectData& ObjectData,
 
 	if (ObjectHeaders.Num())
 	{
-		const FRemoteServerId LocalServerId = UE::RemoteObject::GetGlobalServerId();
+		const FRemoteServerId LocalServerId = FRemoteServerId::GetLocalServerId();
 
 		// Try to find any existing WeakObject serial numbers for the objects that are about to be constructed
 		for (FRemoteObjectHeader& ObjectHeader : ObjectHeaders)
diff --git a/Engine/Source/Runtime/CoreUObject/Private/UObject/RemoteObjectTransfer.cpp b/Engine/Source/Runtime/CoreUObject/Private/UObject/RemoteObjectTransfer.cpp
--- a/Engine/Source/Runtime/CoreUObject/Private/UObject/RemoteObjectTransfer.cpp
+++ b/Engine/Source/Runtime/CoreUObject/Private/UObject/RemoteObjectTransfer.cpp
@@ -399,7 +399,7 @@ public:
 						HighestPriorityRequest->GetPriority(),
 						ObjectIdToRequest,
 						CurrentResidentServerId,
-						GetGlobalServerId());
+						FRemoteServerId::GetLocalServerId());
 			
 					if (bUpdatingPriority)
 					{
@@ -476,12 +476,12 @@ public:
 			{
 				// if the used object is not an asset, but the root of it is, then we
 				// expect that the used object owner is the local server
-				check(OwnerServerId == GetGlobalServerId());
+				check(OwnerServerId == FRemoteServerId::GetLocalServerId());
 
 				continue;
 			}
 
-			if (RootOwnerServerId != GetGlobalServerId() && !RootOwnerServerId.IsAsset())
+			if (RootOwnerServerId != FRemoteServerId::GetLocalServerId() && !RootOwnerServerId.IsAsset())
 			{
 				UE_LOG(LogRemoteObject, VeryVerbose, TEXT("FRemoteObjectTransferQueue: borrowed obj: %s (%s) owner server: %s (root: %s (s) root owner: %s)"),
 					*UsedObjectId.ToString(),
@@ -620,7 +620,7 @@ public:
 
 	void RequestObjectMigration(FRemoteWorkPriority RequestPriority, FRemoteObjectId ObjectId, FRemoteServerId DestinationServerId)
 	{
-		check(DestinationServerId != GetGlobalServerId());
+		check(DestinationServerId != FRemoteServerId::GetLocalServerId());
 
 		// enqueue this request
 		FObjectMigrationRequests* ObjectRequests = nullptr;
@@ -704,7 +704,7 @@ public:
 #endif
 
 			// If we didn't expect it locally, send it on to the final destination
-			if (DestinationServerId != GetGlobalServerId())
+			if (DestinationServerId != FRemoteServerId::GetLocalServerId())
 			{
 				MigrateObjectToRemoteServer(ObjectId, DestinationServerId);
 			}
@@ -748,7 +748,7 @@ public:
 
 		UObject* Result = MigrateObjectFromObjectData(ObjectData, MigrationContext);
 		checkf(UE::CoreUObject::Private::FObjectHandleUtils::GetRemoteId(Result) == MigrationContext.ObjectId, TEXT("PeerId:%s requested object %s migration from PeerId:%s but received %s"),
-			   *GetGlobalServerId().ToString(), *MigrationContext.ObjectId.ToString(), *MigrationContext.RemoteServerId.ToString(), *UE::CoreUObject::Private::FObjectHandleUtils::GetRemoteId(Result).ToString());
+			   *FRemoteServerId::GetLocalServerId().ToString(), *MigrationContext.ObjectId.ToString(), *MigrationContext.RemoteServerId.ToString(), *UE::CoreUObject::Private::FObjectHandleUtils::GetRemoteId(Result).ToString());
 		checkf(MigrationContext.OwnerServerId == Handle::GetOwnerServerId(Result), TEXT("Expected owner id %s doesn't match post-deserialize owner: %s"), *MigrationContext.OwnerServerId.ToString(), *Handle::GetOwnerServerId(Result).ToString());
 
 		if (bFoundRequest)
@@ -805,20 +805,20 @@ UObject* MigrateObjectFromObjectData(FRemoteObjectData& ObjectData, const FUObje
 	// Notify that object data has been migrated from a remote server
 	OnObjectDataReceivedDelegate.Broadcast(ObjectData, MigrationContext);
 
-	FRemoteServerId LocalServerId = UE::RemoteObject::GetGlobalServerId();
+	FRemoteServerId LocalServerId = FRemoteServerId::GetLocalServerId();
 	TArray<FRemoteObjectId> ReceivedObjectRemoteIds;
 	TArray<UObject*> ReceivedObjects;
 	int32 RequestedObjectIndex = DeserializeObjectData(ObjectData, &MigrationContext, ReceivedObjectRemoteIds, ReceivedObjects);
 
 	// Debug section
-	checkf(ReceivedObjects.Num() > 0, TEXT("PeerId:%s unable to deserialize object data (%d bytes)"), *GetGlobalServerId().ToString(), ObjectData.GetNumBytes());
+	checkf(ReceivedObjects.Num() > 0, TEXT("PeerId:%s unable to deserialize object data (%d bytes)"), *FRemoteServerId::GetLocalServerId().ToString(), ObjectData.GetNumBytes());
 	if (ReceivedObjects.Num() > 0)
 	{
 		TRACE_CPUPROFILER_EVENT_SCOPE(CheckForNewOrRemovedSubObjects);
 
 		if (UE_LOG_ACTIVE(LogRemoteObject, VeryVerbose))
 		{
-			UE_LOG(LogRemoteObject, VeryVerbose, TEXT("PeerId:%s deserialized object data (%d bytes)"), *GetGlobalServerId().ToString(), ObjectData.GetNumBytes());
+			UE_LOG(LogRemoteObject, VeryVerbose, TEXT("PeerId:%s deserialized object data (%d bytes)"), *FRemoteServerId::GetLocalServerId().ToString(), ObjectData.GetNumBytes());
 			for(int32 SubObjectIndex = 0; SubObjectIndex < ReceivedObjects.Num(); SubObjectIndex++)
 			{
 				FRemoteObjectId SubObjectId = UE::CoreUObject::Private::FObjectHandleUtils::GetRemoteId(ReceivedObjects[SubObjectIndex]);
@@ -930,7 +930,7 @@ UObject* MigrateObjectFromObjectData(FRemoteObjectData& ObjectData, const FUObje
 void MigrateObjectFromRemoteServer(FRemoteObjectId ObjectId, FRemoteServerId CurrentOwnerServerId, UObject* DestinationOuter)
 {
 	UE_LOG(LogRemoteObject, Verbose, TEXT("Aborting transaction, server %s needs obj id %s from server %s"),
-		*UE::RemoteObject::GetGlobalServerId().ToString(),
+		*FRemoteServerId::GetLocalServerId().ToString(),
 		*ObjectId.ToString(),
 		*CurrentOwnerServerId.ToString());
 
@@ -1094,7 +1094,7 @@ void TransferObjectOwnershipToRemoteServer(UObject* Object, FRemoteServerId Dest
 	FMigrateSendParams SendRemoteObjectParams = { .MigrationContext = MigrationContext, .ObjectData = ObjectData };
 	GTransferQueue.SendRemoteObject(SendRemoteObjectParams);
 
-	checkf(SentObjects.Num() > 0, TEXT("PeerId:%s unable to migrate object %s to PeerId:%s"), *GetGlobalServerId().ToString(), *UE::CoreUObject::Private::FObjectHandleUtils::GetRemoteId(Object).ToString(), *DestinationServerId.ToString());
+	checkf(SentObjects.Num() > 0, TEXT("PeerId:%s unable to migrate object %s to PeerId:%s"), *FRemoteServerId::GetLocalServerId().ToString(), *UE::CoreUObject::Private::FObjectHandleUtils::GetRemoteId(Object).ToString(), *DestinationServerId.ToString());
 
 	{
 		// PostMigrate callbacks may be manipulating migrated objects and may try to migrate them back so prevent that from happening since
@@ -1250,7 +1250,7 @@ void StoreObjectToDatabase(UObject* Object)
 		return;
 	}
 
-	FRemoteServerId LocalServerId = GetGlobalServerId();
+	FRemoteServerId LocalServerId = FRemoteServerId::GetLocalServerId();
 	FRemoteServerId OwnerServerId = UE::RemoteObject::Handle::GetOwnerServerId(Object);
 	if (LocalServerId != OwnerServerId)
 	{
@@ -1300,7 +1300,7 @@ void StoreObjectToDatabase(UObject* Object)
 			StoreRemoteObjectDataDelegate.Execute(Params);
 		}
 
-		checkf(SentObjects.Num() > 0, TEXT("PeerId:%s unable to store object %s to DatabaseId:%s"), *GetGlobalServerId().ToString(), *FRemoteObjectId(Object).ToString(), *DatabaseId.ToString());
+		checkf(SentObjects.Num() > 0, TEXT("PeerId:%s unable to store object %s to DatabaseId:%s"), *FRemoteServerId::GetLocalServerId().ToString(), *FRemoteObjectId(Object).ToString(), *DatabaseId.ToString());
 
 		for (UObject* ReferencedObject : ReferencedObjects)
 		{
@@ -1316,7 +1316,7 @@ void StoreObjectToDatabase(UObject* Object)
 		
 		// Assert if we're about to enter an infinite loop.
 		checkf(!IsRemoteReference(Object) || OutermostUnreachableRemoteReference != Object,
-			TEXT("PeerId:%s failed to store %s to database"), *GetGlobalServerId().ToString(), *Object->GetPathName());
+			TEXT("PeerId:%s failed to store %s to database"), *FRemoteServerId::GetLocalServerId().ToString(), *Object->GetPathName());
 
 	} while (IsRemoteReference(Object));
 }
@@ -1329,9 +1329,9 @@ FRemoteObjectReference::FRemoteObjectReference(const FObjectPtr& Ptr)
 	ObjectId = Ptr.GetRemoteId();
 	if (ObjectId.IsValid())
 	{
-		if (Ptr.IsResolved())
+		if (!Ptr.IsRemote())
 		{
-			ServerId = UE::RemoteObject::GetGlobalServerId();
+			ServerId = FRemoteServerId::GetLocalServerId();
 		}
 		else if (UE::RemoteObject::Handle::FRemoteObjectStub* Stub = UE::RemoteObject::Private::FindRemoteObjectStub(ObjectId))
 		{
@@ -1371,7 +1371,7 @@ bool FRemoteObjectReference::NetSerialize(FArchive& Ar, class UPackageMap* Map,
 	using namespace UE::RemoteObject::Serialization;
 
 	// Serialize remote object id
-	Ar << *this;
+	Serialize(Ar);
 
 	// Serialize the pathname of the object if possible which will then be used on the receiving end to find assets loaded by both servers in memory	
 	FRemoteObjectPathName PathName;
@@ -1429,15 +1429,15 @@ void FRemoteObjectReference::NetDequantize(FRemoteObjectId InObjectId, FRemoteSe
 	}
 }
 
-FArchive& operator<<(FArchive& Ar, FRemoteObjectReference& Ref)
+bool FRemoteObjectReference::Serialize(FArchive& Ar)
 {
-	Ar << Ref.ObjectId;
-	Ar << Ref.ServerId;
+	Ar << ObjectId;
+	Ar << ServerId;
 
 	if (Ar.IsLoading())
 	{
-		UE::RemoteObject::Private::RegisterRemoteObjectId(Ref.ObjectId, Ref.ServerId);
+		UE::RemoteObject::Private::RegisterRemoteObjectId(ObjectId, ServerId);
 	}
 
-	return Ar;
+	return true;
 }
diff --git a/Engine/Source/Runtime/CoreUObject/Private/UObject/UObjectArray.cpp b/Engine/Source/Runtime/CoreUObject/Private/UObject/UObjectArray.cpp
--- a/Engine/Source/Runtime/CoreUObject/Private/UObject/UObjectArray.cpp
+++ b/Engine/Source/Runtime/CoreUObject/Private/UObject/UObjectArray.cpp
@@ -247,6 +247,10 @@ void FUObjectArray::AllocateUObjectIndex(UObjectBase* Object, EInternalObjectFla
 	{
 		RemoteId = FRemoteObjectId::Generate(Object, *Object->GetFName().ToString(), InitialFlags);
 	}
+	else
+	{
+		RemoteId = UE::RemoteObject::Private::FRemoteIdLocalizationHelper::GetLocalized(RemoteId);
+	}
 #endif
 
 	LockInternalArray();
@@ -312,7 +316,7 @@ void FUObjectArray::AllocateUObjectIndex(UObjectBase* Object, EInternalObjectFla
 #endif
 
 #if UE_WITH_REMOTE_OBJECT_HANDLE
-	ObjectItem->RemoteId = RemoteId;
+	ObjectItem->SetRemoteId(RemoteId);
 #endif // UE_WITH_REMOTE_OBJECT_HANDLE
 	Object->InternalIndex = Index;
 
diff --git a/Engine/Source/Runtime/CoreUObject/Private/UObject/UObjectHash.cpp b/Engine/Source/Runtime/CoreUObject/Private/UObject/UObjectHash.cpp
--- a/Engine/Source/Runtime/CoreUObject/Private/UObject/UObjectHash.cpp
+++ b/Engine/Source/Runtime/CoreUObject/Private/UObject/UObjectHash.cpp
@@ -11,6 +11,7 @@
 #include "UObject/Class.h"
 #include "UObject/GarbageCollectionGlobals.h"
 #include "UObject/Package.h"
+#include "UObject/RemoteObjectPrivate.h"
 #include "UObject/ObjectHandlePrivate.h"
 #include "UObject/UObjectIterator.h"
 #include "UObject/ObjectVisibility.h"
@@ -1274,6 +1275,7 @@ UObject* StaticFindObjectFastInternalThreadSafe(FUObjectHashTables& ThreadHash,
 {
 #if UE_WITH_REMOTE_OBJECT_HANDLE
 	const EInternalObjectFlags ExclusiveInternalFlags = DefaultInternalExclusionFlags | InExlusiveInternalFlags;
+	RemoteId = UE::RemoteObject::Private::FRemoteIdLocalizationHelper::GetLocalized(RemoteId);
 	uint32 Hash = GetTypeHash(RemoteId);
 
 	FHashTableLock HashLock(ThreadHash);
diff --git a/Engine/Source/Runtime/CoreUObject/Private/UObject/UObjectMigrationContext.cpp b/Engine/Source/Runtime/CoreUObject/Private/UObject/UObjectMigrationContext.cpp
--- a/Engine/Source/Runtime/CoreUObject/Private/UObject/UObjectMigrationContext.cpp
+++ b/Engine/Source/Runtime/CoreUObject/Private/UObject/UObjectMigrationContext.cpp
@@ -24,7 +24,7 @@ EObjectMigrationRecvType FUObjectMigrationContext::GetObjectMigrationRecvType(FR
 		return EObjectMigrationRecvType::Invalid;
 	}
 
-	const FRemoteServerId LocalServerId = UE::RemoteObject::GetGlobalServerId();
+	const FRemoteServerId LocalServerId = FRemoteServerId::GetLocalServerId();
 	const bool bAlreadyOwns = UE::RemoteObject::Handle::IsOwned(InObjectId);
 
 	const bool bAssignOwnership = (OwnerServerId == LocalServerId) && !bAlreadyOwns;
@@ -62,7 +62,7 @@ EObjectMigrationSendType FUObjectMigrationContext::GetObjectMigrationSendType(FR
 		return EObjectMigrationSendType::Invalid;
 	}
 
-	const FRemoteServerId LocalServerId = UE::RemoteObject::GetGlobalServerId();
+	const FRemoteServerId LocalServerId = FRemoteServerId::GetLocalServerId();
 	const bool bCurrentlyOwns = UE::RemoteObject::Handle::IsOwned(InObjectId);
 
 	const bool bOwnershipTransfer = (OwnerServerId == RemoteServerId) && bCurrentlyOwns;
diff --git a/Engine/Source/Runtime/CoreUObject/Private/UObject/WeakObjectPtr.cpp b/Engine/Source/Runtime/CoreUObject/Private/UObject/WeakObjectPtr.cpp
--- a/Engine/Source/Runtime/CoreUObject/Private/UObject/WeakObjectPtr.cpp
+++ b/Engine/Source/Runtime/CoreUObject/Private/UObject/WeakObjectPtr.cpp
@@ -232,7 +232,9 @@ FUObjectItem* FWeakObjectPtr::TryResolveRemoteObject() const
 	using namespace UE::RemoteObject::Handle;
 	using namespace UE::RemoteObject::Private;
 
-	UObject* ResolvedObject = StaticFindObjectFastInternal(ObjectRemoteId);
+	FRemoteObjectId RemoteId = ObjectRemoteId;
+
+	UObject* ResolvedObject = StaticFindObjectFastInternal(RemoteId);
 
 	if (ResolvedObject)
 	{		
@@ -246,17 +248,20 @@ FUObjectItem* FWeakObjectPtr::TryResolveRemoteObject() const
 			TouchResidentObject(ResolvedObject);
 		}
 	}
-	else if (FRemoteObjectStub* Stub = FindRemoteObjectStub(ObjectRemoteId))
+	else if (FRemoteServerId::IsGlobalServerIdInitialized())
 	{
-		ResolvedObject = ResolveObject(Stub, ERemoteReferenceType::Weak);
+		if (FRemoteObjectStub* Stub = FindRemoteObjectStub(RemoteId))
+		{
+			ResolvedObject = ResolveObject(Stub, ERemoteReferenceType::Weak);
+		}
 	}
 
 	if (ResolvedObject)
 	{
 		int32 NewIndex = GUObjectArray.ObjectToIndex(ResolvedObject);
 		FUObjectItem* ObjectItem = GUObjectArray.IndexToObject(NewIndex);
-		checkf(ObjectItem->GetRemoteId() == ObjectRemoteId, TEXT("Remote object %s has a different id (%s) than the weak object pointer that resolved it (%s)"),
-			*ResolvedObject->GetPathName(), *ObjectItem->GetRemoteId().ToString(), *ObjectRemoteId.ToString());
+		checkf(ObjectItem->GetRemoteId() == RemoteId, TEXT("Remote object %s has a different id (%s) than the weak object pointer that resolved it (%s)"),
+			*ResolvedObject->GetPathName(), *ObjectItem->GetRemoteId().ToString(), *RemoteId.ToString());
 		if (ObjectIndex != NewIndex)
 		{
 			// Currently we can't just rely on object id to resolve weak object pointers and we don't want to heep hitting the resolve path
@@ -269,7 +274,7 @@ FUObjectItem* FWeakObjectPtr::TryResolveRemoteObject() const
 		if (ObjectSerialNumber != ObjectItem->GetSerialNumber())
 		{
 			checkf(ObjectSerialNumber == 0, TEXT("Attempting to change existing and valid serial number %d to %d when resolving remote object %s (%s)"),
-				ObjectSerialNumber, ObjectItem->GetSerialNumber(), *ResolvedObject->GetPathName(), *ObjectRemoteId.ToString());
+				ObjectSerialNumber, ObjectItem->GetSerialNumber(), *ResolvedObject->GetPathName(), *RemoteId.ToString());
 			FPlatformAtomics::AtomicStore(&const_cast<FWeakObjectPtr*>(this)->ObjectSerialNumber, ObjectItem->GetSerialNumber());
 		}
 		return ObjectItem;
diff --git a/Engine/Source/Runtime/CoreUObject/Public/UObject/Class.h b/Engine/Source/Runtime/CoreUObject/Public/UObject/Class.h
--- a/Engine/Source/Runtime/CoreUObject/Public/UObject/Class.h
+++ b/Engine/Source/Runtime/CoreUObject/Public/UObject/Class.h
@@ -4582,6 +4582,16 @@ template<> struct TBaseStructure<FRemoteObjectId>
 	static COREUOBJECT_API UScriptStruct* Get();
 };
 
+template<>
+struct TStructOpsTypeTraits<FRemoteObjectId> : public TStructOpsTypeTraitsBase2<FRemoteObjectId>
+{
+	enum
+	{
+		WithNetSerializer = true,
+		WithSerializer = true,
+	};
+};
+
 struct FRemoteObjectTables;
 template<> struct TBaseStructure<FRemoteObjectTables>
 {
@@ -4617,6 +4627,16 @@ template<> struct TBaseStructure<FRemoteServerId>
 	static COREUOBJECT_API UScriptStruct* Get();
 };
 
+template<>
+struct TStructOpsTypeTraits<FRemoteServerId> : public TStructOpsTypeTraitsBase2<FRemoteServerId>
+{
+	enum
+	{
+		WithNetSerializer = true,
+		WithSerializer = true,
+	};
+};
+
 struct FRemoteObjectReference;
 template<> struct TBaseStructure<FRemoteObjectReference>
 {
diff --git a/Engine/Source/Runtime/CoreUObject/Public/UObject/RemoteExecutor.h b/Engine/Source/Runtime/CoreUObject/Public/UObject/RemoteExecutor.h
--- a/Engine/Source/Runtime/CoreUObject/Public/UObject/RemoteExecutor.h
+++ b/Engine/Source/Runtime/CoreUObject/Public/UObject/RemoteExecutor.h
@@ -66,7 +66,7 @@ struct FRemoteWorkPriority
 	static_assert(sizeof(FRemoteServerId) == 4);
 	static_assert(sizeof(FRemoteTransactionId) == 4);
 
-	FRemoteServerId GetRootServerId() const { return FRemoteServerId(static_cast<uint32>(PackedData >> 32)); }
+	FRemoteServerId GetRootServerId() const { return FRemoteServerId::FromIdNumber(static_cast<uint32>(PackedData >> 32)); }
 	uint32 GetWorkDepth() const { return (PackedData >> 24) & 0xFFu; }
 	FRemoteTransactionId GetRootWorkTransactionId() const { return FRemoteTransactionId(PackedData & 0xFFFFFFu); }
 
diff --git a/Engine/Source/Runtime/CoreUObject/Public/UObject/RemoteObject.h b/Engine/Source/Runtime/CoreUObject/Public/UObject/RemoteObject.h
--- a/Engine/Source/Runtime/CoreUObject/Public/UObject/RemoteObject.h
+++ b/Engine/Source/Runtime/CoreUObject/Public/UObject/RemoteObject.h
@@ -21,8 +21,11 @@ namespace UE::RemoteObject
 	/**
 	* Returns a unique id associated with this (server) process
 	*/
-	COREUOBJECT_API FRemoteServerId GetGlobalServerId();
-
+	UE_DEPRECATED(5.6, "Use FRemoteServerId::GetLocalServerId() instead.")
+	FORCEINLINE FRemoteServerId GetGlobalServerId()
+	{
+		return FRemoteServerId::GetLocalServerId();
+	}
 }
 
 namespace UE::RemoteObject::Handle
diff --git a/Engine/Source/Runtime/CoreUObject/Public/UObject/RemoteObjectTransfer.h b/Engine/Source/Runtime/CoreUObject/Public/UObject/RemoteObjectTransfer.h
--- a/Engine/Source/Runtime/CoreUObject/Public/UObject/RemoteObjectTransfer.h
+++ b/Engine/Source/Runtime/CoreUObject/Public/UObject/RemoteObjectTransfer.h
@@ -227,15 +227,10 @@ public:
 		return UE::RemoteObject::Handle::IsRemote(ObjectId);
 	}
 
-	bool Serialize(FArchive& Ar)
-	{
-		Ar << *this;
-		return true;
-	}
-
+	COREUOBJECT_API bool Serialize(FArchive& Ar);
 	COREUOBJECT_API bool NetSerialize(FArchive& Ar, class UPackageMap* Map, bool& bOutSuccess);
 
-	friend COREUOBJECT_API FArchive& operator<<(FArchive& Ar, FRemoteObjectReference& Ref);
+	friend FArchive& operator<<(FArchive& Ar, FRemoteObjectReference& Ref);
 
 private:
 	friend struct UE::Net::FRemoteObjectReferenceNetSerializer;
@@ -251,3 +246,9 @@ struct TStructOpsTypeTraits<FRemoteObjectReference> : public TStructOpsTypeTrait
 		WithSerializer = true,
 	};
 };
+
+inline FArchive& operator<<(FArchive& Ar, FRemoteObjectReference& Ref)
+{
+	Ref.Serialize(Ar);
+	return Ar;
+}
diff --git a/Engine/Source/Runtime/CoreUObject/Public/UObject/RemoteObjectTypes.h b/Engine/Source/Runtime/CoreUObject/Public/UObject/RemoteObjectTypes.h
--- a/Engine/Source/Runtime/CoreUObject/Public/UObject/RemoteObjectTypes.h
+++ b/Engine/Source/Runtime/CoreUObject/Public/UObject/RemoteObjectTypes.h
@@ -16,6 +16,10 @@ namespace UE::Net
 {
 	struct FRemoteObjectReferenceNetSerializer;
 }
+namespace UE::RemoteObject::Private
+{
+	struct FRemoteIdLocalizationHelper;
+}
 
 constexpr inline int32 REMOTE_OBJECT_SERVER_ID_BIT_SIZE = 10;
 constexpr inline int32 REMOTE_OBJECT_SERIAL_NUMBER_BIT_SIZE = 53;
@@ -31,6 +35,7 @@ enum class ERemoteServerIdConstants : uint32
 
 	Database = Max,
 	Asset = Max - 1,
+	Local = Max - 2,
 
 	// Add new reserved server IDs above this line in a descending order (so the next would be Something = Max - 2)
 
@@ -41,6 +46,8 @@ enum class ERemoteServerIdConstants : uint32
 struct FRemoteServerId
 {
 	friend struct FRemoteObjectId;
+	friend UE::RemoteObject::Private::FRemoteIdLocalizationHelper;
+	friend struct UE::Net::FRemoteObjectReferenceNetSerializer;
 
 	FRemoteServerId() = default;
 	FRemoteServerId(const FRemoteServerId&) = default;
@@ -50,37 +57,51 @@ struct FRemoteServerId
 		: Id((uint32)InId)
 	{
 	}
+
+	UE_DEPRECATED(5.6, "Use FRemoteServerId::FromIdNumber(uint32) instead.")
 	explicit FRemoteServerId(uint32 InId)
-		: Id(InId)
+		: Id(0)
 	{
-		checkf(Id < (uint32)ERemoteServerIdConstants::FirstReserved, TEXT("Remote server id can not be greater than %u, got: %u"), (uint32)ERemoteServerIdConstants::FirstReserved - 1, Id);
+		*this = FromIdNumber(InId);
 	}
+
+	UE_DEPRECATED(5.6, "Use FRemoteServerId::FromString(const FString&) instead.")
 	explicit COREUOBJECT_API FRemoteServerId(const FString& InText);
+
+	static COREUOBJECT_API FRemoteServerId FromString(const FString& InText);
 	COREUOBJECT_API FString ToString() const;
 
 	uint32 GetIdNumber() const
 	{
-		return Id;
+		return GetGlobalized().Id;
 	}
+
 	bool IsValid() const
 	{
-		return Id != 0;
+		return Id != (uint32)ERemoteServerIdConstants::Invalid;
+	}
+
+	inline int32 Compare(FRemoteServerId Other) const
+	{
+		VerifyComparableServerIds(Id, Other.Id);
+		return (int32)GlobalizeId(Id) - (int32)GlobalizeId(Other.Id);
 	}
+
 	bool operator == (FRemoteServerId Other) const
 	{
-		return Id == Other.Id;
+		return Compare(Other) == 0;
 	}
 	bool operator != (FRemoteServerId Other) const
 	{
-		return Id != Other.Id;
+		return Compare(Other) != 0;
 	}
 	bool operator < (FRemoteServerId Other) const
 	{
-		return Id < Other.Id;
+		return Compare(Other) < 0;
 	}
 	bool operator <= (FRemoteServerId Other) const
 	{
-		return Id <= Other.Id;
+		return Compare(Other) <= 0;
 	}
 	bool IsAsset() const
 	{
@@ -90,11 +111,73 @@ struct FRemoteServerId
 	{
 		return (Id == (uint32)ERemoteServerIdConstants::Database);
 	}
+	bool IsLocal() const
+	{
+		return (Id == (uint32)ERemoteServerIdConstants::Local) || (*this == GlobalServerId);
+	}
 
+	static FORCEINLINE FRemoteServerId GetLocalServerId()
+	{
+		return FRemoteServerId(ERemoteServerIdConstants::Local);
+	}
+	
 	friend COREUOBJECT_API FArchive& operator<<(FArchive& Ar, FRemoteServerId& Id);
 
+	static COREUOBJECT_API void InitGlobalServerId(FRemoteServerId Id);
+
+	static COREUOBJECT_API bool IsGlobalServerIdInitialized();
+
+	static COREUOBJECT_API FRemoteServerId FromIdNumber(uint32 InNumber);
+
+	COREUOBJECT_API bool Serialize(FArchive& Ar);
+	COREUOBJECT_API bool NetSerialize(FArchive& Ar, class UPackageMap* Map, bool& bOutSuccess);
+
+	FORCEINLINE uint32 GetTypeHash() const
+	{
+		return ::GetTypeHash(GetLocalized().Id);
+	}
+
+	friend void AssignGlobalServerIdDebuggingState();
+
 private:
 
+	inline FRemoteServerId GetLocalized() const
+	{
+		return (Id == GlobalServerId.Id && GlobalServerId.IsValid()) ? FRemoteServerId(ERemoteServerIdConstants::Local) : *this;
+	}
+
+	inline FRemoteServerId GetGlobalized() const
+	{
+		if (Id == (uint32)ERemoteServerIdConstants::Local)
+		{
+			checkf(GlobalServerId.Id != (uint32)ERemoteServerIdConstants::Local, TEXT("Unable to convert Local server id to global. Global server id hasn't been initialized yet."));
+			return GlobalServerId;
+		}
+		return *this;
+	}
+
+	inline static void VerifyComparableServerIds(uint32 Id0, uint32 Id1)
+	{
+#if DO_CHECK
+		if (!IsGlobalServerIdInitialized())
+		{
+			checkf(Id0 != (uint32)ERemoteServerIdConstants::Local || Id1 == (uint32)ERemoteServerIdConstants::Invalid || Id1 >= (uint32)ERemoteServerIdConstants::FirstReserved,
+				TEXT("When the global server id is not initialized it's only possible to compare predefined server IDs (lhs: %u, rhs: %u)"), Id0, Id1);
+			checkf(Id1 != (uint32)ERemoteServerIdConstants::Local || Id0 == (uint32)ERemoteServerIdConstants::Invalid || Id0 >= (uint32)ERemoteServerIdConstants::FirstReserved,
+				TEXT("When the global server id is not initialized it's only possible to compare predefined server IDs (lhs: %u, rhs: Local)"), Id0, Id1);
+		}
+#endif
+	}
+
+	inline static uint32 GlobalizeId(uint32 InId)
+	{
+		if (InId == (uint32)ERemoteServerIdConstants::Local)
+		{
+			return GlobalServerId.Id;
+		}
+		return InId;
+	}
+
 	/** 
 	 * Constructor used exclusively by FRemoteObjectId::GetServerId() to bypass range changes in the other constructors 
 	 */
@@ -104,6 +187,8 @@ private:
 	}
 
 	uint32 Id = (uint32)ERemoteServerIdConstants::Invalid;
+
+	COREUOBJECT_API static FRemoteServerId GlobalServerId;
 };
 
 struct FRemoteObjectId
@@ -120,9 +205,25 @@ private:
 		uint64 Id = 0;
 	};
 
-	FORCEINLINE uint64 GetComparableValue() const
+	inline int64 GetComparableValue() const
 	{
-		return uint64(SerialNumber) | (uint64(ServerId) << REMOTE_OBJECT_SERIAL_NUMBER_BIT_SIZE);
+		return int64(SerialNumber) | (int64(FRemoteServerId::GlobalizeId(ServerId)) << REMOTE_OBJECT_SERIAL_NUMBER_BIT_SIZE);
+	}
+
+	inline int64 Compare(FRemoteObjectId Other) const
+	{
+		FRemoteServerId::VerifyComparableServerIds(ServerId, Other.ServerId);
+		return GetComparableValue() - Other.GetComparableValue();
+	}
+
+	inline FRemoteObjectId GetLocalized() const
+	{
+		return FRemoteObjectId(GetServerId().GetLocalized(), SerialNumber);
+	}
+
+	inline FRemoteObjectId GetGlobalized() const
+	{
+		return FRemoteObjectId(GetServerId().GetGlobalized(), SerialNumber);
 	}
 
 public:
@@ -155,34 +256,34 @@ public:
 
 	FORCEINLINE uint32 GetTypeHash() const
 	{
-		return ::GetTypeHash(GetComparableValue());
+		return ::GetTypeHash(GetLocalized().Id);
 	}
 
 	FORCEINLINE bool operator==(const FRemoteObjectId& Other) const
 	{
 		// Don't compare against the Reserved bit
-		return GetComparableValue() == Other.GetComparableValue();
+		return Compare(Other) == 0;
 	}
 	FORCEINLINE bool operator!=(const FRemoteObjectId& Other) const
 	{
 		// Don't compare against the Reserved bit
-		return GetComparableValue() != Other.GetComparableValue();
+		return Compare(Other) != 0;
 	}
 	FORCEINLINE bool operator<(const FRemoteObjectId& Other) const
 	{
-		return GetComparableValue() < Other.GetComparableValue();
+		return Compare(Other) < 0;
 	}
 	FORCEINLINE bool operator<=(const FRemoteObjectId& Other) const
 	{
-		return GetComparableValue() <= Other.GetComparableValue();
+		return Compare(Other) <= 0;
 	}
 	FORCEINLINE bool operator>(const FRemoteObjectId& Other) const
 	{
-		return GetComparableValue() > Other.GetComparableValue();
+		return Compare(Other) > 0;
 	}
 	FORCEINLINE bool operator>=(const FRemoteObjectId& Other) const
 	{
-		return GetComparableValue() >= Other.GetComparableValue();
+		return Compare(Other) >= 0;
 	}
 	FORCEINLINE bool IsValid() const
 	{
@@ -191,7 +292,7 @@ public:
 
 	FORCEINLINE uint64 GetIdNumber() const
 	{
-		return Id;
+		return GetGlobalized().Id;
 	}
 	
 	FORCEINLINE FRemoteServerId GetServerId() const
@@ -204,12 +305,23 @@ public:
 		return GetServerId().IsAsset();
 	}
 
+	FORCEINLINE bool IsLocal() const
+	{
+		return GetServerId().IsLocal();
+	}
+
 	COREUOBJECT_API FString ToString() const;
 
 	COREUOBJECT_API static FRemoteObjectId Generate(UObjectBase* InObject, const TCHAR* InName, EInternalObjectFlags InInitialFlags = EInternalObjectFlags::None);
 
+	COREUOBJECT_API bool Serialize(FArchive& Ar);
+	COREUOBJECT_API bool NetSerialize(FArchive& Ar, class UPackageMap* Map, bool& bOutSuccess);
+
 	friend COREUOBJECT_API FArchive& operator<<(FArchive& Ar, FRemoteObjectId& Id);
 	friend struct UE::Net::FRemoteObjectReferenceNetSerializer;
+	friend UE::RemoteObject::Private::FRemoteIdLocalizationHelper;
+	friend struct FUObjectItem;
+	friend class FNetGUIDCache;	
 };
 
 FORCEINLINE uint32 GetTypeHash(const FRemoteObjectId& ObjectId)
@@ -219,5 +331,5 @@ FORCEINLINE uint32 GetTypeHash(const FRemoteObjectId& ObjectId)
 
 FORCEINLINE uint32 GetTypeHash(const FRemoteServerId& ServerId)
 {
-	return ServerId.GetIdNumber();
+	return ServerId.GetTypeHash();
 }
diff --git a/Engine/Source/Runtime/CoreUObject/Public/UObject/UObjectArray.h b/Engine/Source/Runtime/CoreUObject/Public/UObject/UObjectArray.h
--- a/Engine/Source/Runtime/CoreUObject/Public/UObject/UObjectArray.h
+++ b/Engine/Source/Runtime/CoreUObject/Public/UObject/UObjectArray.h
@@ -102,8 +102,10 @@ public:
 	int32 ClusterRootIndex;
 
 #if UE_WITH_REMOTE_OBJECT_HANDLE
+private:
 	// Globally unique id of this object
 	FRemoteObjectId RemoteId;
+public:
 #endif
 
 #if STATS || ENABLE_STATNAMEDEVENTS_UOBJECT
@@ -203,6 +205,10 @@ public:
 	}
 
 #if UE_WITH_REMOTE_OBJECT_HANDLE
+	FORCEINLINE void SetRemoteId(FRemoteObjectId Id)
+	{
+		RemoteId = Id.GetLocalized();
+	}
 	FORCEINLINE FRemoteObjectId GetRemoteId() const
 	{
 		return RemoteId;
diff --git a/Engine/Source/Runtime/Engine/Private/Actor.cpp b/Engine/Source/Runtime/Engine/Private/Actor.cpp
--- a/Engine/Source/Runtime/Engine/Private/Actor.cpp
+++ b/Engine/Source/Runtime/Engine/Private/Actor.cpp
@@ -1293,7 +1293,7 @@ void AActor::PostMigrate(const FUObjectMigrationContext& MigrationContext)
 	Super::PostMigrate(MigrationContext);
 
 	// We are using this as a proxy for having ownership until we sort out those mechanics
-	const FRemoteServerId LocalServerId = UE::RemoteObject::GetGlobalServerId();
+	const FRemoteServerId LocalServerId = FRemoteServerId::GetLocalServerId();
 	const bool bHasOwnership = UE::RemoteObject::Handle::IsOwned(this);
 
 	UWorld* const World = GetWorld();
diff --git a/Engine/Source/Runtime/Engine/Private/PackageMapClient.cpp b/Engine/Source/Runtime/Engine/Private/PackageMapClient.cpp
--- a/Engine/Source/Runtime/Engine/Private/PackageMapClient.cpp
+++ b/Engine/Source/Runtime/Engine/Private/PackageMapClient.cpp
@@ -2945,7 +2945,7 @@ FRemoteObjectId FNetGUIDCache::GetRemoteObjectIdFromNetworkGUID(const FNetworkGU
 {
 	// The NetGUID ObjectId is a 64-bit value with the first bit as the static flag. Since this value is ignored by
 	// FRemoteObjectId (it's a reserved bit) the value is masked away.
-	return FRemoteObjectId::CreateFromInt(NetGuid.ObjectId & ~1);
+	return FRemoteObjectId::CreateFromInt(NetGuid.ObjectId & ~1).GetLocalized();
 }
 #endif
 
diff --git a/Engine/Source/Runtime/Engine/Private/PlayerController.cpp b/Engine/Source/Runtime/Engine/Private/PlayerController.cpp
--- a/Engine/Source/Runtime/Engine/Private/PlayerController.cpp
+++ b/Engine/Source/Runtime/Engine/Private/PlayerController.cpp
@@ -5036,7 +5036,7 @@ void APlayerController::PostMigrate(const FUObjectMigrationContext& MigrationCon
 		EObjectMigrationRecvType MigrationRecvType = MigrationContext.GetObjectMigrationRecvType(this);
 
 		UE_LOG(LogPlayerController, Verbose, TEXT("[%s] APlayerController::PostMigrate [%s]: %s bHasOwnership=%d MigrationSide=Receive RecvType=%s PlayerId=%s"), 
-			   *UE::RemoteObject::GetGlobalServerId().ToString(), *FRemoteObjectId(this).ToString(), *GetName(), bHasOwnership, ToString(MigrationRecvType), *CachedConnectionPlayerId.ToDebugString());
+			   *FRemoteServerId::GetLocalServerId().ToString(), *FRemoteObjectId(this).ToString(), *GetName(), bHasOwnership, ToString(MigrationRecvType), *CachedConnectionPlayerId.ToDebugString());
 
 		ensure(NetConnection == nullptr);
 
@@ -5061,7 +5061,7 @@ void APlayerController::PostMigrate(const FUObjectMigrationContext& MigrationCon
 				ensure(PrevHandshakeTransactionId == GetClientHandshakeId());
 
 				UE_LOG(LogPlayerController, Verbose, TEXT("[%s] APlayerController::PostMigrate [%s]: Binding %s to connection %s"), 
-					   *UE::RemoteObject::GetGlobalServerId().ToString(), *FRemoteObjectId(this).ToString(), 
+					   *FRemoteServerId::GetLocalServerId().ToString(), *FRemoteObjectId(this).ToString(), 
 					   *GetName(), *NetConnection->GetName());
 			}
 		}
@@ -5076,7 +5076,7 @@ void APlayerController::PostMigrate(const FUObjectMigrationContext& MigrationCon
 		EObjectMigrationSendType MigrationSendType = MigrationContext.GetObjectMigrationSendType(this);
 
 		UE_LOG(LogPlayerController, Verbose, TEXT("[%s] APlayerController::PostMigrate [%s]: %s bHasOwnership=%d MigrationSide=Send SendType=%s PlayerId=%s"), 
-			   *UE::RemoteObject::GetGlobalServerId().ToString(), *FRemoteObjectId(this).ToString(), *GetName(), bHasOwnership, ToString(MigrationSendType), *CachedConnectionPlayerId.ToDebugString());
+			   *FRemoteServerId::GetLocalServerId().ToString(), *FRemoteObjectId(this).ToString(), *GetName(), bHasOwnership, ToString(MigrationSendType), *CachedConnectionPlayerId.ToDebugString());
 
 		const bool bDetachControllerFromConnection = true;
 		const bool bDeregisterFromWorld = (MigrationSendType == EObjectMigrationSendType::ReassignOwnership);
@@ -5128,7 +5128,7 @@ void APlayerController::PostMigrate(const FUObjectMigrationContext& MigrationCon
 					PrevNetConnection->OwningActor = PrevNetConnection->PlayerController;
 					
 					UE_LOG(LogPlayerController, Verbose, TEXT("[%s] APlayerController::PostMigrate [%s]: Replacing %s with new player controller %s on connection %s"), 
-						   *UE::RemoteObject::GetGlobalServerId().ToString(), *FRemoteObjectId(this).ToString(), 
+						   *FRemoteServerId::GetLocalServerId().ToString(), *FRemoteObjectId(this).ToString(), 
 						   *GetName(), *PrevNetConnection->PlayerController->GetName(), *PrevNetConnection->GetName());
 
 					GetWorld()->AddController(PrevNetConnection->PlayerController);
@@ -5136,7 +5136,7 @@ void APlayerController::PostMigrate(const FUObjectMigrationContext& MigrationCon
 				else
 				{
 					UE_LOG(LogPlayerController, Verbose, TEXT("[%s] APlayerController::PostMigrate [%s]: Replacing %s with existing player controller %s on connection %s"),
-						   *UE::RemoteObject::GetGlobalServerId().ToString(), *FRemoteObjectId(this).ToString(), 
+						   *FRemoteServerId::GetLocalServerId().ToString(), *FRemoteObjectId(this).ToString(), 
 						   *GetName(), *PrevNetConnection->PlayerController->GetName(), *PrevNetConnection->GetName());
 				}
 
@@ -5148,7 +5148,7 @@ void APlayerController::PostMigrate(const FUObjectMigrationContext& MigrationCon
 		if (bDeregisterFromWorld)
 		{
 			UE_LOG(LogPlayerController, Verbose, TEXT("[%s] APlayerController::PostMigrate [%s]: Removing %s from the world"), 
-				   *UE::RemoteObject::GetGlobalServerId().ToString(), *FRemoteObjectId(this).ToString(), *GetName());
+				   *FRemoteServerId::GetLocalServerId().ToString(), *FRemoteObjectId(this).ToString(), *GetName());
 
 			GetWorld()->RemoveController(this);
 		}
diff --git a/Engine/Source/Runtime/Engine/Private/UnrealEngine.cpp b/Engine/Source/Runtime/Engine/Private/UnrealEngine.cpp
--- a/Engine/Source/Runtime/Engine/Private/UnrealEngine.cpp
+++ b/Engine/Source/Runtime/Engine/Private/UnrealEngine.cpp
@@ -9358,7 +9358,7 @@ bool UEngine::HandleObjCommand( const TCHAR* Cmd, FOutputDevice& Ar )
 			FString ServerIdText;
 			if (FParse::Value(Cmd, TEXT("SERVERID="), ServerIdText))
 			{
-				ServerId = FRemoteServerId(ServerIdText);
+				ServerId = FRemoteServerId::FromString(ServerIdText);
 			}
 		}
 #endif
diff --git a/Engine/Source/Runtime/Engine/Private/World.cpp b/Engine/Source/Runtime/Engine/Private/World.cpp
--- a/Engine/Source/Runtime/Engine/Private/World.cpp
+++ b/Engine/Source/Runtime/Engine/Private/World.cpp
@@ -6617,7 +6617,7 @@ void UWorld::CreatePhysicsScene(const AWorldSettings* Settings)
  	PhysicsSceneNameBuilder.Appendf(TEXT("%s | World Type [%s]"), IsNetMode(NM_DedicatedServer) ? TEXT("Server Physics") : TEXT("Client Physics"), LexToString(WorldType));
  
 	#if UE_WITH_REMOTE_OBJECT_HANDLE
- 		PhysicsSceneNameBuilder.Appendf(TEXT(" | Remote Server ID [%u]"), UE::RemoteObject::GetGlobalServerId().GetIdNumber()); 
+ 		PhysicsSceneNameBuilder.Appendf(TEXT(" | Remote Server ID [%s]"), *FRemoteServerId::GetLocalServerId().ToString());
 	#endif
 
 	const FName PhysicsName = FName(PhysicsSceneNameBuilder);
diff --git a/Engine/Source/Runtime/Experimental/Iris/Core/Private/Iris/Serialization/RemoteObjectReferenceNetSerializer.cpp b/Engine/Source/Runtime/Experimental/Iris/Core/Private/Iris/Serialization/RemoteObjectReferenceNetSerializer.cpp
--- a/Engine/Source/Runtime/Experimental/Iris/Core/Private/Iris/Serialization/RemoteObjectReferenceNetSerializer.cpp
+++ b/Engine/Source/Runtime/Experimental/Iris/Core/Private/Iris/Serialization/RemoteObjectReferenceNetSerializer.cpp
@@ -179,7 +179,7 @@ void FRemoteObjectReferenceNetSerializer::Dequantize(FNetSerializationContext& C
 
 	RemoteObjectPathNameNetSerializer->Dequantize(Context, InternalArgs);
 
-	TargetValue.NetDequantize(ObjectId, FRemoteServerId(QuantizedValue.ServerId), PathName);
+	TargetValue.NetDequantize(ObjectId.GetLocalized(), FRemoteServerId::FromIdNumber(QuantizedValue.ServerId).GetLocalized(), PathName);
 }
 
 bool FRemoteObjectReferenceNetSerializer::IsEqual(FNetSerializationContext& Context, const FNetIsEqualArgs& Args)
