add TickAbortedRequest to FRemoteSubsystemBase
- this allows us to run logic immediately after a transaction aborts, so that we can more quickly send requests for dependencies
- refactor of the object migration and remote physics query subsystems to take advantage of this, reducing the latency between the transaction aborting and it receiving its known dependencies and retrying

diff --git a/Engine/Source/Runtime/CoreUObject/Private/UObject/RemoteExecutor.cpp b/Engine/Source/Runtime/CoreUObject/Private/UObject/RemoteExecutor.cpp
--- a/Engine/Source/Runtime/CoreUObject/Private/UObject/RemoteExecutor.cpp
+++ b/Engine/Source/Runtime/CoreUObject/Private/UObject/RemoteExecutor.cpp
@@ -538,6 +538,15 @@ public:
 								*ExecutingWork->RequestId.ToString(),
 								*AbortReasonDescription);
 						}
+
+						if (!bWorkComplete)
+						{
+							for (int32 SubsystemIndex = 0; SubsystemIndex < Subsystems.Num(); SubsystemIndex++)
+							{
+								FRemoteSubsystemBase* Subsystem = Subsystems[SubsystemIndex];
+								Subsystem->TickAbortedRequest();
+							}
+						}
 					}
 					else
 					{
diff --git a/Engine/Source/Runtime/CoreUObject/Private/UObject/RemoteObjectTransfer.cpp b/Engine/Source/Runtime/CoreUObject/Private/UObject/RemoteObjectTransfer.cpp
--- a/Engine/Source/Runtime/CoreUObject/Private/UObject/RemoteObjectTransfer.cpp
+++ b/Engine/Source/Runtime/CoreUObject/Private/UObject/RemoteObjectTransfer.cpp
@@ -57,6 +57,10 @@ public:
 	TMap<FRemoteObjectId, FRemoteObjectId> RequiredObjectsCanonicalRootByObjectID;
 	TSet<FRemoteObjectId> RequiredObjectsCanonicalRoots;
 
+	// when a transaction aborts and requires an object, this gets
+	// filled in with the object id
+	FRemoteObjectId NewRequiredObject;
+
 	// this set gets cleared before running the 
 	// transaction and only tracks the objects touched 
 	// during the most recent run
@@ -314,33 +318,16 @@ public:
 		}
 	}
 
-	void TickRequest() final
+	void SendRequests()
 	{
 		using namespace UE::RemoteObject::Private;
 		using namespace UE::RemoteObject::Handle;
 		using namespace UE::RemoteObject::Serialization;
 
-		check(ActiveRequest->RequiredObjects.Num() == ActiveRequest->RequiredObjectsCanonicalRootByObjectID.Num());
-
-		// reset the list of used objects each time before we run
-		ActiveRequest->UsedObjects.Reset();
-
-		UE_LOG(LogRemoteObject, VeryVerbose, TEXT("FRemoteObjectTransferQueue: TickRequest(%s) %d required objs (%d pending requests):"), *ActiveRequest->GetRequestId().ToString(), ActiveRequest->RequiredObjects.Num(), PendingObjectRequests.Num());
-
-		if (UE_GET_LOG_VERBOSITY(LogRemoteObject) == ELogVerbosity::VeryVerbose)
-		{
-			for (const FRemoteObjectId& RequiredObjectId : ActiveRequest->RequiredObjects)
-			{
-				FRemoteObjectId CanonicalObjectId = ActiveRequest->RequiredObjectsCanonicalRootByObjectID.FindChecked(RequiredObjectId);
-				UE_LOG(LogRemoteObject, VeryVerbose, TEXT("FRemoteObjectTransferQueue: TickRequest(%s) %s [root %s] %s"), *ActiveRequest->GetRequestId().ToString(), *RequiredObjectId.ToString(), *CanonicalObjectId.ToString(), Handle::IsRemote(RequiredObjectId) ? TEXT("<remote>") : TEXT("<local>"));
-			}
-		}
-
-		// send the requests
 		for (const FRemoteObjectId& RequiredObjectId : ActiveRequest->RequiredObjects)
 		{
-			FRemoteObjectId RequiredCanonicalRootObjectId = ActiveRequest->RequiredObjectsCanonicalRootByObjectID.FindChecked(RequiredObjectId);
-			FRemoteObjectId ObjectIdToRequest = RequiredCanonicalRootObjectId.IsValid() ? RequiredCanonicalRootObjectId : RequiredObjectId;
+			FRemoteObjectId* RequiredCanonicalRootObjectId = ActiveRequest->RequiredObjectsCanonicalRootByObjectID.Find(RequiredObjectId);
+			FRemoteObjectId ObjectIdToRequest = RequiredCanonicalRootObjectId ? *RequiredCanonicalRootObjectId : RequiredObjectId;
 
 			if (Handle::IsRemote(ObjectIdToRequest))
 			{
@@ -411,18 +398,58 @@ public:
 					}
 				}
 			}
-			else if (!RequiredCanonicalRootObjectId.IsValid())
+			else if (!RequiredCanonicalRootObjectId)
 			{
 				if (UObject* Object = StaticFindObjectFastInternal(RequiredObjectId))
 				{
 					UObject* const RootObject = FindCanonicalRootObjectForSerialization(Object);
 					const FRemoteObjectId CanonicalRootObjectId = UE::CoreUObject::Private::FObjectHandleUtils::GetRemoteId(RootObject);
-					ActiveRequest->RequiredObjectsCanonicalRootByObjectID.FindChecked(RequiredObjectId) = CanonicalRootObjectId;
+					ActiveRequest->RequiredObjectsCanonicalRootByObjectID.Add(RequiredObjectId, CanonicalRootObjectId);
+					ActiveRequest->RequiredObjectsCanonicalRoots.Add(CanonicalRootObjectId);
 				}
 			}
 		}
 	}
 
+	void TickRequest() final
+	{
+		using namespace UE::RemoteObject::Private;
+		using namespace UE::RemoteObject::Handle;
+		using namespace UE::RemoteObject::Serialization;
+
+		// reset the list of used objects each time before we run
+		ActiveRequest->UsedObjects.Reset();
+		ActiveRequest->NewRequiredObject = FRemoteObjectId();
+
+		UE_LOG(LogRemoteObject, VeryVerbose, TEXT("FRemoteObjectTransferQueue: TickRequest(%s) %d required objs (%d pending requests):"), *ActiveRequest->GetRequestId().ToString(), ActiveRequest->RequiredObjects.Num(), PendingObjectRequests.Num());
+
+		if (UE_GET_LOG_VERBOSITY(LogRemoteObject) == ELogVerbosity::VeryVerbose)
+		{
+			for (const FRemoteObjectId& RequiredObjectId : ActiveRequest->RequiredObjects)
+			{
+				FRemoteObjectId* CanonicalObjectId = ActiveRequest->RequiredObjectsCanonicalRootByObjectID.Find(RequiredObjectId);
+				UE_LOG(LogRemoteObject, VeryVerbose, TEXT("FRemoteObjectTransferQueue: TickRequest(%s) %s [root %s] %s"), *ActiveRequest->GetRequestId().ToString(), *RequiredObjectId.ToString(), CanonicalObjectId ? *(CanonicalObjectId->ToString()) : TEXT("<unk>"), Handle::IsRemote(RequiredObjectId) ? TEXT("<remote>") : TEXT("<local>"));
+			}
+		}
+
+		SendRequests();
+	}
+
+	void TickAbortedRequest() final
+	{
+		using namespace UE::RemoteObject::Private;
+		using namespace UE::RemoteObject::Handle;
+		using namespace UE::RemoteObject::Serialization;
+
+		if (ActiveRequest->NewRequiredObject.IsValid())
+		{
+			// add it to the active request's list of required objects
+			ActiveRequest->RequiredObjects.Add(ActiveRequest->NewRequiredObject);
+
+			SendRequests();
+		}
+	}
+
 	bool AreDependenciesSatisfied() const final
 	{
 		for (const FRemoteObjectId& RequiredObjectId : ActiveRequest->RequiredObjects)
@@ -949,25 +976,29 @@ void MigrateObjectFromRemoteServer(FRemoteObjectId ObjectId, FRemoteServerId Cur
 		}
 
 		// add it to the active request's list of required objects
-		if (!GTransferQueue.ActiveRequest->RequiredObjects.Contains(ObjectId))
+		bool bAlreadyAdded = false;
+		GTransferQueue.ActiveRequest->RequiredObjects.Add(ObjectId, &bAlreadyAdded);
+
+		if (!bAlreadyAdded)
 		{
 			// does any other existing request also require this object? if so, try to
 			// grab the cached canonical id from them
-			FRemoteObjectId CanonicalRootObjectId;
-
 			for (int32 RequestIndex = 0; RequestIndex < GTransferQueue.GetRequestCount(); RequestIndex++)
 			{
 				FRemoteObjectRequest* ExistingRequest = GTransferQueue.GetRequestByIndex(RequestIndex);
 
 				if (FRemoteObjectId* FoundRequiredRemoteObjectID = ExistingRequest->RequiredObjects.Find(ObjectId))
 				{
-					CanonicalRootObjectId = ExistingRequest->RequiredObjectsCanonicalRootByObjectID.FindChecked(*FoundRequiredRemoteObjectID);
+					FRemoteObjectId* CanonicalRootObjectId = ExistingRequest->RequiredObjectsCanonicalRootByObjectID.Find(*FoundRequiredRemoteObjectID);
+
+					if (CanonicalRootObjectId)
+					{
+						GTransferQueue.ActiveRequest->RequiredObjectsCanonicalRootByObjectID.Add(ObjectId, *CanonicalRootObjectId);
+						GTransferQueue.ActiveRequest->RequiredObjectsCanonicalRoots.Add(*CanonicalRootObjectId);
+						break;
+					}
 				}
 			}
-
-			GTransferQueue.ActiveRequest->RequiredObjects.Add(ObjectId);
-			GTransferQueue.ActiveRequest->RequiredObjectsCanonicalRootByObjectID.Add(ObjectId, CanonicalRootObjectId);
-			GTransferQueue.ActiveRequest->RequiredObjectsCanonicalRoots.Add(CanonicalRootObjectId);
 		}
 		
 		if (UE_LOG_ACTIVE(LogRemoteObject, VeryVerbose))
@@ -1010,25 +1041,15 @@ void TouchResidentObject(UObject* Object)
 		{
 			const FRemoteObjectId ObjectId = UE::CoreUObject::Private::FObjectHandleUtils::GetRemoteId(Object);
 
-			if (!GTransferQueue.ActiveRequest->UsedObjects.Contains(ObjectId))
-			{
-				GTransferQueue.ActiveRequest->UsedObjects.Add(ObjectId);
+			bool bAlreadyInUsedObjectsList = false;
+			GTransferQueue.ActiveRequest->UsedObjects.Add(ObjectId, &bAlreadyInUsedObjectsList);
 
+			if (!bAlreadyInUsedObjectsList)
+			{
 				OnObjectTouchedDelegate.Broadcast(GTransferQueue.ActiveRequest->RequestId, ObjectId);
 			}
 
-			// add it to the active request's list of required objects
-			if (!GTransferQueue.ActiveRequest->RequiredObjects.Contains(ObjectId))
-			{
-				GTransferQueue.ActiveRequest->RequiredObjects.Add(ObjectId);
-				
-				// NOTE - even though this object is resident, we can't fetch the canonical root
-				// object because it could trigger endless recursion into this function which is
-				// called when object pointers get resolved
-				FRemoteObjectId CanonicalRootObjectId;
-				GTransferQueue.ActiveRequest->RequiredObjectsCanonicalRootByObjectID.Add(ObjectId, CanonicalRootObjectId);
-				GTransferQueue.ActiveRequest->RequiredObjectsCanonicalRoots.Add(CanonicalRootObjectId);
-			}
+			GTransferQueue.ActiveRequest->NewRequiredObject = ObjectId;
 		};
 	}
 }
@@ -1172,40 +1193,6 @@ void OnObjectDataDenied(FRemoteObjectId ObjectId, FRemoteServerId RemoteServerId
 	GTransferQueue.DenyReceiveRequest(ObjectId);
 }
 
-void PrefetchRemoteObject(FRemoteObjectId ObjectId, UObject* DestinationOuter)
-{
-	if (!AutoRTFM::IsClosed())
-	{
-		// if we are outside of a transaction (dev/test workflow)
-		// then prefetch is a no-op
-		return;
-	}
-
-	bool bAdded = false;
-	
-	UE_AUTORTFM_OPEN
-	{
-		// check bInMultiServerCommit here to ensure the UsedObjects list doesn't 
-		// accidentally mutate during the actual multi-server commit process
-		if (!GTransferQueue.bInMultiServerCommit && !GTransferQueue.ActiveRequest->UsedObjects.Contains(ObjectId))
-		{
-			GTransferQueue.ActiveRequest->UsedObjects.Add(ObjectId);
-
-			OnObjectTouchedDelegate.Broadcast(GTransferQueue.ActiveRequest->RequestId, ObjectId);
-		}
-
-		// it's possible a previous prefetch or strong fetch already added this to the RequiredObjects list
-		bAdded = !GTransferQueue.ActiveRequest->RequiredObjects.Contains(ObjectId);
-
-		if (bAdded)
-		{
-			GTransferQueue.ActiveRequest->RequiredObjects.Add(ObjectId);
-			GTransferQueue.ActiveRequest->RequiredObjectsCanonicalRoots.Add(ObjectId);
-			UE_LOG(LogRemoteObject, VeryVerbose, TEXT("PrefetchRemoteObject(%s) : obj %s"), *GTransferQueue.ActiveRequest->GetRequestId().ToString(), *ObjectId.ToString());
-		}
-	};
-}
-
 void InitRemoteObjectTransfer()
 {
 	UE::RemoteExecutor::RegisterRemoteSubsystem(&GTransferQueue);
diff --git a/Engine/Source/Runtime/CoreUObject/Public/UObject/RemoteExecutor.h b/Engine/Source/Runtime/CoreUObject/Public/UObject/RemoteExecutor.h
--- a/Engine/Source/Runtime/CoreUObject/Public/UObject/RemoteExecutor.h
+++ b/Engine/Source/Runtime/CoreUObject/Public/UObject/RemoteExecutor.h
@@ -109,6 +109,7 @@ public:
 	virtual void BeginRequest() = 0;
 	virtual void TickSubsystem() = 0;
 	virtual void TickRequest() = 0;
+	virtual void TickAbortedRequest() = 0;
 	virtual bool AreDependenciesSatisfied() const = 0;
 	virtual void BeginMultiServerCommit(TArray<FRemoteServerId>& OutMultiServerCommitRemoteServers) = 0;
 	virtual void ExecuteMultiServerCommit() = 0;
diff --git a/Engine/Source/Runtime/CoreUObject/Public/UObject/RemoteObjectTransfer.h b/Engine/Source/Runtime/CoreUObject/Public/UObject/RemoteObjectTransfer.h
--- a/Engine/Source/Runtime/CoreUObject/Public/UObject/RemoteObjectTransfer.h
+++ b/Engine/Source/Runtime/CoreUObject/Public/UObject/RemoteObjectTransfer.h
@@ -55,12 +55,6 @@ namespace UE::RemoteObject::Transfer
 		FRemoteObjectData ObjectData;
 	};
 
-	/**
-	* (Advanced API) Adds the remote object to the current transaction's requested object list, causing
-	* the object to be migrated to this server. Does not abort the transaction if the object isn't yet local
-	*/
-	COREUOBJECT_API void PrefetchRemoteObject(FRemoteObjectId ObjectId, UObject* DestinationOuter);
-
 	/**
 	* Called when remote object data has been received from a remote server
 	* @param ObjectId Remote object id
