Fixing a crash when comparing or resolving remote object/server ids when remote object support is not compiled in

diff --git a/Engine/Source/Runtime/CoreUObject/Private/UObject/RemoteObject.cpp b/Engine/Source/Runtime/CoreUObject/Private/UObject/RemoteObject.cpp
--- a/Engine/Source/Runtime/CoreUObject/Private/UObject/RemoteObject.cpp
+++ b/Engine/Source/Runtime/CoreUObject/Private/UObject/RemoteObject.cpp
@@ -43,7 +43,9 @@ void AssignGlobalServerIdDebuggingState()
 
 void FRemoteServerId::InitGlobalServerId(FRemoteServerId Id)
 {
-	checkf(GlobalServerId.Id == (uint32)ERemoteServerIdConstants::Local, TEXT("Global server id has already been initialized (%s)"), *GlobalServerId.ToString());
+	// Guard against re-initializing the global id unless remote object support is disabled and the id has been initialized to invalid value
+	checkf(GlobalServerId.Id == (uint32)ERemoteServerIdConstants::Local || (!FRemoteObjectId::RemoteObjectSupportCompiledIn && !GlobalServerId.IsValid()), 
+		TEXT("Global server id has already been initialized (%s)"), *GlobalServerId.ToString());
 	GlobalServerId = Id;
 }
 
@@ -249,6 +251,7 @@ void InitRemoteObjects()
 
 	if (!FRemoteObjectId::RemoteObjectSupportCompiledIn)
 	{
+		FRemoteServerId::InitGlobalServerId(FRemoteServerId());
 		// Always init global server id debug visualizer support but early out if remote object support is not compiled in (UE_WITH_REMOTE_OBJECT_HANDLE is 0)
 		return;
 	}
diff --git a/Engine/Source/Runtime/CoreUObject/Public/UObject/RemoteObjectTypes.h b/Engine/Source/Runtime/CoreUObject/Public/UObject/RemoteObjectTypes.h
--- a/Engine/Source/Runtime/CoreUObject/Public/UObject/RemoteObjectTypes.h
+++ b/Engine/Source/Runtime/CoreUObject/Public/UObject/RemoteObjectTypes.h
@@ -143,7 +143,7 @@ private:
 
 	inline FRemoteServerId GetLocalized() const
 	{
-		return (Id == GlobalServerId.Id && GlobalServerId.IsValid()) ? FRemoteServerId(ERemoteServerIdConstants::Local) : *this;
+		return (Id == GlobalServerId.Id) ? FRemoteServerId(ERemoteServerIdConstants::Local) : *this;
 	}
 
 	inline FRemoteServerId GetGlobalized() const
