[PropertyBag]
- Fix missing lock causing race conditions

diff --git a/Engine/Source/Runtime/CoreUObject/Private/UObject/PropertyBagRepository.cpp b/Engine/Source/Runtime/CoreUObject/Private/UObject/PropertyBagRepository.cpp
--- a/Engine/Source/Runtime/CoreUObject/Private/UObject/PropertyBagRepository.cpp
+++ b/Engine/Source/Runtime/CoreUObject/Private/UObject/PropertyBagRepository.cpp
@@ -808,6 +808,7 @@ UObject* FPropertyBagRepository::CreateInstanceDataObject(UObject* Owner, FArchi
 	UObject* Outer = nullptr;
 	if (IsInstanceDataObjectOuterChainEnabled())
 	{
+		FPropertyBagRepositoryLock LockRepo(this);
 		if (const FPropertyBagAssociationData* OuterData = AssociatedData.Find(Owner->GetOuter()))
 		{
 			Outer = OuterData->InstanceDataObject;
