[ZenLoader]
- Remove some problematic types from going parallel since they are doing recursive stuff and causing deadlocks

diff --git a/Engine/Source/Runtime/CoreUObject/Private/Serialization/AsyncLoading2.cpp b/Engine/Source/Runtime/CoreUObject/Private/Serialization/AsyncLoading2.cpp
--- a/Engine/Source/Runtime/CoreUObject/Private/Serialization/AsyncLoading2.cpp
+++ b/Engine/Source/Runtime/CoreUObject/Private/Serialization/AsyncLoading2.cpp
@@ -54,7 +54,6 @@
 #include "Misc/MTAccessDetector.h"
 #include "UObject/GarbageCollectionInternal.h"
 #include "ProfilingDebugging/MiscTrace.h"
-#include "StructUtils/UserDefinedStruct.h"
 #include "Serialization/LoadTimeTracePrivate.h"
 #include "ProfilingDebugging/CountersTrace.h"
 #include "Serialization/AsyncPackage.h"
@@ -8404,17 +8403,39 @@ bool FAsyncPackage2::CanRunNodeAsync(const FAsyncLoadEventSpec* Spec) const
 {
 	if (Spec->Func == &FAsyncPackage2::Event_PreloadLinkerLoadExports)
 	{
-		// UserDefinedStruct often try to load softobject ptr from text while
-		// they are preloading. This is very deadlock prone so exclude those
-		// types from going wide for now.
+		TRACE_CPUPROFILER_EVENT_SCOPE(FAsyncPackage2::CanRunNodeAsync);
+
+		// This will be moved to .ini settings soon.
+		static TArray<UClass*> ExcludedClasses =
+		{
+			// DataTable are loading other datatable from under a mutex
+			// which is a deadlock hazard. Skip for now.
+			FindObject<UClass>(nullptr, TEXT("/Script/Engine.DataTable")),
+
+			// UserDefinedStruct often try to load softobject ptr from text while
+			// they are preloading. This is very deadlock prone so exclude those
+			// types from going wide for now.
+			FindObject<UClass>(nullptr, TEXT("/Script/CoreUObject.UserDefinedStruct")),
+
+			// Blueprint can recurse into other assets preload, which is a deadlock
+			// hazard, exclude since the majority of their time is spent in create exports anyway.
+			FindObject<UClass>(nullptr, TEXT("/Script/Engine.Blueprint")),
+		};
+
 		const int32 ExportCount = Data.Exports.Num();
 		for (int32 ExportIndex = 0; ExportIndex < ExportCount; ++ExportIndex)
 		{
 			FExportObject& Export = Data.Exports[ExportIndex];
 
-			if (UUserDefinedStruct* Struct = Cast<UUserDefinedStruct>(Export.Object))
+			if (Export.Object)
 			{
-				return false;
+				for (UClass* ExcludedClass : ExcludedClasses)
+				{
+					if (Export.Object->GetClass()->IsChildOf(ExcludedClass))
+					{
+						return false;
+					}
+				}
 			}
 		}
 	}
