Moving types out of NoExportTypes.h
Structs: FPrimaryAssetId, FPrimaryAssetType, FTestUninitializedScriptStructMembersTest, FFilePath, FDirectoryPath, FTemplateString
Classes: UObject, UField, UStruct, UScriptStruct, UFunction, UClass, UPackage, UProperty (& derived types), UModel, UGCBarrier, UPropertyWrapper(& derived types).
Intrinsic classes moved: UPropertyWrapper and derived types.

UPROPERTY(SkipSerialization) used on natively serialized fields that need to be added to the GC schema.
TObjectPtr: Add GetAccessTrackedObjectPtr to allow "downgrade" from untracked ptr without changing some referencing code.
Add ObjectPtrLinkedListBuilder.h to handle changing raw pointer Next in UField to TObjectPtr.
UStruct::SuperStructed changed to TObjectPtr from TNonAccessTrackedObjectPtr. Added an iterator to avoid access tracking and use it in places in Core/etc to avoid cook time regression.

UHT: When adding referenced headers from #include directives to check for circular includes, check for longer path match to avoid misunderstanding some include direcives in Class.h
UHT: Change check for UObject header file to NoExportTypes explicitly when checking legality of 'immutable'.
UHT: Emit StaticConfigName function for class if there is no superclass (UObject)
UHT: Pass class as super argument to DECLARE_CLASS2 if there is no super (UObject)
UHT: When trying to set ClassCastFlags on class, check for FClassName as well as UClassName (legacy UProperty types)
UHT: Use ClassCastFlags enum rather than class source name (for legacy UProperty classes)
UHT: Explicitly omit DECLARE_SERIALIZER for UObject
UHT: Use DECLARE_WITHIN_UPACKAGE macro when ClassWithin is UPackage.
UHT: Don't try and bind superclass for UObject, because UObjectBaseUtility is not parsed.

diff --git a/Engine/Source/Programs/LowLevelTests/FoundationTests/Tests/UObjectTests.cpp b/Engine/Source/Programs/LowLevelTests/FoundationTests/Tests/UObjectTests.cpp
--- a/Engine/Source/Programs/LowLevelTests/FoundationTests/Tests/UObjectTests.cpp
+++ b/Engine/Source/Programs/LowLevelTests/FoundationTests/Tests/UObjectTests.cpp
@@ -53,4 +53,17 @@ TEST_CASE("CoreUObject::Interface::Basic", "[CoreUObject]")
     CHECK(Function != nullptr);
     Obj->ProcessEvent(Function, &Params);
     CHECK_EQUALS("BP call", Params.Result, FName(Params.BaseName, Params.Number));
+}
+
+TEST_CASE("CoreUObject::Class::SuperStructIterator", "[CoreUObject]")
+{
+    UClass* TestObjectClass = UTestObject::StaticClass();
+    TArray<UClass*> Supers;
+    for (UStruct* Struct : TestObjectClass->GetSuperStructIterator())
+    {
+        Supers.Add(CastChecked<UClass>(Struct));
+    }
+    REQUIRE_MESSAGE("Two supers", Supers.Num() == 2);
+    CHECK_MESSAGE("Self first", Supers[0] == TestObjectClass);
+    CHECK_MESSAGE("UObject second", Supers[1] == UObject::StaticClass());
 }
\ No newline at end of file
diff --git a/Engine/Source/Programs/Shared/EpicGames.UHT/Exporters/CodeGen/UhtHeaderCodeGeneratorHFile.cs b/Engine/Source/Programs/Shared/EpicGames.UHT/Exporters/CodeGen/UhtHeaderCodeGeneratorHFile.cs
--- a/Engine/Source/Programs/Shared/EpicGames.UHT/Exporters/CodeGen/UhtHeaderCodeGeneratorHFile.cs
+++ b/Engine/Source/Programs/Shared/EpicGames.UHT/Exporters/CodeGen/UhtHeaderCodeGeneratorHFile.cs
@@ -1567,7 +1567,7 @@ namespace EpicGames.UHT.Exporters.CodeGen
 
 			// export the class's config name
 			UhtClass? superClass = classObj.SuperClass;
-			if (superClass != null && classObj.Config.Length > 0 && classObj.Config != superClass.Config)
+			if (classObj.Config.Length > 0 && (superClass == null || classObj.Config != superClass.Config))
 			{
 				builder.Append("\tstatic const TCHAR* StaticConfigName() {return TEXT(\"").Append(classObj.Config).Append("\");} \\\r\n \\\r\n");
 			}
@@ -1624,7 +1624,7 @@ namespace EpicGames.UHT.Exporters.CodeGen
 				.Append("\tDECLARE_CLASS2(")
 				.Append(classObj)
 				.Append(", ")
-				.Append(superClass != null ? superClass.SourceName : "None")
+				.Append(superClass != null ? superClass.SourceName : classObj.SourceName) // UObject passes itself as super to the macro
 				.Append(", COMPILED_IN_FLAGS(")
 				.Append(classObj.ClassFlags.HasAnyFlags(EClassFlags.Abstract) ? "CLASS_Abstract" : "0");
 
@@ -1634,7 +1634,7 @@ namespace EpicGames.UHT.Exporters.CodeGen
 			{
 				builder
 					.Append("CASTCLASS_")
-					.Append(classObj.SourceName);
+					.Append(classObj.ClassCastFlags.ToString());
 			}
 			else
 			{
@@ -1642,7 +1642,10 @@ namespace EpicGames.UHT.Exporters.CodeGen
 			}
 			builder.Append(", TEXT(\"").Append(classObj.Package.SourceName).Append("\"), ").Append(info.UnregisteredSingletonName).Append(") \\\r\n");
 
-			builder.Append("\tDECLARE_SERIALIZER(").Append(classObj.SourceName).Append(") \\\r\n");
+			if (classObj != Session.UObject)
+			{
+				builder.Append("\tDECLARE_SERIALIZER(").Append(classObj.SourceName).Append(") \\\r\n");
+			}
 
 			// Add the serialization function declaration if we generated one
 			if (classObj.SerializerArchiveType != UhtSerializerArchiveType.None && classObj.SerializerArchiveType != UhtSerializerArchiveType.All)
@@ -1652,7 +1655,15 @@ namespace EpicGames.UHT.Exporters.CodeGen
 
 			if (superClass != null && classObj.ClassWithin != superClass.ClassWithin)
 			{
-				builder.Append("\tDECLARE_WITHIN(").Append(classObj.ClassWithin.SourceName).Append(") \\\r\n");
+				if (classObj.ClassWithin == Session.UPackage)
+				{
+					builder.Append("\tDECLARE_WITHIN_UPACKAGE() \\\r\n");
+
+				}
+				else
+				{
+					builder.Append("\tDECLARE_WITHIN(").Append(classObj.ClassWithin.SourceName).Append(") \\\r\n");
+				}
 			}
 			return builder;
 		}
diff --git a/Engine/Source/Programs/Shared/EpicGames.UHT/Types/UhtClass.cs b/Engine/Source/Programs/Shared/EpicGames.UHT/Types/UhtClass.cs
--- a/Engine/Source/Programs/Shared/EpicGames.UHT/Types/UhtClass.cs
+++ b/Engine/Source/Programs/Shared/EpicGames.UHT/Types/UhtClass.cs
@@ -692,7 +692,10 @@ namespace EpicGames.UHT.Types
 		/// <inheritdoc/>
 		public override void BindSuperAndBases()
 		{
-			BindSuper(SuperIdentifier, UhtFindOptions.Class);
+			if (this != Session.UObject) // Don't try to find UObjectBaseUtility
+			{
+				BindSuper(SuperIdentifier, UhtFindOptions.Class);
+			}
 			BindBases(BaseIdentifiers, UhtFindOptions.Class);
 			BindVerseInterfaces();       
 			base.BindSuperAndBases();
@@ -984,7 +987,8 @@ namespace EpicGames.UHT.Types
 
 					// Check to see if this class matches an entry in the ClassCastFlags.  If it does, that
 					// becomes our ClassCastFlags.  If not, then we inherit the flags from the super.
-					if (Enum.TryParse<EClassCastFlags>(SourceName, false, out EClassCastFlags found))
+					// Also for legacy UObject-derived property types, try replacing the U with F. e.g. UByteProperty -> FByteProperty
+					if (Enum.TryParse<EClassCastFlags>(SourceName, false, out EClassCastFlags found) || Enum.TryParse($"F{SourceName[1..]}", false, out found))
 					{
 						if (found != EClassCastFlags.None)
 						{
@@ -1679,6 +1683,10 @@ namespace EpicGames.UHT.Types
 				{
 					this.LogError($"Interface '{SourceName}' must derive from an interface or UInterface");
 				}
+				else if (ClassType == UhtClassType.Class && this != Session.UObject)
+				{
+					this.LogError($"UClass '{SourceName}' should derive from UObject or another UClass");
+				}
 			}
 
 			if (ClassFlags.HasAnyFlags(EClassFlags.NeedsDeferredDependencyLoading) && !IsChildOf(Session.UClass))
diff --git a/Engine/Source/Programs/Shared/EpicGames.UHT/Types/UhtHeaderFile.cs b/Engine/Source/Programs/Shared/EpicGames.UHT/Types/UhtHeaderFile.cs
--- a/Engine/Source/Programs/Shared/EpicGames.UHT/Types/UhtHeaderFile.cs
+++ b/Engine/Source/Programs/Shared/EpicGames.UHT/Types/UhtHeaderFile.cs
@@ -271,8 +271,12 @@ namespace EpicGames.UHT.Types
 		/// <param name="referenceType">How is the include in question referenced</param>
 		public void AddReferencedHeader(string id, UhtHeaderReferenceType referenceType)
 		{
-			UhtHeaderFile? headerFile = Session.FindHeaderFile(Path.GetFileName(id));
-			if (headerFile != null)
+			string shortPath = Path.GetFileName(id);
+			UhtHeaderFile? headerFile = Session.FindHeaderFile(shortPath);
+			// Check that the found header file matches the path fragment from the include
+			// directive as the session currently only maps header files by their leaf file name
+			// and does not mimic full C++ include paths
+			if (headerFile != null && headerFile.ModuleRelativeFilePath.EndsWith(id, StringComparison.Ordinal))
 			{
 				AddReferencedHeader(headerFile, referenceType);
 			}
diff --git a/Engine/Source/Programs/Shared/EpicGames.UHT/Types/UhtScriptStruct.cs b/Engine/Source/Programs/Shared/EpicGames.UHT/Types/UhtScriptStruct.cs
--- a/Engine/Source/Programs/Shared/EpicGames.UHT/Types/UhtScriptStruct.cs
+++ b/Engine/Source/Programs/Shared/EpicGames.UHT/Types/UhtScriptStruct.cs
@@ -891,9 +891,9 @@ namespace EpicGames.UHT.Types
 
 			if (ScriptStructFlags.HasAnyFlags(EStructFlags.Immutable))
 			{
-				if (HeaderFile != Session.UObject.HeaderFile)
+				if (HeaderFile != Session.NoExportTypesHeader)
 				{
-					this.LogError("Immutable is being phased out in favor of SerializeNative, and is only legal on the mirror structs declared in UObject");
+					this.LogError("Immutable is being phased out in favor of SerializeNative, and is only legal on the mirror structs declared in NoExportTypes.h");
 				}
 			}
 
diff --git a/Engine/Source/Programs/Shared/EpicGames.UHT/Utils/UhtSession.cs b/Engine/Source/Programs/Shared/EpicGames.UHT/Utils/UhtSession.cs
--- a/Engine/Source/Programs/Shared/EpicGames.UHT/Utils/UhtSession.cs
+++ b/Engine/Source/Programs/Shared/EpicGames.UHT/Utils/UhtSession.cs
@@ -833,6 +833,21 @@ namespace EpicGames.UHT.Utils
 			}
 		}
 
+		/// <summary>
+		/// After headers are parsed, returns the UPackage class.
+		/// </summary>
+		public UhtClass UPackage
+		{
+			get
+			{
+				if (_upackage == null)
+				{
+					throw new UhtIceException("UPackage was not defined.");
+				}
+				return _upackage;
+			}
+		}
+
 		/// <summary>
 		/// After headers are parsed, returns the UClass class.
 		/// </summary>
@@ -908,6 +923,11 @@ namespace EpicGames.UHT.Utils
 		/// </summary>
 		public UhtEnum? EVerseNativeCallResult { get; set; } = null;
 
+		/// <summary>
+		/// NoExportTypes.h header in CoreUObject module.
+		/// </summary>
+		public UhtHeaderFile? NoExportTypesHeader { get; private set; } = null;
+
 		private readonly List<UhtModule> _modules = [];
 		private readonly List<UhtHeaderFile> _headerFiles = [];
 		private readonly List<UhtHeaderFile> _sortedHeaderFiles = [];
@@ -918,6 +938,7 @@ namespace EpicGames.UHT.Utils
 		private readonly List<UhtMessage> _messages = [];
 		private Task? _messageTask = null;
 		private UhtClass? _uobject = null;
+		private UhtClass? _upackage = null;
 		private UhtClass? _uclass = null;
 		private UhtClass? _uinterface = null;
 		private UhtClass? _iinterface = null;
@@ -2159,6 +2180,7 @@ namespace EpicGames.UHT.Utils
 				// Locate the NoExportTypes.h file and add it to every other header file
 				if (_headerFileDictionary.TryGetValue("NoExportTypes.h", out UhtHeaderFile? noExportTypes))
 				{
+					NoExportTypesHeader = noExportTypes;
 					foreach (UhtHeaderFile headerFile in _headerFiles)
 					{
 						if (headerFile != noExportTypes)
@@ -2199,6 +2221,7 @@ namespace EpicGames.UHT.Utils
 				PopulateSymbolTable();
 
 				_uobject = (UhtClass?)FindType(null, UhtFindOptions.SourceName | UhtFindOptions.Class, "UObject");
+				_upackage = (UhtClass?)FindType(null, UhtFindOptions.SourceName | UhtFindOptions.Class, "UPackage");
 				_uclass = (UhtClass?)FindType(null, UhtFindOptions.SourceName | UhtFindOptions.Class, "UClass");
 				_uinterface = (UhtClass?)FindType(null, UhtFindOptions.SourceName | UhtFindOptions.Class, "UInterface");
 				_iinterface = (UhtClass?)FindType(null, UhtFindOptions.SourceName | UhtFindOptions.Class, "IInterface");
diff --git a/Engine/Source/Runtime/CoreUObject/Private/Internationalization/StabilizeLocalizationKeys.cpp b/Engine/Source/Runtime/CoreUObject/Private/Internationalization/StabilizeLocalizationKeys.cpp
--- a/Engine/Source/Runtime/CoreUObject/Private/Internationalization/StabilizeLocalizationKeys.cpp
+++ b/Engine/Source/Runtime/CoreUObject/Private/Internationalization/StabilizeLocalizationKeys.cpp
@@ -16,7 +16,7 @@ void StabilizeLocalizationKeys::StabilizeLocalizationKeysForProperty(FProperty*
 	auto GetPropertyName = [InProp]() -> FString
 	{
 		// UserDefinedStructs should use the display name rather than the raw name, since the raw name doesn't match what the user generally sees or entered as the property name
-		for (UStruct* PropOwnerStruct = InProp->GetOwnerStruct(); PropOwnerStruct; PropOwnerStruct = PropOwnerStruct->GetSuperStruct())
+		for (UStruct* PropOwnerStruct : InProp->GetOwnerStruct()->GetSuperStructIterator())
 		{
 			static const FName UserDefinedStructName = TEXT("UserDefinedStruct");
 			if (PropOwnerStruct->GetClass()->GetFName() == UserDefinedStructName)
diff --git a/Engine/Source/Runtime/CoreUObject/Private/UObject/Class.cpp b/Engine/Source/Runtime/CoreUObject/Private/UObject/Class.cpp
--- a/Engine/Source/Runtime/CoreUObject/Private/UObject/Class.cpp
+++ b/Engine/Source/Runtime/CoreUObject/Private/UObject/Class.cpp
@@ -46,6 +46,7 @@
 #include "UObject/CoreRedirects.h"
 #include "UObject/ObjectMacros.h"
 #include "UObject/RemoteObjectTransfer.h"
+#include "UObject/TestUninitializedScriptStructMembersTest.h"
 #include "Internationalization/PolyglotTextData.h"
 #include "Serialization/ArchiveScriptReferenceCollector.h"
 #include "Serialization/ArchiveUObjectFromStructuredArchive.h"
@@ -84,8 +85,6 @@
 DEFINE_LOG_CATEGORY(LogScriptSerialization);
 DEFINE_LOG_CATEGORY(LogClass);
 
-UE_IMPLEMENT_STRUCT("/Script/CoreUObject", TestUninitializedScriptStructMembersTest);
-
 LLM_DEFINE_TAG(UObject_UClass);
 
 #if defined(_MSC_VER) && _MSC_VER == 1900
@@ -604,12 +603,6 @@ void UField::SetAssociatedFField(FField* InField)
 }
 #endif // WITH_EDITORONLY_DATA
 
-IMPLEMENT_CORE_INTRINSIC_CLASS(UField, UObject,
-	{
-		UE::GC::DeclareIntrinsicMembers(Class, { UE_GC_MEMBER(UField, Next) });
-	}
-);
-
 
 
 /*-----------------------------------------------------------------------------
@@ -1352,7 +1345,7 @@ void UStruct::LoadTaggedPropertiesFromText(FStructuredArchive::FSlot Slot, uint8
 		// Resolve any redirects if necessary
 		if (bUseRedirects && !UnderlyingArchive.HasAnyPortFlags(PPF_DuplicateForPIE | PPF_Duplicate))
 		{
-			for (UStruct* CheckStruct = GetOwnerStruct(); CheckStruct; CheckStruct = CheckStruct->GetSuperStruct())
+			for (UStruct* CheckStruct : GetOwnerStruct()->GetSuperStructIterator())
 			{
 				FName NewTagName = FProperty::FindRedirectedPropertyName(CheckStruct, PropertyName);
 				if (!NewTagName.IsNone())
@@ -1684,7 +1677,7 @@ void UStruct::SerializeVersionedTaggedProperties(FStructuredArchive::FSlot Slot,
 					// No need to check redirects on platforms where everything is cooked. Always check for save games
 					if (bUseRedirects && !UnderlyingArchive.HasAnyPortFlags(PPF_DuplicateForPIE | PPF_Duplicate))
 					{
-						for (UStruct* CheckStruct = GetOwnerStruct(); CheckStruct; CheckStruct = CheckStruct->GetSuperStruct())
+						for (UStruct* CheckStruct : GetOwnerStruct()->GetSuperStructIterator())
 						{
 							FName NewTagName = FProperty::FindRedirectedPropertyName(CheckStruct, Tag.Name);
 							if (!NewTagName.IsNone())
@@ -2283,7 +2276,7 @@ void UStruct::Serialize(FArchive& Ar)
 	UStruct* SuperStructBefore = GetSuperStruct();
 #endif
 
-	Ar << SuperStruct.GetAccessTrackedObjectPtr();
+	Ar << SuperStruct;
 
 #if USTRUCT_FAST_ISCHILDOF_IMPL == USTRUCT_ISCHILDOF_STRUCTARRAY
 	if (Ar.IsLoading())
@@ -2515,7 +2508,7 @@ void UStruct::AddReferencedObjects(UObject* InThis, FReferenceCollector& Collect
 	if( GIsEditor )
 	{
 		// Required by the unified GC when running in the editor
-		Collector.AddReferencedObject( This->SuperStruct.GetAccessTrackedObjectPtr(), This );
+		Collector.AddReferencedObject( This->SuperStruct, This );
 		Collector.AddReferencedObject( This->Children, This );
 		Collector.AddReferencedObjects(This->ScriptAndPropertyObjectReferences, This);
 	}
@@ -2556,23 +2549,20 @@ FString UStruct::GetAuthoredNameForField(const FField* Field) const
 bool UStruct::GetBoolMetaDataHierarchical(const FName& Key) const
 {
 	bool bResult = false;
-	const UStruct* TestStruct = this;
-	while( TestStruct )
+	for(const UStruct* TestStruct : GetSuperStructIterator())
 	{
 		if( TestStruct->HasMetaData(Key) )
 		{
 			bResult = TestStruct->GetBoolMetaData(Key);
 			break;
 		}
-
-		TestStruct = TestStruct->SuperStruct;
 	}
 	return bResult;
 }
 
 bool UStruct::GetStringMetaDataHierarchical(const FName& Key, FString* OutValue) const
 {
-	for (const UStruct* TestStruct = this; TestStruct != nullptr; TestStruct = TestStruct->GetSuperStruct())
+	for (const UStruct* TestStruct : GetSuperStructIterator())
 	{
 		if (const FString* FoundMetaData = TestStruct->FindMetaData(Key))
 		{
@@ -2590,7 +2580,7 @@ bool UStruct::GetStringMetaDataHierarchical(const FName& Key, FString* OutValue)
 
 const UStruct* UStruct::HasMetaDataHierarchical(const FName& Key) const
 {
-	for (const UStruct* TestStruct = this; TestStruct != nullptr; TestStruct = TestStruct->GetSuperStruct())
+	for (const UStruct* TestStruct : GetSuperStructIterator())
 	{
 		if (TestStruct->HasMetaData(Key))
 		{
@@ -2711,16 +2701,6 @@ void UStruct::InstanceSubobjectTemplates( TNotNull<void*> Data, void const* Defa
 	}
 }
 
-IMPLEMENT_CORE_INTRINSIC_CLASS(UStruct, UField,
-	{
-		Class->CppClassStaticFunctions = UOBJECT_CPPCLASS_STATICFUNCTIONS_FORCLASS(UStruct);
-
-		// Note: None of the *Link members need to be emitted, as they only contain properties
-		// that are in the Children chain or SuperStruct->Children chains.
-		UE::GC::DeclareIntrinsicMembers(Class, { UE_GC_MEMBER(UStruct, SuperStruct), UE_GC_MEMBER(UStruct, Children), UE_GC_MEMBER(UStruct, ScriptAndPropertyObjectReferences) });
-	}
-);
-
 void UStruct::TagSubobjects(EObjectFlags NewFlags)
 {
 	Super::TagSubobjects(NewFlags);
@@ -2754,7 +2734,8 @@ bool UStruct::IsChildOf( const UStruct* SomeBase ) const
 	}
 
 	bool bOldResult = false;
-	for ( const UStruct* TempStruct=this; TempStruct; TempStruct=TempStruct->GetSuperStruct() )
+	// Avoid access tracking when traversing supers to reduce costs in cooking
+	for (const UStruct* TempStruct : GetSuperStructIterator())
 	{
 		if ( TempStruct == SomeBase )
 		{
@@ -4248,7 +4229,7 @@ int32 FStructUtils::AttemptToFindUninitializedScriptStructMembers()
 
 	int32 UninitializedScriptStructMemberCount = 0;
 	int32 UninitializedObjectPropertyCount = 0;
-	UScriptStruct* TestUninitializedScriptStructMembersTestStruct = TBaseStructure<FTestUninitializedScriptStructMembersTest>::Get();
+	UScriptStruct* TestUninitializedScriptStructMembersTestStruct = StaticStruct<FTestUninitializedScriptStructMembersTest>();
 	check(TestUninitializedScriptStructMembersTestStruct != nullptr);
 	
 	{
@@ -4384,6 +4365,33 @@ bool FAutomationTestAttemptToFindUninitializedScriptStructMembers::RunTest(const
 	}
 }
 
+#if WITH_TESTS
+#include "Tests/TestHarnessAdapter.h"
+#if WITH_LOW_LEVEL_TESTS
+#include "TestCommon/Expectations.h"
+#endif
+
+TEST_CASE_NAMED(FTestUninitializedStructMemberTest, "System::CoreUObject::ScriptStruct::UninitializedMemberTest", "[CoreUObject]")
+{
+	const void* BadPointer = (void*)0xFEFEFEFEFEFEFEFEull;
+	UScriptStruct* TestUninitializedScriptStructMembersTestStruct = StaticStruct<FTestUninitializedScriptStructMembersTest>();
+
+	// First test if the tests aren't broken
+	FScriptStructTestWrapper WrapperFE(TestUninitializedScriptStructMembersTestStruct, 0xFE);
+	const FObjectProperty* UninitializedProperty = CastFieldChecked<const FObjectProperty>(TestUninitializedScriptStructMembersTestStruct->FindPropertyByName(TEXT("UninitializedObjectReference")));
+	const FObjectProperty* InitializedProperty = CastFieldChecked<const FObjectProperty>(TestUninitializedScriptStructMembersTestStruct->FindPropertyByName(TEXT("InitializedObjectReference")));
+	
+	//using reinterpret_cast to avoid any methods of TObjectPtr being invoked
+	const TObjectPtr<UObject>* UninitializedPropValue = UninitializedProperty->GetPropertyValuePtr_InContainer(WrapperFE.GetData());
+	const void* const* RawValue = reinterpret_cast<const void* const*>(UninitializedPropValue);
+	CHECK_MESSAGE("FTestUninitializedScriptStructMembersTest::UninitializedObjectReference should be BadPointer value", *RawValue == BadPointer);
+
+	const TObjectPtr<UObject>* InitializedPropValue = InitializedProperty->GetPropertyValuePtr_InContainer(WrapperFE.GetData());
+	RawValue = reinterpret_cast<const void* const*>(InitializedPropValue);
+	CHECK_MESSAGE("FTestUninitializedScriptStructMembersTest::InitializedObjectReference should be null", *RawValue == nullptr);
+}
+#endif
+
 #if WITH_METADATA
 /** 
 * Checks if MetaData value contains short type name 
@@ -4769,10 +4777,6 @@ bool FTestUndeclaredScriptStructObjectReferencesTest::Serialize(FArchive& Ar)
 	return true;
 }
 
-IMPLEMENT_CORE_INTRINSIC_CLASS(UScriptStruct, UStruct,
-	{
-	}
-);
 
 /*-----------------------------------------------------------------------------
 	UClass implementation.
@@ -5692,7 +5696,7 @@ FTopLevelAssetPath UClass::GetReinstancedClassPathName() const
 		delete [] StructBaseChainArray;
 
 		int32 Depth = 0;
-		for (UStruct* Ptr = static_cast<UStruct*>(this); Ptr; Ptr = Ptr->GetSuperStruct())
+		for (UStruct* Ptr : static_cast<UStruct*>(this)->GetSuperStructIterator())
 		{
 			++Depth;
 		}
@@ -5700,7 +5704,7 @@ FTopLevelAssetPath UClass::GetReinstancedClassPathName() const
 		FStructBaseChain** Bases = new FStructBaseChain*[Depth];
 		{
 			FStructBaseChain** Base = Bases + Depth;
-			for (UStruct* Ptr = static_cast<UStruct*>(this); Ptr; Ptr = Ptr->GetSuperStruct())
+			for (UStruct* Ptr : static_cast<UStruct*>(this)->GetSuperStructIterator())
 			{
 				*--Base = Ptr;
 			}
@@ -7331,24 +7335,6 @@ bool UClass::CanCreateInstanceDataObject() const
 }
 #endif
 
-#if WITH_EDITORONLY_DATA
-#define UCLASS_EDITOR_GC_MEMBERS UE_GC_MEMBER(UClass, ClassGeneratedBy)
-#else
-#define UCLASS_EDITOR_GC_MEMBERS
-#endif
-
-IMPLEMENT_CORE_INTRINSIC_CLASS(UClass, UStruct,
-	{
-		Class->CppClassStaticFunctions = UOBJECT_CPPCLASS_STATICFUNCTIONS_FORCLASS(UClass);
-PRAGMA_DISABLE_DEPRECATION_WARNINGS
-		UE::GC::DeclareIntrinsicMembers(Class, {	UE_GC_MEMBER(UClass, ClassDefaultObject),
-													UE_GC_MEMBER(UClass, ClassWithin), 
-													UE_GC_MEMBER(UClass, NetFields),
-													UCLASS_EDITOR_GC_MEMBERS} );
-PRAGMA_ENABLE_DEPRECATION_WARNINGS
-	}
-);
-
 void GetPrivateStaticClassBody(
 	const TCHAR* PackageName,
 	const TCHAR* Name,
@@ -7917,18 +7903,6 @@ UScriptStruct* TBaseStructure<FInt32Interval>::Get()
 	return ScriptStruct;
 }
 
-UScriptStruct* TBaseStructure<FPrimaryAssetType>::Get()
-{
-	static UScriptStruct* ScriptStruct = StaticGetBaseStructureInternal(TEXT("PrimaryAssetType"));
-	return ScriptStruct;
-}
-
-UScriptStruct* TBaseStructure<FPrimaryAssetId>::Get()
-{
-	static UScriptStruct* ScriptStruct = StaticGetBaseStructureInternal(TEXT("PrimaryAssetId"));
-	return ScriptStruct;
-}
-
 UScriptStruct* TBaseStructure<FPolyglotTextData>::Get()
 {
 	static UScriptStruct* ScriptStruct = StaticGetBaseStructureInternal(TEXT("PolyglotTextData"));
@@ -7959,13 +7933,6 @@ UScriptStruct* TBaseStructure<FFrameRate>::Get()
 	return ScriptStruct;
 }
 
-
-UScriptStruct* TBaseStructure<FTestUninitializedScriptStructMembersTest>::Get()
-{
-	static UScriptStruct* ScriptStruct = StaticGetBaseStructureInternal(TEXT("TestUninitializedScriptStructMembersTest"));
-	return ScriptStruct;
-}
-
 UScriptStruct* TBaseStructure<FRemoteObjectId>::Get()
 {
 	static UScriptStruct* ScriptStruct = StaticGetBaseStructureInternal(TEXT("RemoteObjectId"));
@@ -8048,11 +8015,6 @@ UE_DEFINE_CORE_VARIANT_TYPE(Sphere3,	Sphere);
 	
 #undef UE_DEFINE_CORE_VARIANT_TYPE
 
-IMPLEMENT_CORE_INTRINSIC_CLASS(UFunction, UStruct,
-	{
-	}
-);
-
 UDelegateFunction::UDelegateFunction(const FObjectInitializer& ObjectInitializer, UFunction* InSuperFunction, EFunctionFlags InFunctionFlags, SIZE_T ParamsSize)
 	: UFunction(ObjectInitializer, InSuperFunction, InFunctionFlags, ParamsSize)
 {
diff --git a/Engine/Source/Runtime/CoreUObject/Private/UObject/Enum.cpp b/Engine/Source/Runtime/CoreUObject/Private/UObject/Enum.cpp
--- a/Engine/Source/Runtime/CoreUObject/Private/UObject/Enum.cpp
+++ b/Engine/Source/Runtime/CoreUObject/Private/UObject/Enum.cpp
@@ -1072,10 +1072,6 @@ int64 UEnum::ParseEnum(const TCHAR*& Str)
 	}
 }
 
-IMPLEMENT_CORE_INTRINSIC_CLASS(UEnum, UField,
-	{
-	}
-);
 
 #if WITH_DEV_AUTOMATION_TESTS 
 
diff --git a/Engine/Source/Runtime/CoreUObject/Private/UObject/GarbageCollectionHistory.cpp b/Engine/Source/Runtime/CoreUObject/Private/UObject/GarbageCollectionHistory.cpp
--- a/Engine/Source/Runtime/CoreUObject/Private/UObject/GarbageCollectionHistory.cpp
+++ b/Engine/Source/Runtime/CoreUObject/Private/UObject/GarbageCollectionHistory.cpp
@@ -14,6 +14,7 @@
 #include "UObject/UObjectBase.h"
 #include "UObject/Object.h"
 #include "UObject/UnrealType.h"
+#include "UObject/GarbageCollectionInternal.h"
 #include "UObject/GCObject.h"
 #include "UObject/GCScopeLock.h"
 #include "UObject/GCVerseCellInfo.h"
@@ -54,19 +55,6 @@ static FAutoConsoleCommand CmdSetHistorySize(
 		})
 );
 
-/** Dummy class to represent GC barrier inside of GC history object graph */
-class UGCBarrier : public UObject
-{
-	DECLARE_CLASS_INTRINSIC(UGCBarrier,
-		UObject,
-		CLASS_Transient,
-		TEXT("/Script/CoreUObject"));
-};
-
-IMPLEMENT_CORE_INTRINSIC_CLASS(UGCBarrier, UObject,
-	{
-	});
-
 FGCHistory::~FGCHistory()
 {
 	Cleanup();
diff --git a/Engine/Source/Runtime/CoreUObject/Private/UObject/GarbageCollectionInternal.h b/Engine/Source/Runtime/CoreUObject/Private/UObject/GarbageCollectionInternal.h
--- a/Engine/Source/Runtime/CoreUObject/Private/UObject/GarbageCollectionInternal.h
+++ b/Engine/Source/Runtime/CoreUObject/Private/UObject/GarbageCollectionInternal.h
@@ -8,13 +8,24 @@
 
 #include "CoreMinimal.h"
 #include "Stats/Stats.h"
+#include "UObject/Object.h"
 #include "UObject/GarbageCollection.h"
 #include "UObject/GCScopeLock.h"
 
+#include "GarbageCollectionInternal.generated.h"
+
 extern FGCCSyncObject* GGCSingleton;
 
 /** Returns true if GC wants to run */
 FORCEINLINE bool IsGarbageCollectionWaiting()
 {
 	return GGCSingleton->IsGCWaiting();
-}
\ No newline at end of file
+}
+
+/** Dummy class to represent GC barrier inside of GC history object graph */
+UCLASS(Transient)
+class UGCBarrier : public UObject
+{
+    GENERATED_BODY()
+public:
+};
\ No newline at end of file
diff --git a/Engine/Source/Runtime/CoreUObject/Private/UObject/Package.cpp b/Engine/Source/Runtime/CoreUObject/Private/UObject/Package.cpp
--- a/Engine/Source/Runtime/CoreUObject/Private/UObject/Package.cpp
+++ b/Engine/Source/Runtime/CoreUObject/Private/UObject/Package.cpp
@@ -430,8 +430,3 @@ void UPackage::SetSavedHash(const FIoHash& InSavedHash)
 }
 
 #endif
-
-IMPLEMENT_CORE_INTRINSIC_CLASS(UPackage, UObject,
-	{
-	}
-);
diff --git a/Engine/Source/Runtime/CoreUObject/Private/UObject/Property.cpp b/Engine/Source/Runtime/CoreUObject/Private/UObject/Property.cpp
--- a/Engine/Source/Runtime/CoreUObject/Private/UObject/Property.cpp
+++ b/Engine/Source/Runtime/CoreUObject/Private/UObject/Property.cpp
@@ -600,35 +600,6 @@ struct TStructOpsTypeTraits<FTimecode> : public TStructOpsTypeTraitsBase2<FTimec
 };
 UE_IMPLEMENT_STRUCT("/Script/CoreUObject", Timecode);
 
-template<>
-struct TStructOpsTypeTraits<FPrimaryAssetType> : public TStructOpsTypeTraitsBase2<FPrimaryAssetType>
-{
-	enum
-	{
-		WithZeroConstructor = true,
-		WithCopy = true,
-		WithIdenticalViaEquality = true,
-		WithExportTextItem = true,
-		WithImportTextItem = true,
-		WithStructuredSerializeFromMismatchedTag = true,
-	};
-};
-UE_IMPLEMENT_STRUCT("/Script/CoreUObject", PrimaryAssetType);
-
-template<>
-struct TStructOpsTypeTraits<FPrimaryAssetId> : public TStructOpsTypeTraitsBase2<FPrimaryAssetId>
-{
-	enum
-	{
-		WithZeroConstructor = true,
-		WithCopy = true,
-		WithIdenticalViaEquality = true,
-		WithExportTextItem = true,
-		WithImportTextItem = true,
-		WithStructuredSerializeFromMismatchedTag = true,
-	};
-};
-UE_IMPLEMENT_STRUCT("/Script/CoreUObject", PrimaryAssetId);
 
 
 template<>
diff --git a/Engine/Source/Runtime/CoreUObject/Private/UObject/PropertyPrivate.cpp b/Engine/Source/Runtime/CoreUObject/Private/UObject/PropertyPrivate.cpp
--- a/Engine/Source/Runtime/CoreUObject/Private/UObject/PropertyPrivate.cpp
+++ b/Engine/Source/Runtime/CoreUObject/Private/UObject/PropertyPrivate.cpp
@@ -116,10 +116,6 @@ void UProperty::SetAssociatedFField(FField* InField)
 #endif // WITH_EDITORONLY_DATA
 
 
-IMPLEMENT_CORE_INTRINSIC_CLASS(UProperty, UField,
-	{
-	}
-);
 
 
 
@@ -154,51 +150,16 @@ void UEnumProperty::Serialize(FArchive& Ar)
 	}
 }
 
-void UEnumProperty::AddReferencedObjects(UObject* InThis, FReferenceCollector& Collector)
-{
-	UEnumProperty* This = CastChecked<UEnumProperty>(InThis);
-	Collector.AddReferencedObject(This->Enum, This);
-	Collector.AddReferencedObject(This->UnderlyingProp, This);
-	Super::AddReferencedObjects(InThis, Collector);
-}
-
-namespace UE::CoreUObject::Private
-{
-	struct FEnumPropertyFriend
-	{
-		static const int32 EnumOffset = STRUCT_OFFSET(UEnumProperty, Enum);
-		static const int32 UnderlyingPropOffset = STRUCT_OFFSET(UEnumProperty, UnderlyingProp);
-	};
-}
-
-IMPLEMENT_CORE_INTRINSIC_CLASS(UEnumProperty, UProperty,
-	{
-		UE::GC::DeclareIntrinsicMembers(Class, { UE_GC_MEMBER(UEnumProperty, Enum), UE_GC_MEMBER(UEnumProperty, UnderlyingProp) });
-	}
-);
-
 void UArrayProperty::Serialize(FArchive& Ar)
 {
 	Super::Serialize(Ar);
 	Ar << Inner;
 	checkSlow(Inner || HasAnyFlags(RF_ClassDefaultObject) || !IsValidChecked(this));
 }
-void UArrayProperty::AddReferencedObjects(UObject* InThis, FReferenceCollector& Collector)
-{
-	UArrayProperty* This = CastChecked<UArrayProperty>(InThis);
-	Collector.AddReferencedObject(This->Inner, This);
-	Super::AddReferencedObjects(This, Collector);
-}
 
-IMPLEMENT_CORE_INTRINSIC_CLASS(UArrayProperty, UProperty,
-	{
-		UE::GC::DeclareIntrinsicMembers(Class, { UE_GC_MEMBER(UArrayProperty, Inner) });
-
-		// Ensure that TArray and FScriptArray are interchangeable, as FScriptArray will be used to access a native array property
-		// from script that is declared as a TArray in C++.
-		static_assert(sizeof(FScriptArray) == sizeof(TArray<uint8>), "FScriptArray and TArray<uint8> must be interchangable.");
-	}
-);
+// Ensure that TArray and FScriptArray are interchangeable, as FScriptArray will be used to access a native array property
+// from script that is declared as a TArray in C++.
+static_assert(sizeof(FScriptArray) == sizeof(TArray<uint8>), "FScriptArray and TArray<uint8> must be interchangable.");
 
 
 void UObjectPropertyBase::BeginDestroy()
@@ -245,18 +206,6 @@ void UObjectPropertyBase::SetPropertyClass(UClass* NewPropertyClass)
 }
 #endif // USE_UPROPERTY_LOAD_DEFERRING
 
-void UObjectPropertyBase::AddReferencedObjects(UObject* InThis, FReferenceCollector& Collector)
-{
-	UObjectPropertyBase* This = CastChecked<UObjectPropertyBase>(InThis);
-	Collector.AddReferencedObject(This->PropertyClass, This);
-	Super::AddReferencedObjects(This, Collector);
-}
-
-IMPLEMENT_CORE_INTRINSIC_CLASS(UObjectPropertyBase, UProperty,
-	{
-		UE::GC::DeclareIntrinsicMembers(Class, { UE_GC_MEMBER(UObjectProperty, PropertyClass) });
-	}
-);
 
 
 UBoolProperty::UBoolProperty(const FObjectInitializer& ObjectInitializer)
@@ -346,11 +295,6 @@ void UBoolProperty::Serialize(FArchive& Ar)
 }
 PRAGMA_ENABLE_DEPRECATION_WARNINGS
 
-IMPLEMENT_CORE_INTRINSIC_CLASS(UBoolProperty, UProperty,
-	{
-	}
-);
-
 void UByteProperty::Serialize(FArchive& Ar)
 {
 	Super::Serialize(Ar);
@@ -360,18 +304,6 @@ void UByteProperty::Serialize(FArchive& Ar)
 		Ar.Preload(Enum);
 	}
 }
-void UByteProperty::AddReferencedObjects(UObject* InThis, FReferenceCollector& Collector)
-{
-	UByteProperty* This = CastChecked<UByteProperty>(InThis);
-	Collector.AddReferencedObject(This->Enum, This);
-	Super::AddReferencedObjects(This, Collector);
-}
-
-IMPLEMENT_CORE_INTRINSIC_CLASS(UByteProperty, UNumericProperty,
-	{
-		UE::GC::DeclareIntrinsicMembers(Class, { UE_GC_MEMBER(UByteProperty, Enum) });
-	}
-);
 
 void UClassProperty::BeginDestroy()
 {
@@ -429,19 +361,6 @@ void UClassProperty::SetMetaClass(UClass* NewMetaClass)
 }
 #endif // USE_UPROPERTY_LOAD_DEFERRING
 
-void UClassProperty::AddReferencedObjects(UObject* InThis, FReferenceCollector& Collector)
-{
-	UClassProperty* This = CastChecked<UClassProperty>(InThis);
-	Collector.AddReferencedObject(This->MetaClass, This);
-	Super::AddReferencedObjects(This, Collector);
-}
-
-IMPLEMENT_CORE_INTRINSIC_CLASS(UClassProperty, UObjectProperty,
-	{
-		UE::GC::DeclareIntrinsicMembers(Class, { UE_GC_MEMBER(UClassProperty, MetaClass) });
-	}
-);
-
 void UDelegateProperty::Serialize(FArchive& Ar)
 {
 	Super::Serialize(Ar);
@@ -471,42 +390,8 @@ void UDelegateProperty::BeginDestroy()
 }
 
 
-IMPLEMENT_CORE_INTRINSIC_CLASS(UDelegateProperty, UProperty,
-	{
-		UE::GC::DeclareIntrinsicMembers(Class, { UE_GC_MEMBER(UDelegateProperty, SignatureFunction) });
-	}
-);
-
-
-IMPLEMENT_CORE_INTRINSIC_CLASS(UDoubleProperty, UNumericProperty,
-	{
-	}
-);
-
-IMPLEMENT_CORE_INTRINSIC_CLASS(UFloatProperty, UNumericProperty,
-	{
-	}
-);
-
-IMPLEMENT_CORE_INTRINSIC_CLASS(UIntProperty, UNumericProperty,
-	{
-	}
-);
 
-IMPLEMENT_CORE_INTRINSIC_CLASS(UInt8Property, UNumericProperty,
-	{
-	}
-);
 
-IMPLEMENT_CORE_INTRINSIC_CLASS(UInt16Property, UNumericProperty,
-	{
-	}
-);
-
-IMPLEMENT_CORE_INTRINSIC_CLASS(UInt64Property, UNumericProperty,
-	{
-	}
-);
 
 void UInterfaceProperty::BeginDestroy()
 {
@@ -565,24 +450,6 @@ void UInterfaceProperty::SetInterfaceClass(UClass* NewInterfaceClass)
 }
 #endif // USE_UPROPERTY_LOAD_DEFERRING
 
-void UInterfaceProperty::AddReferencedObjects(UObject* InThis, FReferenceCollector& Collector)
-{
-	UInterfaceProperty* This = CastChecked<UInterfaceProperty>(InThis);
-	Collector.AddReferencedObject(This->InterfaceClass, This);
-	Super::AddReferencedObjects(This, Collector);
-}
-
-IMPLEMENT_CORE_INTRINSIC_CLASS(UInterfaceProperty, UProperty,
-	{
-		UE::GC::DeclareIntrinsicMembers(Class, { UE_GC_MEMBER(UInterfaceProperty, InterfaceClass) });
-	}
-);
-
-IMPLEMENT_CORE_INTRINSIC_CLASS(ULazyObjectProperty, UObjectPropertyBase,
-	{
-	}
-);
-
 
 UMapProperty::UMapProperty(const FObjectInitializer& ObjectInitializer, ECppProperty, int32 InOffset, EPropertyFlags InFlags)
 	: UProperty(ObjectInitializer, EC_CppProperty, InOffset, InFlags)
@@ -599,25 +466,9 @@ void UMapProperty::Serialize(FArchive& Ar)
 	Ar << ValueProp;
 }
 
-void UMapProperty::AddReferencedObjects(UObject* InThis, FReferenceCollector& Collector)
-{
-	UMapProperty* This = CastChecked<UMapProperty>(InThis);
-
-	Collector.AddReferencedObject(This->KeyProp, This);
-	Collector.AddReferencedObject(This->ValueProp, This);
-
-	Super::AddReferencedObjects(This, Collector);
-}
-
-IMPLEMENT_CORE_INTRINSIC_CLASS(UMapProperty, UProperty,
-	{
-		UE::GC::DeclareIntrinsicMembers(Class, { UE_GC_MEMBER(UMapProperty, KeyProp), UE_GC_MEMBER(UMapProperty, ValueProp) });
-
-		// Ensure that TArray and FScriptMap are interchangeable, as FScriptMap will be used to access a native array property
-		// from script that is declared as a TArray in C++.
-		static_assert(sizeof(FScriptMap) == sizeof(TMap<uint32, uint8>), "FScriptMap and TMap<uint32, uint8> must be interchangable.");
-	}
-);
+// Ensure that TArray and FScriptMap are interchangeable, as FScriptMap will be used to access a native array property
+// from script that is declared as a TArray in C++.
+static_assert(sizeof(FScriptMap) == sizeof(TMap<uint32, uint8>), "FScriptMap and TMap<uint32, uint8> must be interchangable.");
 
 void UMulticastDelegateProperty::Serialize(FArchive& Ar)
 {
@@ -647,37 +498,6 @@ void UMulticastDelegateProperty::BeginDestroy()
 	Super::BeginDestroy();
 }
 
-IMPLEMENT_CORE_INTRINSIC_CLASS(UMulticastDelegateProperty, UProperty,
-	{
-		UE::GC::DeclareIntrinsicMembers(Class, { UE_GC_MEMBER(UMulticastDelegateProperty, SignatureFunction) });
-	}
-);
-
-IMPLEMENT_CORE_INTRINSIC_CLASS(UMulticastInlineDelegateProperty, UMulticastDelegateProperty,
-	{
-	}
-);
-
-IMPLEMENT_CORE_INTRINSIC_CLASS(UMulticastSparseDelegateProperty, UMulticastDelegateProperty,
-	{
-	}
-);
-
-IMPLEMENT_CORE_INTRINSIC_CLASS(UNameProperty, UProperty,
-	{
-	}
-);
-
-IMPLEMENT_CORE_INTRINSIC_CLASS(UNumericProperty, UProperty,
-	{
-	}
-);
-
-IMPLEMENT_CORE_INTRINSIC_CLASS(UObjectProperty, UObjectPropertyBase,
-	{
-	}
-);
-
 
 USetProperty::USetProperty(const FObjectInitializer& ObjectInitializer, ECppProperty, int32 InOffset, EPropertyFlags InFlags)
 	: UProperty(ObjectInitializer, EC_CppProperty, InOffset, InFlags)
@@ -692,24 +512,9 @@ void USetProperty::Serialize(FArchive& Ar)
 	Ar << ElementProp;
 }
 
-void USetProperty::AddReferencedObjects(UObject* InThis, FReferenceCollector& Collector)
-{
-	USetProperty* This = CastChecked<USetProperty>(InThis);
-
-	Collector.AddReferencedObject(This->ElementProp, This);
-
-	Super::AddReferencedObjects(This, Collector);
-}
-
-IMPLEMENT_CORE_INTRINSIC_CLASS(USetProperty, UProperty,
-	{
-		UE::GC::DeclareIntrinsicMembers(Class, { UE_GC_MEMBER(USetProperty, ElementProp) });
-
-		// Ensure that TArray and FScriptMap are interchangeable, as FScriptMap will be used to access a native array property
-		// from script that is declared as a TArray in C++.
-		static_assert(sizeof(FScriptSet) == sizeof(TMap<uint32, uint8>), "FScriptSet and TSet<uint32, uint8> must be interchangable.");
-	}
-);
+// Ensure that TArray and FScriptMap are interchangeable, as FScriptMap will be used to access a native array property
+// from script that is declared as a TArray in C++.
+static_assert(sizeof(FScriptSet) == sizeof(TMap<uint32, uint8>), "FScriptSet and TSet<uint32, uint8> must be interchangable.");
 
 
 void USoftClassProperty::BeginDestroy()
@@ -768,29 +573,6 @@ void USoftClassProperty::SetMetaClass(UClass* NewMetaClass)
 }
 #endif // USE_UPROPERTY_LOAD_DEFERRING
 
-void USoftClassProperty::AddReferencedObjects(UObject* InThis, FReferenceCollector& Collector)
-{
-	USoftClassProperty* This = CastChecked<USoftClassProperty>(InThis);
-	Collector.AddReferencedObject(This->MetaClass, This);
-	Super::AddReferencedObjects(This, Collector);
-}
-
-IMPLEMENT_CORE_INTRINSIC_CLASS(USoftClassProperty, USoftObjectProperty,
-	{
-		UE::GC::DeclareIntrinsicMembers(Class, { UE_GC_MEMBER(USoftClassProperty, MetaClass) });
-	}
-);
-
-IMPLEMENT_CORE_INTRINSIC_CLASS(USoftObjectProperty, UObjectPropertyBase,
-	{
-	}
-);
-
-IMPLEMENT_CORE_INTRINSIC_CLASS(UStrProperty, UProperty,
-	{
-	}
-);
-
 
 static inline void PreloadInnerStructMembers(UStructProperty* StructProperty)
 {
@@ -870,55 +652,3 @@ void UStructProperty::Serialize(FArchive& Ar)
 	}
 }
 
-void UStructProperty::AddReferencedObjects(UObject* InThis, FReferenceCollector& Collector)
-{
-	UStructProperty* This = CastChecked<UStructProperty>(InThis);
-	Collector.AddReferencedObject(This->Struct, This);
-	Super::AddReferencedObjects(This, Collector);
-}
-
-IMPLEMENT_CORE_INTRINSIC_CLASS(UStructProperty, UProperty,
-	{
-		UE::GC::DeclareIntrinsicMembers(Class, { UE_GC_MEMBER(UStructProperty, Struct) });
-	}
-);
-
-IMPLEMENT_CORE_INTRINSIC_CLASS(UUInt16Property, UNumericProperty,
-	{
-	}
-);
-
-IMPLEMENT_CORE_INTRINSIC_CLASS(UUInt32Property, UNumericProperty,
-	{
-	}
-);
-
-IMPLEMENT_CORE_INTRINSIC_CLASS(UUInt64Property, UNumericProperty,
-	{
-	}
-);
-
-IMPLEMENT_CORE_INTRINSIC_CLASS(UWeakObjectProperty, UObjectPropertyBase,
-	{
-	}
-);
-
-IMPLEMENT_CORE_INTRINSIC_CLASS(UTextProperty, UProperty,
-	{
-	}
-);
-
-IMPLEMENT_CORE_INTRINSIC_CLASS(UPropertyWrapper, UObject,
-	{
-	}
-);
-
-IMPLEMENT_CORE_INTRINSIC_CLASS(UMulticastDelegatePropertyWrapper, UPropertyWrapper,
-	{
-	}
-);
-
-IMPLEMENT_CORE_INTRINSIC_CLASS(UMulticastInlineDelegatePropertyWrapper, UMulticastDelegatePropertyWrapper,
-	{
-	}
-);
diff --git a/Engine/Source/Runtime/CoreUObject/Private/UObject/TestUninitializedScriptStructMembersTest.h b/Engine/Source/Runtime/CoreUObject/Private/UObject/TestUninitializedScriptStructMembersTest.h
new file mode 100644
--- /dev/null
+++ b/Engine/Source/Runtime/CoreUObject/Private/UObject/TestUninitializedScriptStructMembersTest.h
@@ -0,0 +1,25 @@
+// Copyright Epic Games, Inc. All Rights Reserved.
+
+#pragma once
+
+#include "UObject/ObjectPtr.h"
+
+#include "TestUninitializedScriptStructMembersTest.generated.h"
+
+class UObject;
+
+// Helper struct to test if member initialization tests work properly
+USTRUCT()
+struct FTestUninitializedScriptStructMembersTest
+{
+    GENERATED_BODY()
+
+	UPROPERTY(Transient)
+	TObjectPtr<UObject> UninitializedObjectReference{NoInit};
+
+	UPROPERTY(Transient)
+	TObjectPtr<UObject> InitializedObjectReference = nullptr;
+
+	UPROPERTY(Transient)
+	float UnusedValue;
+};
diff --git a/Engine/Source/Runtime/CoreUObject/Public/UObject/Class.h b/Engine/Source/Runtime/CoreUObject/Public/UObject/Class.h
--- a/Engine/Source/Runtime/CoreUObject/Public/UObject/Class.h
+++ b/Engine/Source/Runtime/CoreUObject/Public/UObject/Class.h
@@ -75,9 +75,17 @@
 #include "UObject/UObjectGlobals.h"
 #include "UObject/UnrealNames.h"
 #include "UObject/ObjectPtr.h"
+#include "UObject/ObjectPtrLinkedListBuilder.h"
+#include "UObject/ObjectPtrLinkedListIterator.h"
 #include <type_traits>
 #include <atomic>
 
+#if UE_ENABLE_INCLUDE_ORDER_DEPRECATED_IN_5_5
+#include "Templates/IsTriviallyDestructible.h"
+#endif
+
+#include "Class.generated.h"
+
 class FArchive;
 class FEditPropertyChain;
 class FField;
@@ -149,17 +157,25 @@ struct FRepRecord
 //
 // Base class of reflection data objects.
 //
+UCLASS(MinimalAPI, Abstract, Config = Engine)
 class UField : public UObject
 {
-	DECLARE_CASTED_CLASS_INTRINSIC_WITH_API(UField, UObject, CLASS_Abstract, TEXT("/Script/CoreUObject"), CASTCLASS_UField, COREUOBJECT_API)
+	GENERATED_BODY()
+public:
 
 	typedef UField BaseFieldClass;
 	typedef UClass FieldTypeClass;
 
 	/** Next Field in the linked list */
-	UField*			Next;
+	UPROPERTY(SkipSerialization)
+	TObjectPtr<UField>			Next;
 
 	// Constructors.
+	UField(const FObjectInitializer& ObjectInitializer) 
+	: Super(ObjectInitializer)
+	{
+	}
+
 	UField(EStaticConstructor, EObjectFlags InFlags);
 
 	// UObject interface.
@@ -344,7 +360,7 @@ class UField : public UObject
 	COREUOBJECT_API virtual void SetAssociatedFField(FField* InField);
 #endif // WITH_EDITORONLY_DATA
 
-	using FLinkedListBuilder = TLinkedListBuilder<UField>;
+	using FLinkedListBuilder = TObjectPtrLinkedListBuilder<UField>;
 };
 
 #if WITH_EDITORONLY_DATA
@@ -395,21 +411,25 @@ enum class EStructStateFlags : uint16;
 /**
  * Base class for all UObject types that contain fields.
  */
+UCLASS(MinimalAPI, Config = Engine)
 class UStruct : public UField
 #if USTRUCT_FAST_ISCHILDOF_IMPL == USTRUCT_ISCHILDOF_STRUCTARRAY
 	, private FStructBaseChain
 #endif
 {
-	DECLARE_CASTED_CLASS_INTRINSIC_WITH_API(UStruct, UField, CLASS_MatchedSerializers, TEXT("/Script/CoreUObject"), CASTCLASS_UStruct, COREUOBJECT_API)
+	GENERATED_BODY()
 
 	// Variables.
 protected:
 	friend struct Z_Construct_UClass_UStruct_Statics;
 private:
 	/** Struct this inherits from, may be null */
-	ObjectPtr_Private::TNonAccessTrackedObjectPtr<UStruct> SuperStruct;
+	UPROPERTY(SkipSerialization)
+	TObjectPtr<UStruct> SuperStruct;
+
 public:
 	/** Pointer to start of linked list of child fields */
+	UPROPERTY(SkipSerialization)
 	TObjectPtr<UField> Children;
 	
 	/** Pointer to start of linked list of child fields */
@@ -439,6 +459,7 @@ public:
 	FProperty* PostConstructLink;
 
 	/** Array of object references embedded in script code and referenced by FProperties. Mirrored for easy access by realtime garbage collection code */
+	UPROPERTY(SkipSerialization)
 	TArray<TObjectPtr<UObject>> ScriptAndPropertyObjectReferences;
 
 	typedef TArray<TPair<TFieldPath<FField>, int32>> FUnresolvedScriptPropertiesArray;
@@ -447,6 +468,7 @@ public:
 
 #if WITH_EDITORONLY_DATA
 	/** List of wrapper UObjects for FProperties */
+	UPROPERTY(SkipSerialization)
 	TArray<TObjectPtr<UPropertyWrapper>> PropertyWrappers;
 	/** Unique id incremented each time this class properties get destroyed */
 	int32 FieldPathSerialNumber;
@@ -470,7 +492,14 @@ protected:
 #endif
 
 public:
+	using FSuperStructIterator = TObjectPtrLinkedListIterator<UStruct, &UStruct::SuperStruct, true>;
+	using FConstSuperStructIterator = TObjectPtrLinkedListIteratorConst<UStruct, &UStruct::SuperStruct, true>;
+
 	// Constructors.
+	UStruct(const FObjectInitializer& ObjectInitializer) 
+	: Super(ObjectInitializer)
+	{
+	}
 	COREUOBJECT_API UStruct( EStaticConstructor, int32 InSize, int32 InAlignment, EObjectFlags InFlags );
 	COREUOBJECT_API explicit UStruct(UStruct* InSuperStruct, SIZE_T ParamsSize = 0, SIZE_T Alignment = 0);
 	COREUOBJECT_API explicit UStruct(const FObjectInitializer& ObjectInitializer, UStruct* InSuperStruct, SIZE_T ParamsSize = 0, SIZE_T Alignment = 0 );
@@ -668,6 +697,17 @@ public:
 	}
 #endif
 
+	/** Returns an iterate to traverse the super chain without doing TObjectPtr access tracking, starting with this struct */
+	FSuperStructIterator GetSuperStructIterator()
+	{
+		return FSuperStructIterator(this);
+	}
+	/** Returns an iterate to traverse the super chain without doing TObjectPtr access tracking, starting with this struct */
+	FConstSuperStructIterator GetSuperStructIterator() const
+	{
+		return FConstSuperStructIterator(this);
+	}
+
 	/** Struct this inherits from, may be null */
 	UStruct* GetSuperStruct() const
 	{
@@ -963,8 +1003,10 @@ enum EStructFlags
 /**
  * Reflection data for a standalone structure declared in a header or as a user defined struct
  */
+UCLASS(MinimalAPI, Config = Engine)
 class UScriptStruct : public UStruct
 {
+	GENERATED_BODY()
 public:
 	/** Interface to template to manage dynamic access to C++ struct construction and destruction **/
 	struct ICppStructOps
@@ -1764,8 +1806,6 @@ public:
 	#define UE_IMPLEMENT_STRUCT(PackageNameText, BaseName) \
 		static UScriptStruct::TAutoCppStructOps<F##BaseName> BaseName##_Ops(FTopLevelAssetPath(TEXT(PackageNameText), TEXT(#BaseName))); 
 
-	DECLARE_CASTED_CLASS_INTRINSIC_NO_CTOR(UScriptStruct, UStruct, CLASS_MatchedSerializers, TEXT("/Script/CoreUObject"), CASTCLASS_UScriptStruct, COREUOBJECT_API)
-
 	COREUOBJECT_API UScriptStruct( EStaticConstructor, int32 InSize, int32 InAlignment, EObjectFlags InFlags );
 	COREUOBJECT_API explicit UScriptStruct(const FObjectInitializer& ObjectInitializer, UScriptStruct* InSuperStruct, ICppStructOps* InCppStructOps = nullptr, EStructFlags InStructFlags = STRUCT_NoFlags, SIZE_T ExplicitSize = 0, SIZE_T ExplicitAlignment = 0);
 	COREUOBJECT_API explicit UScriptStruct(const FObjectInitializer& ObjectInitializer = FObjectInitializer::Get());
@@ -2014,10 +2054,10 @@ public:
 //
 // Reflection data for a replicated or Kismet callable function.
 //
+UCLASS(MinimalAPI, Within=Class, Config = Engine)
 class UFunction : public UStruct
 {
-	DECLARE_CASTED_CLASS_INTRINSIC_WITH_API(UFunction, UStruct, 0, TEXT("/Script/CoreUObject"), CASTCLASS_UFunction, COREUOBJECT_API)
-	DECLARE_WITHIN(UClass)
+	GENERATED_BODY()
 public:
 	// Persistent variables.
 
@@ -2088,6 +2128,10 @@ public:
 	COREUOBJECT_API void Invoke(UObject* Obj, FFrame& Stack, RESULT_DECL);
 
 	// Constructors.
+	UFunction(const FObjectInitializer& ObjectInitializer) 
+	: Super(ObjectInitializer)
+	{
+	}
 	COREUOBJECT_API explicit UFunction(const FObjectInitializer& ObjectInitializer, UFunction* InSuperFunction, EFunctionFlags InFunctionFlags = FUNC_None, SIZE_T ParamsSize = 0 );
 	COREUOBJECT_API explicit UFunction(UFunction* InSuperFunction, EFunctionFlags InFunctionFlags = FUNC_None, SIZE_T ParamsSize = 0);
 
@@ -2233,12 +2277,14 @@ ENUM_CLASS_FLAGS(EGetByNameFlags)
 //
 // Reflection data for an enumeration.
 //
+UCLASS(MinimalAPI, Config = Engine)
 class UEnum : public UField
 {
-	DECLARE_CASTED_CLASS_INTRINSIC_NO_CTOR(UEnum, UField, 0, TEXT("/Script/CoreUObject"), CASTCLASS_UEnum, COREUOBJECT_API)
-	COREUOBJECT_API UEnum(const FObjectInitializer& ObjectInitialzer);
+	GENERATED_BODY()
 
 public:
+	COREUOBJECT_API UEnum(const FObjectInitializer& ObjectInitialzer);
+
 	/** How this enum is declared in C++, affects the internal naming of enum values */
 	enum class ECppForm : uint8
 	{
@@ -2974,10 +3020,10 @@ enum class EGetSparseClassDataMethod : uint8
 /**
  * An object class.
  */
+UCLASS(MinimalAPI, Within=Package, Config = Engine)
 class UClass : public UStruct
 {
-	DECLARE_CASTED_CLASS_INTRINSIC_NO_CTOR(UClass, UStruct, 0, TEXT("/Script/CoreUObject"), CASTCLASS_UClass, COREUOBJECT_API)
-	DECLARE_WITHIN_UPACKAGE()
+	GENERATED_BODY()
 
 public:
 	friend class FRestoreClassInfo;
@@ -3012,10 +3058,12 @@ public:
 	EClassCastFlags ClassCastFlags;
 
 	/** The required type for the outer of instances of this class */
+	UPROPERTY(SkipSerialization)
 	TObjectPtr<UClass> ClassWithin;
 
 #if WITH_EDITORONLY_DATA
 	/** This is the blueprint that caused the generation of this class, or null if it is a native compiled-in class */
+	UPROPERTY(SkipSerialization)
 	TObjectPtr<UObject> ClassGeneratedBy;
 
 	/** Linked list of properties to be destroyed when this class is destroyed that couldn't be destroyed in PurgeClass **/
@@ -3041,7 +3089,8 @@ public:
 	TArray<FRepRecord> ClassReps;
 
 	/** List of network relevant fields (functions) */
-	TArray<UField*> NetFields;
+	UPROPERTY(SkipSerialization)
+	TArray<TObjectPtr<UField>> NetFields;
 
 #if WITH_EDITOR
 	// Editor only properties
@@ -3096,6 +3145,7 @@ public:
 
 	/** The class default object; used for delta serialization and object initialization */
 	UE_DEPRECATED(5.6, "ClassDefaultObject will be made private in the next release. For general purpose access to the CDO consider using the global GetDefault<>() or GetMutableDefault<>() functions.")
+	UPROPERTY(SkipSerialization)
 	TObjectPtr<UObject> ClassDefaultObject;
 
 protected:
@@ -4321,14 +4371,6 @@ struct FStructUtils
 #endif // !(UE_BUILD_TEST || UE_BUILD_SHIPPING)
 };
 
-// Helper struct to test if member initialization tests work properly
-struct FTestUninitializedScriptStructMembersTest
-{
-	UObject* UninitializedObjectReference;
-	UObject* InitializedObjectReference = nullptr;
-	float UnusedValue;
-};
-
 /*-----------------------------------------------------------------------------
 	Mirrors of mirror structures in Object.h. These are used by generated code 
 	to facilitate correct offsets and alignments for structures containing these
@@ -4529,19 +4571,8 @@ struct FSoftObjectPath;
 struct FSoftClassPath;
 
 struct FPrimaryAssetType;
-
-template<> struct TBaseStructure<FPrimaryAssetType>
-{
-	static COREUOBJECT_API UScriptStruct* Get();
-};
-
 struct FPrimaryAssetId;
 
-template<> struct TBaseStructure<FPrimaryAssetId>
-{
-	static COREUOBJECT_API UScriptStruct* Get();
-};
-
 struct FDateTime;
 
 template<> struct TBaseStructure<FDateTime>
@@ -4557,12 +4588,6 @@ template<> struct TBaseStructure<FPolyglotTextData>
 };
 
 struct FAssetBundleData;
-
-template<> struct TBaseStructure<FTestUninitializedScriptStructMembersTest>
-{
-	static COREUOBJECT_API UScriptStruct* Get();
-};
-
 struct FTopLevelAssetPath;
 
 template<> struct TBaseStructure<FRemoteObjectId>
@@ -4656,6 +4681,3 @@ UE_DECLARE_CORE_VARIANT_TYPE(Sphere3, Sphere);
 
 #undef UE_DECLARE_CORE_VARIANT_TYPE
 
-#if UE_ENABLE_INCLUDE_ORDER_DEPRECATED_IN_5_5
-#include "Templates/IsTriviallyDestructible.h"
-#endif
diff --git a/Engine/Source/Runtime/CoreUObject/Public/UObject/NoExportTypes.h b/Engine/Source/Runtime/CoreUObject/Public/UObject/NoExportTypes.h
--- a/Engine/Source/Runtime/CoreUObject/Public/UObject/NoExportTypes.h
+++ b/Engine/Source/Runtime/CoreUObject/Public/UObject/NoExportTypes.h
@@ -2275,34 +2275,6 @@ struct FMusicalTime
 	int32 TicksPerBeat;
 };
 
-/** 
- * A type of primary asset, used by the Asset Manager system.
- * @note The full C++ class is located here: Engine\Source\Runtime\CoreUObject\Public\UObject\PrimaryAssetId.h
- */
-USTRUCT(noexport, BlueprintType, IsAlwaysAccessible, HasDefaults)
-struct FPrimaryAssetType
-{
-	/** The Type of this object, by default its base class's name */
-	UPROPERTY(EditAnywhere, SaveGame, BlueprintReadWrite, Category = PrimaryAssetType)
-	FName Name;
-};
-
-/** 
- * This identifies an object as a "primary" asset that can be searched for by the AssetManager and used in various tools
- * @note The full C++ class is located here: Engine\Source\Runtime\CoreUObject\Public\UObject\PrimaryAssetId.h
- */
-USTRUCT(noexport, BlueprintType, IsAlwaysAccessible, HasDefaults)
-struct FPrimaryAssetId
-{
-	/** The Type of this object, by default its base class's name */
-	UPROPERTY(EditAnywhere, SaveGame, BlueprintReadWrite, Category = PrimaryAssetId)
-	FPrimaryAssetType PrimaryAssetType;
-
-	/** The Name of this object, by default its short name */
-	UPROPERTY(EditAnywhere, SaveGame, BlueprintReadWrite, Category = PrimaryAssetId)
-	FName PrimaryAssetName;
-};
-
 /** Enumerates the valid types of range bounds (mirrored from Engine/Source/Runtime/Core/Public/Math/RangeBound.h) */
 UENUM(BlueprintType)
 namespace ERangeBoundTypes
@@ -2792,429 +2764,6 @@ enum class EDataValidationResult : uint8
 	NotValidated
 };
 
-/** 
- * @note Mirrored from Engine/Source/Runtime/CoreUObject/Public/UObject/Class.h
- */
-USTRUCT(noexport, IsAlwaysAccessible, HasDefaults)
-struct FTestUninitializedScriptStructMembersTest
-{
-	UPROPERTY(Transient)
-	TObjectPtr<UObject> UninitializedObjectReference;
-
-	UPROPERTY(Transient)
-	TObjectPtr<UObject> InitializedObjectReference;
-
-	UPROPERTY(Transient)
-	float UnusedValue;
-};
-
-/**
- * Direct base class for all UE objects
- * @note The full C++ class is located here: Engine\Source\Runtime\CoreUObject\Public\UObject\Object.h
- */
-UCLASS(abstract, noexport, MatchedSerializers)
-class UObject
-{
-	GENERATED_BODY()
-public:
-
-	UObject(const FObjectInitializer& ObjectInitializer = FObjectInitializer::Get());
-	UObject(FVTableHelper& Helper);
-	~UObject();
-	
-	/**
-	 * Executes some portion of the ubergraph.
-	 *
-	 * @param	EntryPoint	The entry point to start code execution at.
-	 */
-	UFUNCTION(BlueprintImplementableEvent, meta=(BlueprintInternalUseOnly = "true"))
-	void ExecuteUbergraph(int32 EntryPoint);
-};
-
-/**
- * @note The full C++ class is located here: Engine\Source\Runtime\CoreUObject\Public\UObject\Class.h
- */
-
-UCLASS(abstract, noexport, intrinsic, Config = Engine)
-class UField : public UObject
-{
-	GENERATED_BODY()
-public:
-	UField(const FObjectInitializer& ObjectInitializer = FObjectInitializer::Get());
-};
-
-UCLASS(noexport, intrinsic, Config = Engine)
-class UEnum : public UField
-{
-	GENERATED_BODY()
-public:
-	UEnum(const FObjectInitializer& ObjectInitialzer);
-};
-
-UCLASS(noexport, intrinsic, MatchedSerializers, Config = Engine)
-class UStruct : public UField
-{
-	GENERATED_BODY()
-public:
-	UStruct(const FObjectInitializer& ObjectInitializer = FObjectInitializer::Get());
-};
-
-UCLASS(noexport, intrinsic, MatchedSerializers, Config = Engine)
-class UScriptStruct : public UStruct
-{
-	GENERATED_BODY()
-public:
-	UScriptStruct(const FObjectInitializer& ObjectInitializer = FObjectInitializer::Get());
-};
-
-UCLASS(noexport, intrinsic, Config = Engine)
-class UFunction : public UStruct
-{
-	GENERATED_BODY()
-public:
-	UFunction(const FObjectInitializer& ObjectInitializer = FObjectInitializer::Get());
-};
-
-UCLASS(noexport, intrinsic, within=Package, Config = Engine)
-class UClass : public UStruct
-{
-	GENERATED_BODY()
-public:
-	UClass(const FObjectInitializer& ObjectInitializer = FObjectInitializer::Get());
-};
-
-/**
- * @note The full C++ class is located here: Engine\Source\Runtime\CoreUObject\Public\UObject\Package.h
- */
-
-UCLASS(noexport, intrinsic, Config = Engine)
-class UPackage : public UObject
-{
-	GENERATED_BODY()
-};
-
-/**
- * @note The full C++ class is located here: Engine\Source\Runtime\CoreUObject\Public\UObject\UnrealTypePrivate.h
- */
-
-UCLASS(noexport, intrinsic, abstract, Config = Engine)
-class UProperty : public UField
-{
-	GENERATED_BODY()
-public:
-	UProperty(const FObjectInitializer& ObjectInitializer = FObjectInitializer::Get());
-};
-
-UCLASS(noexport, intrinsic, Config = Engine)
-class UStructProperty : public UProperty
-{
-	GENERATED_BODY()
-public:
-	UStructProperty(const FObjectInitializer& ObjectInitializer = FObjectInitializer::Get());
-};
-
-UCLASS(noexport, intrinsic, Config = Engine)
-class UEnumProperty : public UProperty
-{
-	GENERATED_BODY()
-public:
-	UEnumProperty(const FObjectInitializer& ObjectInitializer = FObjectInitializer::Get());
-};
-
-UCLASS(noexport, intrinsic, Config = Engine)
-class UBoolProperty : public UProperty
-{
-	GENERATED_BODY()
-public:
-	UBoolProperty(const FObjectInitializer& ObjectInitializer = FObjectInitializer::Get());
-};
-
-UCLASS(abstract, noexport, intrinsic, Config = Engine)
-class UNumericProperty : public UProperty
-{
-	GENERATED_BODY()
-public:
-	UNumericProperty(const FObjectInitializer& ObjectInitializer = FObjectInitializer::Get());
-};
-
-UCLASS(noexport, intrinsic, Config = Engine)
-class UIntProperty : public UNumericProperty
-{
-	GENERATED_BODY()
-public:
-	UIntProperty(const FObjectInitializer& ObjectInitializer = FObjectInitializer::Get());
-};
-
-UCLASS(noexport, intrinsic, Config = Engine)
-class UInt8Property : public UNumericProperty
-{
-	GENERATED_BODY()
-public:
-	UInt8Property(const FObjectInitializer& ObjectInitializer = FObjectInitializer::Get());
-};
-
-UCLASS(noexport, intrinsic, Config = Engine)
-class UInt16Property : public UNumericProperty
-{
-	GENERATED_BODY()
-public:
-	UInt16Property(const FObjectInitializer& ObjectInitializer = FObjectInitializer::Get());
-};
-
-UCLASS(noexport, intrinsic, Config = Engine)
-class UInt64Property : public UNumericProperty
-{
-	GENERATED_BODY()
-public:
-	UInt64Property(const FObjectInitializer& ObjectInitializer = FObjectInitializer::Get());
-};
-
-UCLASS(noexport, intrinsic, Config = Engine)
-class UByteProperty : public UNumericProperty
-{
-	GENERATED_BODY()
-public:
-	UByteProperty(const FObjectInitializer& ObjectInitializer = FObjectInitializer::Get());
-};
-
-UCLASS(noexport, intrinsic, Config = Engine)
-class UUint16Property : public UNumericProperty
-{
-	GENERATED_BODY()
-public:
-	UUint16Property(const FObjectInitializer& ObjectInitializer = FObjectInitializer::Get());
-};
-
-UCLASS(noexport, intrinsic, Config = Engine)
-class UUint32Property : public UNumericProperty
-{
-	GENERATED_BODY()
-public:
-	UUint32Property(const FObjectInitializer& ObjectInitializer = FObjectInitializer::Get());
-};
-
-UCLASS(noexport, intrinsic, Config = Engine)
-class UUint64Property : public UNumericProperty
-{
-	GENERATED_BODY()
-public:
-	UUint64Property(const FObjectInitializer& ObjectInitializer = FObjectInitializer::Get());
-};
-
-UCLASS(noexport, intrinsic, Config = Engine)
-class UFloatProperty : public UNumericProperty
-{
-	GENERATED_BODY()
-public:
-	UFloatProperty(const FObjectInitializer& ObjectInitializer = FObjectInitializer::Get());
-};
-
-UCLASS(noexport, intrinsic, Config = Engine)
-class UDoubleProperty : public UNumericProperty
-{
-	GENERATED_BODY()
-public:
-	UDoubleProperty(const FObjectInitializer& ObjectInitializer = FObjectInitializer::Get());
-};
-
-UCLASS(abstract, noexport, intrinsic, Config = Engine)
-class UObjectPropertyBase : public UProperty
-{
-	GENERATED_BODY()
-public:
-	UObjectPropertyBase(const FObjectInitializer& ObjectInitializer = FObjectInitializer::Get());
-};
-
-UCLASS(noexport, intrinsic, Config = Engine)
-class UObjectProperty : public UObjectPropertyBase
-{
-	GENERATED_BODY()
-public:
-	UObjectProperty(const FObjectInitializer& ObjectInitializer = FObjectInitializer::Get());
-};
-
-UCLASS(noexport, intrinsic, Config = Engine)
-class UClassProperty : public UObjectProperty
-{
-	GENERATED_BODY()
-public:
-	UClassProperty(const FObjectInitializer& ObjectInitializer = FObjectInitializer::Get());
-};
-
-UCLASS(noexport, intrinsic, Config = Engine)
-class USoftObjectProperty : public UObjectPropertyBase
-{
-	GENERATED_BODY()
-public:
-	USoftObjectProperty(const FObjectInitializer& ObjectInitializer = FObjectInitializer::Get());
-};
-
-UCLASS(noexport, intrinsic, Config = Engine)
-class USoftClassProperty : public USoftObjectProperty
-{
-	GENERATED_BODY()
-public:
-	USoftClassProperty(const FObjectInitializer& ObjectInitializer = FObjectInitializer::Get());
-};
-
-UCLASS(noexport, intrinsic, Config = Engine)
-class UWeakObjectProperty : public UObjectPropertyBase
-{
-	GENERATED_BODY()
-public:
-	UWeakObjectProperty(const FObjectInitializer& ObjectInitializer = FObjectInitializer::Get());
-};
-
-UCLASS(noexport, intrinsic, Config = Engine)
-class ULazyObjectProperty : public UObjectPropertyBase
-{
-	GENERATED_BODY()
-public:
-	ULazyObjectProperty(const FObjectInitializer& ObjectInitializer = FObjectInitializer::Get());
-};
-
-UCLASS(noexport, intrinsic, Config = Engine)
-class UInterfaceProperty : public UProperty
-{
-	GENERATED_BODY()
-public:
-	UInterfaceProperty(const FObjectInitializer& ObjectInitializer = FObjectInitializer::Get());
-};
-
-UCLASS(noexport, intrinsic, Config = Engine)
-class UDelegateProperty : public UProperty
-{
-	GENERATED_BODY()
-public:
-	UDelegateProperty(const FObjectInitializer& ObjectInitializer = FObjectInitializer::Get());
-};
-
-UCLASS(abstract, noexport, intrinsic, Config = Engine)
-class UMulticastDelegateProperty : public UProperty
-{
-	GENERATED_BODY()
-public:
-	UMulticastDelegateProperty(const FObjectInitializer& ObjectInitializer = FObjectInitializer::Get());
-};
-
-UCLASS(noexport, intrinsic, Config = Engine)
-class UMulticastInlineDelegateProperty : public UMulticastDelegateProperty
-{
-	GENERATED_BODY()
-public:
-	UMulticastInlineDelegateProperty(const FObjectInitializer& ObjectInitializer = FObjectInitializer::Get());
-};
-
-UCLASS(noexport, intrinsic, Config = Engine)
-class UMulticastSparseDelegateProperty : public UMulticastDelegateProperty
-{
-	GENERATED_BODY()
-public:
-	UMulticastSparseDelegateProperty(const FObjectInitializer& ObjectInitializer = FObjectInitializer::Get());
-};
-
-UCLASS(noexport, intrinsic, Config = Engine)
-class UNameProperty : public UProperty
-{
-	GENERATED_BODY()
-public:
-	UNameProperty(const FObjectInitializer& ObjectInitializer = FObjectInitializer::Get());
-};
-
-UCLASS(noexport, intrinsic, Config = Engine)
-class UStrProperty : public UProperty
-{
-	GENERATED_BODY()
-public:
-	UStrProperty(const FObjectInitializer& ObjectInitializer = FObjectInitializer::Get());
-};
-
-UCLASS(noexport, intrinsic, Config = Engine)
-class UTextProperty : public UProperty
-{
-	GENERATED_BODY()
-public:
-	UTextProperty(const FObjectInitializer& ObjectInitializer = FObjectInitializer::Get());
-};
-
-UCLASS(noexport, intrinsic, Config = Engine)
-class UArrayProperty : public UProperty
-{
-	GENERATED_BODY()
-public:
-	UArrayProperty(const FObjectInitializer& ObjectInitializer = FObjectInitializer::Get());
-};
-
-UCLASS(noexport, intrinsic, Config = Engine)
-class USetProperty : public UProperty
-{
-	GENERATED_BODY()
-public:
-	USetProperty(const FObjectInitializer& ObjectInitializer = FObjectInitializer::Get());
-};
-
-UCLASS(noexport, intrinsic, Config = Engine)
-class UMapProperty : public UProperty
-{
-	GENERATED_BODY()
-public:
-	UMapProperty(const FObjectInitializer& ObjectInitializer = FObjectInitializer::Get());
-};
-
-/** 
- * Structure for file paths that are displayed in the editor with a picker UI. 
- * @note The full C++ class is located here: Engine\Source\Runtime\CoreUObject\Public\UObject\SoftObjectPath.h
- */
-USTRUCT(noexport, BlueprintType)
-struct FFilePath
-{
-	GENERATED_BODY()
-
-	/**
-	* The path to the file.
-	*/
-	UPROPERTY(EditAnywhere, BlueprintReadWrite, Category = FilePath)
-	FString FilePath;
-};
-
-/** 
- * Structure for directory paths that are displayed in the editor with a picker UI. 
- * @note The full C++ class is located here: Engine\Source\Runtime\CoreUObject\Public\UObject\SoftObjectPath.h
- */
-USTRUCT(noexport, BlueprintType)
-struct FDirectoryPath
-{
-	GENERATED_BODY()
-
-	/**
-	* The path to the directory.
-	*/
-	UPROPERTY(EditAnywhere, BlueprintReadWrite, Category = Path)
-	FString Path;
-};
-
-/** 
- * Structure for templated strings that are displayed in the editor with a allowed args. 
- * @note The full C++ class is located here: Engine\Source\Runtime\CoreUObject\Public\UObject\TemplateString.h
- */
-USTRUCT(noexport, BlueprintType)
-struct FTemplateString
-{
-	GENERATED_BODY()
-
-	/**
-	* The format string.
-	*/
-	UPROPERTY(EditAnywhere, BlueprintReadWrite, Category = Template)
-	FString Template;
-	
-	/**
-	* The resolved (localizable) string.
-	*/
-	UPROPERTY(VisibleAnywhere, BlueprintReadWrite, Category = Template)
-	FText Resolved;
-};
-
 /** 
  * @note The full C++ class is located here: Engine\Source\Runtime\CoreUObject\Public\UObject\RemoteObjectTypes.h
  */
diff --git a/Engine/Source/Runtime/CoreUObject/Public/UObject/Object.h b/Engine/Source/Runtime/CoreUObject/Public/UObject/Object.h
--- a/Engine/Source/Runtime/CoreUObject/Public/UObject/Object.h
+++ b/Engine/Source/Runtime/CoreUObject/Public/UObject/Object.h
@@ -14,6 +14,8 @@
 #include "UObject/PrimaryAssetId.h"
 #include "Containers/VersePathFwd.h"
 
+#include "Object.generated.h"
+
 struct FAssetData;
 class FAssetRegistryTagsContext;
 namespace UE::AssetRegistry { class FAssetRegistryImpl; }
@@ -90,23 +92,13 @@ ENUM_CLASS_FLAGS(EEditChangePropagationFlags);
  * 
  * @see https://docs.unrealengine.com/ProgrammingAndScripting/ProgrammingWithCPP/UnrealArchitecture/Objects
  */
+UCLASS(Abstract, MinimalAPI, Config = Engine)
 class UObject : public UObjectBaseUtility
 {
-	// Declarations, normally created by UnrealHeaderTool boilerplate code
-	DECLARE_CLASS(UObject,UObject,CLASS_Abstract|CLASS_Intrinsic|CLASS_MatchedSerializers,CASTCLASS_None,TEXT("/Script/CoreUObject"),COREUOBJECT_API)
-	DEFINE_DEFAULT_OBJECT_INITIALIZER_CONSTRUCTOR_CALL(UObject)
+	GENERATED_BODY()
+
+public:
 	typedef UObject WithinClass;
-	static UObject* __VTableCtorCaller(FVTableHelper& Helper)
-	{
-		return new (EC_InternalUseOnlyConstructor, (UObject*)GetTransientPackage(), NAME_None, RF_NeedLoad | RF_ClassDefaultObject | RF_TagGarbageTemp) UObject(Helper);
-	}
-	static const TCHAR* StaticConfigName() 
-	{
-		return TEXT("Engine");
-	}
-	static void StaticRegisterNativesUObject() 
-	{
-	}
 
 	/** Default constructor */
 	COREUOBJECT_API UObject();
@@ -930,9 +922,9 @@ public:
 
 #if WITH_EDITOR
 
-	UE_DEPRECATED(5.5, "This function is no longer called. If you require this functionality, please implement ThreadedPostLoadAssetRegistryTagsOverride. \
-Note that this function MUST be safe to call on a background thread. If you have functionality that doesn't meet that requirement, \
-consider using the AssetRegistry's AssetAddedEvent which fires on the game thread.")
+	UE_DEPRECATED(5.5, "This function is no longer called. If you require this functionality, please implement ThreadedPostLoadAssetRegistryTagsOverride. " \
+		"Note that this function MUST be safe to call on a background thread. If you have functionality that doesn't meet that requirement, " \
+		"consider using the AssetRegistry's AssetAddedEvent which fires on the game thread.")
 	virtual void PostLoadAssetRegistryTags(const FAssetData& InAssetData, TArray<FAssetRegistryTag>& OutTagsAndValuesToUpdate) const final {}
 
 	/** Contains an AssetData referencing the asset to be processed and an array of tags and values to be updated on that asset */
@@ -1742,19 +1734,14 @@ public:
 	DECLARE_FUNCTION(execAutoRtfmStopTransact);
 	DECLARE_FUNCTION(execAutoRtfmAbortIfNot);
 
-	/** Wrapper struct to hold the entrypoint in the right memory address */
-	struct Object_eventExecuteUbergraph_Parms
-	{
-		int32 EntryPoint;
-	};
 
-	/** Execute the ubergraph from a specific entry point */
-	void ExecuteUbergraph(int32 EntryPoint)
-	{
-		Object_eventExecuteUbergraph_Parms Parms;
-		Parms.EntryPoint=EntryPoint;
-		ProcessEvent(FindFunctionChecked(NAME_ExecuteUbergraph),&Parms);
-	}
+	/**
+	 * Executes some portion of the ubergraph.
+	 *
+	 * @param	EntryPoint	The entry point to start code execution at.
+	 */
+	UFUNCTION(BlueprintImplementableEvent, meta=(BlueprintInternalUseOnly = "true"))
+	void ExecuteUbergraph(int32 EntryPoint);
 
 protected: 
 	/** Checks it's ok to perform subobjects check at this time. */
diff --git a/Engine/Source/Runtime/CoreUObject/Public/UObject/ObjectPtr.h b/Engine/Source/Runtime/CoreUObject/Public/UObject/ObjectPtr.h
--- a/Engine/Source/Runtime/CoreUObject/Public/UObject/ObjectPtr.h
+++ b/Engine/Source/Runtime/CoreUObject/Public/UObject/ObjectPtr.h
@@ -769,6 +769,16 @@ public:
 		Slot << ObjectPtr;
 	}
 
+	TObjectPtr<T>& GetAccessTrackedObjectPtr() 
+	{
+		return *this;
+	}
+
+	const TObjectPtr<T>& GetAccessTrackedObjectPtr() const
+	{
+		return *this;
+	}
+
 #if UE_WITH_REMOTE_OBJECT_HANDLE
 	[[nodiscard]] FORCEINLINE FRemoteObjectId GetRemoteId() const
 	{
diff --git a/Engine/Source/Runtime/CoreUObject/Public/UObject/ObjectPtrLinkedListBuilder.h b/Engine/Source/Runtime/CoreUObject/Public/UObject/ObjectPtrLinkedListBuilder.h
new file mode 100644
--- /dev/null
+++ b/Engine/Source/Runtime/CoreUObject/Public/UObject/ObjectPtrLinkedListBuilder.h
@@ -0,0 +1,147 @@
+// Copyright Epic Games, Inc. All Rights Reserved.
+
+#pragma once
+
+#include "Containers/LinkedListBuilder.h"
+#include "UObject/ObjectPtr.h"
+
+template<typename InElementType>
+struct TObjectPtrLinkedListBuilderNextLink
+{
+	using ElementType = InElementType;
+
+	[[nodiscard]] FORCEINLINE static auto GetNextPtr(ElementType& Element)
+	{
+		return &(Element.Next);
+	}
+};
+
+/**
+ * Single linked list builder for some cases where we want to migrate from raw ptrs to TObjectPtr
+ * In the existing cases, the start ptr is TObjectPtr and the links are raw pointers. 
+ * We wish to migrate the links to TObjectPtr.
+ */
+template<typename InElementType, typename InLinkAccessor = TObjectPtrLinkedListBuilderNextLink<InElementType>>
+struct TObjectPtrLinkedListBuilder
+{
+public:
+	using ElementType = InElementType;
+	using LinkAccessor = InLinkAccessor;
+
+	UE_NONCOPYABLE(TObjectPtrLinkedListBuilder);
+
+private:
+	FORCEINLINE void WriteEndPtr(InElementType* NewValue)
+	{
+		// Do not overwrite the same value to avoid dirtying the cache and
+		// also prevent TSAN from thinking we are messing around with existing data.
+		if (*EndPtr != NewValue)
+		{
+			*EndPtr = NewValue;
+		}
+	}
+public:
+
+    // Allow construction from raw ptr because call sites were handling the type mismatch between
+    // the head and next ptrs in lists by converting the TObjectPtr to a raw ptr with MutableView.
+	[[nodiscard]] explicit TObjectPtrLinkedListBuilder(ElementType** ListStartPtr) 
+        : StartPtr(reinterpret_cast<TObjectPtr<ElementType>*>(ListStartPtr))
+        , EndPtr(StartPtr)
+	{
+		check(ListStartPtr);
+	}
+
+	[[nodiscard]] explicit TObjectPtrLinkedListBuilder(TObjectPtr<ElementType>* ListStartPtr) 
+        : StartPtr(ListStartPtr)
+        , EndPtr(ListStartPtr)
+	{
+		check(ListStartPtr);
+	}
+
+	// Move builder back to start and prepare for overwriting
+	// It only changes state of builder, use NullTerminate() to mark list as empty!
+	FORCEINLINE void Restart()
+	{
+		EndPtr = StartPtr;
+	}
+
+	// Append element, don't touch next link
+	FORCEINLINE void AppendNoTerminate(ElementType& Element)
+	{
+		WriteEndPtr(&Element);
+		EndPtr = LinkAccessor::GetNextPtr(Element);
+	}
+
+	// Append element and mark it as last
+	FORCEINLINE void AppendTerminated(ElementType& Element)
+	{
+		AppendNoTerminate(Element);
+		NullTerminate();
+	}
+
+	FORCEINLINE void Remove(ElementType& Element)
+	{
+		TObjectPtr<ElementType>* PrevIt = StartPtr;
+		for (TObjectPtr<ElementType> It = *StartPtr; It; It = GetNext(*It))
+		{
+			if (It == &Element)
+			{
+				TObjectPtr<ElementType>& ElementNextRef = GetNextRef(Element);
+				*PrevIt = ElementNextRef;
+				ElementNextRef = nullptr;
+				break;
+			}
+
+			PrevIt = &It;
+		}
+	}
+
+	// Mark end of the list
+	FORCEINLINE void NullTerminate()
+	{
+		WriteEndPtr(nullptr);
+	}
+
+	FORCEINLINE void MoveToEnd()
+	{
+		for (TObjectPtr<ElementType> It = *StartPtr; It; It = GetNext(*It))
+		{
+			EndPtr = LinkAccessor::GetNextPtr(*It);
+		}
+	}
+
+	FORCEINLINE bool MoveToNext()
+	{
+		if (*EndPtr)
+		{
+			EndPtr = LinkAccessor::GetNextPtr(**EndPtr);
+			return true;
+		}
+
+		return false;
+	}
+
+	[[nodiscard]] FORCEINLINE ElementType* GetNext(ElementType& Element) const
+	{
+		return GetNextRef(Element);
+	}
+
+	[[nodiscard]] FORCEINLINE ElementType* GetListStart() const
+	{
+		return *StartPtr;
+	}
+
+	[[nodiscard]] FORCEINLINE ElementType* GetListEnd() const
+	{
+		return *EndPtr;
+	}
+
+private:
+	[[nodiscard]] FORCEINLINE TObjectPtr<ElementType>& GetNextRef(ElementType& Element) const
+	{
+		return *LinkAccessor::GetNextPtr(Element);
+	}
+
+	TObjectPtr<ElementType>* StartPtr;
+	TObjectPtr<ElementType>* EndPtr;
+};
\ No newline at end of file
diff --git a/Engine/Source/Runtime/CoreUObject/Public/UObject/ObjectPtrLinkedListIterator.h b/Engine/Source/Runtime/CoreUObject/Public/UObject/ObjectPtrLinkedListIterator.h
new file mode 100644
--- /dev/null
+++ b/Engine/Source/Runtime/CoreUObject/Public/UObject/ObjectPtrLinkedListIterator.h
@@ -0,0 +1,89 @@
+// Copyright Epic Games, Inc. All Rights Reserved.
+
+#pragma once
+
+#include "UObject/ObjectPtr.h"
+
+/** Utility type to iterate a linked list of TObjectPtr, optionally avoiding access tracking for performance reasons. */
+template<typename InElementType, TObjectPtr<InElementType> InElementType::*InNextMember, bool bConst, bool bNoAccessTracking>
+struct TObjectPtrLinkedListIteratorBase
+{
+    using ElementType = std::conditional_t<bConst, const InElementType, InElementType>;
+
+    TObjectPtrLinkedListIteratorBase() = default;
+    TObjectPtrLinkedListIteratorBase(ElementType* InStart)
+        : Current(InStart)
+    {
+    }
+    TObjectPtrLinkedListIteratorBase(const TObjectPtrLinkedListIteratorBase&) = default;
+    TObjectPtrLinkedListIteratorBase(TObjectPtrLinkedListIteratorBase&&) = default;
+    TObjectPtrLinkedListIteratorBase& operator=(const TObjectPtrLinkedListIteratorBase&) = default;
+    TObjectPtrLinkedListIteratorBase& operator=(TObjectPtrLinkedListIteratorBase&&) = default;
+
+    TObjectPtrLinkedListIteratorBase& operator++()
+    {
+        return MoveNext();
+    }
+
+    TObjectPtrLinkedListIteratorBase operator++(int)
+    {
+        return MoveNext();
+    }
+
+    explicit operator bool()
+    {
+        return !!Current;
+    }
+
+    ElementType* operator*() const
+    {
+        return Current;
+    }
+
+    ElementType* operator->() const
+    {
+        return Current;
+    }
+
+    bool operator==(const TObjectPtrLinkedListIteratorBase& Other) const
+    {
+        return Current == Other.Current;
+    }
+
+#if !PLATFORM_COMPILER_HAS_GENERATED_COMPARISON_OPERATORS
+    bool operator!=(const TObjectPtrLinkedListIterator& Other) const
+    {
+        return Current != Other.Current;
+    }
+#endif
+
+    TObjectPtrLinkedListIteratorBase begin()
+    {
+        return *this;
+    }
+    TObjectPtrLinkedListIteratorBase end()
+    {
+        return TObjectPtrLinkedListIteratorBase{};
+    }
+private:
+    TObjectPtrLinkedListIteratorBase& MoveNext()
+    {
+        if constexpr (bNoAccessTracking)
+        {
+            Current = ObjectPtr_Private::Friend::NoAccessTrackingGet(Current->*InNextMember);
+        }
+        else 
+        {
+            Current = Current->*InNextMember.Get();
+        }
+        return *this;
+    }
+
+    ElementType* Current = nullptr;
+};
+
+template<typename InElementType, TObjectPtr<InElementType> InElementType::*InNextMember, bool bNoAccessTracking>
+using TObjectPtrLinkedListIterator = TObjectPtrLinkedListIteratorBase<InElementType, InNextMember, /* bConst*/ false, bNoAccessTracking>;
+
+template<typename InElementType, TObjectPtr<InElementType> InElementType::*InNextMember, bool bNoAccessTracking>
+using TObjectPtrLinkedListIteratorConst = TObjectPtrLinkedListIteratorBase<InElementType, InNextMember, /* bConst*/ true, bNoAccessTracking>;
\ No newline at end of file
diff --git a/Engine/Source/Runtime/CoreUObject/Public/UObject/Package.h b/Engine/Source/Runtime/CoreUObject/Public/UObject/Package.h
--- a/Engine/Source/Runtime/CoreUObject/Public/UObject/Package.h
+++ b/Engine/Source/Runtime/CoreUObject/Public/UObject/Package.h
@@ -41,6 +41,8 @@
 #include "Cooker/BuildResultDependenciesMap.h"
 #endif
 
+#include "Package.generated.h"
+
 class Error;
 class FArchive;
 class FLinkerLoad;
@@ -209,9 +211,11 @@ enum class EAssetAccessSpecifier : uint8
 * A package.
 */
 PRAGMA_DISABLE_DEPRECATION_WARNINGS // Required for auto-generated functions referencing PackageFlags
+UCLASS(MinimalAPI, Config=Engine)
 class UPackage : public UObject
 {
-	DECLARE_CASTED_CLASS_INTRINSIC_NO_CTOR_NO_VTABLE_CTOR(UPackage, UObject, 0, TEXT("/Script/CoreUObject"), CASTCLASS_UPackage, COREUOBJECT_API)
+	GENERATED_BODY()
+
 	/** DO NOT USE. This constructor is for internal usage only for hot-reload purposes. */
 	UPackage(FVTableHelper& Helper)
 		: Super(Helper)
@@ -222,7 +226,7 @@ class UPackage : public UObject
 
 public:
 
-	UPackage(const FObjectInitializer& ObjectInitializer = FObjectInitializer::Get())
+	COREUOBJECT_API UPackage(const FObjectInitializer& ObjectInitializer = FObjectInitializer::Get())
 		: Super(ObjectInitializer)
 	{
 	}
diff --git a/Engine/Source/Runtime/CoreUObject/Public/UObject/PrimaryAssetId.h b/Engine/Source/Runtime/CoreUObject/Public/UObject/PrimaryAssetId.h
--- a/Engine/Source/Runtime/CoreUObject/Public/UObject/PrimaryAssetId.h
+++ b/Engine/Source/Runtime/CoreUObject/Public/UObject/PrimaryAssetId.h
@@ -10,8 +10,11 @@
 #include "Serialization/StructuredArchive.h"
 #include "Templates/TypeHash.h"
 #include "UObject/NameTypes.h"
+#include "UObject/StructOpsTypeTraits.h"
 #include "UObject/UnrealNames.h"
 
+#include "PrimaryAssetId.generated.h"
+
 class FCbWriter;
 class FOutputDevice;
 class UObject;
@@ -20,8 +23,11 @@ class UObject;
  * A primary asset type, represented as an FName internally and implicitly convertible back and forth
  * This exists so the blueprint API can understand it's not a normal FName
  */
+USTRUCT(BlueprintType)
 struct FPrimaryAssetType
 {
+	GENERATED_BODY()
+
 	/** Convert from FName */
 	FPrimaryAssetType() = default;
 	FPrimaryAssetType(FName InName) : Name(InName) {}
@@ -93,22 +99,40 @@ struct FPrimaryAssetType
 	}
 
 private:
-	friend struct Z_Construct_UScriptStruct_FPrimaryAssetType_Statics;
-
 	/** The FName representing this type */
+	UPROPERTY(EditAnywhere, SaveGame, BlueprintReadWrite, Category = PrimaryAssetType, meta=(AllowPrivateAccess))
 	FName Name;
 };
 
+template<>
+struct TStructOpsTypeTraits<FPrimaryAssetType> : public TStructOpsTypeTraitsBase2<FPrimaryAssetType>
+{
+	enum
+	{
+		WithZeroConstructor = true,
+		WithCopy = true,
+		WithIdenticalViaEquality = true,
+		WithExportTextItem = true,
+		WithImportTextItem = true,
+		WithStructuredSerializeFromMismatchedTag = true,
+	};
+};
+
 /**
  * This identifies an object as a "primary" asset that can be searched for by the AssetManager and used in various tools
  */
+USTRUCT(BlueprintType)
 struct FPrimaryAssetId
 {
+	GENERATED_BODY()
+
 	/** An FName describing the logical type of this object, usually the name of a base UClass. For example, any Blueprint derived from APawn will have a Primary Asset Type of "Pawn".
 	"PrimaryAssetType:PrimaryAssetName" should form a unique name across your project. */
+	UPROPERTY(EditAnywhere, SaveGame, BlueprintReadWrite, Category = PrimaryAssetId)
 	FPrimaryAssetType PrimaryAssetType;
 	/** An FName describing this asset. This is usually the short name of the object, but could be a full asset path for things like maps, or objects with GetPrimaryId() overridden.
 	"PrimaryAssetType:PrimaryAssetName" should form a unique name across your project. */
+	UPROPERTY(EditAnywhere, SaveGame, BlueprintReadWrite, Category = PrimaryAssetId)
 	FName PrimaryAssetName;
 
 	/** Static names to represent the AssetRegistry tags for the above data */
@@ -200,10 +224,22 @@ struct FPrimaryAssetId
 		return Hash;
 	}
 
-	friend struct Z_Construct_UScriptStruct_FPrimaryAssetId_Statics;
-
 private:
 	 COREUOBJECT_API void SerializeForLog(FCbWriter& Writer, const FPrimaryAssetId& Value);
 };
 
+template<>
+struct TStructOpsTypeTraits<FPrimaryAssetId> : public TStructOpsTypeTraitsBase2<FPrimaryAssetId>
+{
+	enum
+	{
+		WithZeroConstructor = true,
+		WithCopy = true,
+		WithIdenticalViaEquality = true,
+		WithExportTextItem = true,
+		WithImportTextItem = true,
+		WithStructuredSerializeFromMismatchedTag = true,
+	};
+};
+
 COREUOBJECT_API FStringBuilderBase& operator<<(FStringBuilderBase& Builder, const FPrimaryAssetId& Id);
diff --git a/Engine/Source/Runtime/CoreUObject/Public/UObject/PropertyWrapper.h b/Engine/Source/Runtime/CoreUObject/Public/UObject/PropertyWrapper.h
new file mode 100644
--- /dev/null
+++ b/Engine/Source/Runtime/CoreUObject/Public/UObject/PropertyWrapper.h
@@ -0,0 +1,76 @@
+// Copyright Epic Games, Inc. All Rights Reserved.
+
+#pragma once
+
+#include "UObject/Object.h"
+#include "UObject/ObjectMacros.h"
+#include "UObject/ObjectPtr.h"
+
+#include "PropertyWrapper.generated.h"
+
+#define UE_API COREUOBJECT_API
+
+/**
+ * FProperty wrapper object.
+ * The purpose of this object is to provide a UObject wrapper for native FProperties that can
+ * be used by property editors (grids).
+ * Specialized wrappers can be used to allow specialized editors for specific property types.
+ * Property wrappers are owned by UStruct that owns the property they wrap and are tied to its lifetime
+ * so that weak object pointer functionality works as expected.
+ */
+UCLASS(Transient, MinimalAPI)
+class UPropertyWrapper : public UObject
+{
+	GENERATED_BODY()
+
+public:
+	UPropertyWrapper(const FObjectInitializer& ObjectInitializer)
+	: Super(ObjectInitializer)
+	{
+	}
+protected:
+	/** Cached property object */
+	FProperty* DestProperty;
+public:
+	/** Sets the property this object wraps */
+	void SetProperty(FProperty* InProperty)
+	{
+		DestProperty = InProperty;
+	}
+	/* Gets property wrapped by this object */
+	FProperty* GetProperty()
+	{
+		return DestProperty;
+	}
+	/* Gets property wrapped by this object */
+	const FProperty* GetProperty() const
+	{
+		return DestProperty;
+	}
+};
+
+UCLASS(Transient, MinimalAPI)
+class UMulticastDelegatePropertyWrapper : public UPropertyWrapper
+{
+	GENERATED_BODY()
+
+public:
+	UMulticastDelegatePropertyWrapper(const FObjectInitializer& ObjectInitializer)
+	: Super(ObjectInitializer)
+	{
+	}
+};
+
+UCLASS(Transient, MinimalAPI)
+class UMulticastInlineDelegatePropertyWrapper : public UMulticastDelegatePropertyWrapper
+{
+	GENERATED_BODY()
+
+public:
+	UMulticastInlineDelegatePropertyWrapper(const FObjectInitializer& ObjectInitializer)
+	: Super(ObjectInitializer)
+	{
+	}
+};
+
+#undef UE_API 
\ No newline at end of file
diff --git a/Engine/Source/Runtime/CoreUObject/Public/UObject/SoftObjectPath.h b/Engine/Source/Runtime/CoreUObject/Public/UObject/SoftObjectPath.h
--- a/Engine/Source/Runtime/CoreUObject/Public/UObject/SoftObjectPath.h
+++ b/Engine/Source/Runtime/CoreUObject/Public/UObject/SoftObjectPath.h
@@ -749,20 +749,28 @@ struct FSoftObjectPathSerializationScope
 };
 
 /** Structure for file paths that are displayed in the editor with a picker UI. */
+USTRUCT(BlueprintType)
 struct FFilePath
 {
+	GENERATED_BODY()
+
 	/**
 	 * The path to the file.
 	 */
+	UPROPERTY(EditAnywhere, BlueprintReadWrite, Category = FilePath)
 	FString FilePath;
 };
 
 /** Structure for directory paths that are displayed in the editor with a picker UI. */
+USTRUCT(BlueprintType)
 struct FDirectoryPath
 {
+	GENERATED_BODY()
+
 	/**
 	 * The path to the directory.
 	 */
+	UPROPERTY(EditAnywhere, BlueprintReadWrite, Category = Path)
 	FString Path;
 };
 
diff --git a/Engine/Source/Runtime/CoreUObject/Public/UObject/TemplateString.h b/Engine/Source/Runtime/CoreUObject/Public/UObject/TemplateString.h
--- a/Engine/Source/Runtime/CoreUObject/Public/UObject/TemplateString.h
+++ b/Engine/Source/Runtime/CoreUObject/Public/UObject/TemplateString.h
@@ -5,19 +5,27 @@
 #include "Containers/Array.h"
 #include "Containers/UnrealString.h"
 #include "CoreMinimal.h"
+#include "UObject/ObjectMacros.h"
+
+#include "TemplateString.generated.h"
 
 /** Structure for templated strings that are displayed in the editor with an optional set of valid args. */
+USTRUCT(BlueprintType)
 struct FTemplateString
 {
+	GENERATED_BODY()
+
 	/**
 	* The template string.
 	*/
+	UPROPERTY(EditAnywhere, BlueprintReadWrite, Category = Template)
 	FString Template;
 
 	/**
 	* The (localizable) resolved text, used to persist the result.
 	* This is especially useful when resolution isn't always available, or is static (editor vs. runtime, for example)
 	*/
+	UPROPERTY(VisibleAnywhere, BlueprintReadWrite, Category = Template)
 	FText Resolved;
 
 	/** Returns validity based on brace matching, and if provided, arg presence in ValidArgs. */
diff --git a/Engine/Source/Runtime/CoreUObject/Public/UObject/UnrealType.h b/Engine/Source/Runtime/CoreUObject/Public/UObject/UnrealType.h
--- a/Engine/Source/Runtime/CoreUObject/Public/UObject/UnrealType.h
+++ b/Engine/Source/Runtime/CoreUObject/Public/UObject/UnrealType.h
@@ -7260,50 +7260,6 @@ static_assert(sizeof(bool) == sizeof(uint8), "Bool is not one byte.");
 #define CPP_ARRAY_DIM(ArrayName, ClassName) \
 	(sizeof(((ClassName*)0)->ArrayName) / sizeof(((ClassName*)0)->ArrayName[0]))
 
-
-/**
- * FProperty wrapper object.
- * The purpose of this object is to provide a UObject wrapper for native FProperties that can
- * be used by property editors (grids).
- * Specialized wrappers can be used to allow specialized editors for specific property types.
- * Property wrappers are owned by UStruct that owns the property they wrap and are tied to its lifetime
- * so that weak object pointer functionality works as expected.
- */
-class UPropertyWrapper : public UObject
-{
-	DECLARE_CLASS_INTRINSIC_API(UPropertyWrapper, UObject, CLASS_Transient, TEXT("/Script/CoreUObject"), UE_API);
-
-protected:
-	/** Cached property object */
-	FProperty* DestProperty;
-public:
-	/** Sets the property this object wraps */
-	void SetProperty(FProperty* InProperty)
-	{
-		DestProperty = InProperty;
-	}
-	/* Gets property wrapped by this object */
-	FProperty* GetProperty()
-	{
-		return DestProperty;
-	}
-	/* Gets property wrapped by this object */
-	const FProperty* GetProperty() const
-	{
-		return DestProperty;
-	}
-};
-
-class UMulticastDelegatePropertyWrapper : public UPropertyWrapper
-{
-	DECLARE_CLASS_INTRINSIC_API(UMulticastDelegatePropertyWrapper, UPropertyWrapper, CLASS_Transient, TEXT("/Script/CoreUObject"), UE_API);
-};
-
-class UMulticastInlineDelegatePropertyWrapper : public UMulticastDelegatePropertyWrapper
-{
-	DECLARE_CLASS_INTRINSIC_API(UMulticastInlineDelegatePropertyWrapper, UMulticastDelegatePropertyWrapper, CLASS_Transient, TEXT("/Script/CoreUObject"), UE_API);
-};
-
 #undef UE_API
 
 #if UE_ENABLE_INCLUDE_ORDER_DEPRECATED_IN_5_7
@@ -7319,3 +7275,6 @@ class UMulticastInlineDelegatePropertyWrapper : public UMulticastDelegatePropert
 
 // This should really be included everywhere that FStrProperty is used, but that would require a lot of fixup
 #include "UObject/StrProperty.h"
+
+// Types extracted from UnrealType.h that need to be in their own header because of the rule that .generated.h must be the last included file 
+#include "PropertyWrapper.h"
\ No newline at end of file
diff --git a/Engine/Source/Runtime/CoreUObject/Public/UObject/UnrealTypePrivate.h b/Engine/Source/Runtime/CoreUObject/Public/UObject/UnrealTypePrivate.h
--- a/Engine/Source/Runtime/CoreUObject/Public/UObject/UnrealTypePrivate.h
+++ b/Engine/Source/Runtime/CoreUObject/Public/UObject/UnrealTypePrivate.h
@@ -9,12 +9,14 @@
 #include "CoreMinimal.h"
 #include "UnrealType.h"
 
+#include "UnrealTypePrivate.generated.h"
+
 #define USE_UPROPERTY_LOAD_DEFERRING (USE_CIRCULAR_DEPENDENCY_LOAD_DEFERRING && WITH_EDITORONLY_DATA)
 
+UCLASS(Abstract, Config=Engine)
 class UProperty : public UField
 {
-	DECLARE_CASTED_CLASS_INTRINSIC_NO_CTOR(UProperty, UField, CLASS_Abstract, TEXT("/Script/CoreUObject"), CASTCLASS_FProperty, NO_API)
-	DECLARE_WITHIN(UObject)
+	GENERATED_BODY()
 
 public:
 
@@ -179,9 +181,16 @@ public:
 #endif // WITH_EDITORONLY_DATA
 };
 
+UCLASS(Abstract, Config = Engine)
 class UNumericProperty : public UProperty
 {
-	DECLARE_CASTED_CLASS_INTRINSIC(UNumericProperty, UProperty, CLASS_Abstract, TEXT("/Script/CoreUObject"), CASTCLASS_FNumericProperty)
+	GENERATED_BODY()
+
+public:
+	UNumericProperty(const FObjectInitializer& ObjectInitializer = FObjectInitializer::Get())
+	: Super(ObjectInitializer)
+	{
+	}
 
 	UNumericProperty(ECppProperty, int32 InOffset, EPropertyFlags InFlags)
 		: UProperty(FObjectInitializer::Get(), EC_CppProperty, InOffset, InFlags)
@@ -192,15 +201,22 @@ class UNumericProperty : public UProperty
 	{}
 };
 
+UCLASS(Config = Engine)
 class UByteProperty : public UNumericProperty
 {
-	DECLARE_CASTED_CLASS_INTRINSIC(UByteProperty, UNumericProperty, 0, TEXT("/Script/CoreUObject"), CASTCLASS_FByteProperty)
+	GENERATED_BODY()
 
 public:
 
 	// Variables.
+	UPROPERTY(SkipSerialization)
 	TObjectPtr<UEnum> Enum;
 
+	UByteProperty(const FObjectInitializer& ObjectInitializer = FObjectInitializer::Get())
+	: Super(ObjectInitializer)
+	{
+	}
+
 	UByteProperty(ECppProperty, int32 InOffset, EPropertyFlags InFlags, UEnum* InEnum = nullptr)
 	: UNumericProperty(FObjectInitializer::Get(), EC_CppProperty, InOffset, InFlags)
 	, Enum(InEnum)
@@ -215,13 +231,19 @@ public:
 
 	// UObject interface.
 	COREUOBJECT_API virtual void Serialize( FArchive& Ar ) override;
-	static COREUOBJECT_API void AddReferencedObjects(UObject* InThis, FReferenceCollector& Collector);
 	// End of UObject interface
 };
 
+UCLASS(Config = Engine)
 class UInt8Property : public UNumericProperty
 {
-	DECLARE_CASTED_CLASS_INTRINSIC(UInt8Property, UNumericProperty, 0, TEXT("/Script/CoreUObject"), CASTCLASS_FInt8Property)
+	GENERATED_BODY()
+
+public:
+ 	UInt8Property(const FObjectInitializer& ObjectInitializer = FObjectInitializer::Get())
+	: Super(ObjectInitializer)
+	{
+	}
 
 	UInt8Property(ECppProperty, int32 InOffset, EPropertyFlags InFlags)
 		: UNumericProperty(FObjectInitializer::Get(), EC_CppProperty, InOffset, InFlags)
@@ -234,9 +256,16 @@ class UInt8Property : public UNumericProperty
 	}
 };
 
+UCLASS(Config = Engine)
 class UInt16Property : public UNumericProperty
 {
-	DECLARE_CASTED_CLASS_INTRINSIC(UInt16Property, UNumericProperty, 0, TEXT("/Script/CoreUObject"), CASTCLASS_FInt16Property)
+	GENERATED_BODY()
+
+public:
+ 	UInt16Property(const FObjectInitializer& ObjectInitializer = FObjectInitializer::Get())
+	: Super(ObjectInitializer)
+	{
+	}
 
 	UInt16Property(ECppProperty, int32 InOffset, EPropertyFlags InFlags)
 	: UNumericProperty(FObjectInitializer::Get(), EC_CppProperty, InOffset, InFlags)
@@ -249,9 +278,16 @@ class UInt16Property : public UNumericProperty
 	}
 };
 
+UCLASS(Config = Engine)
 class UIntProperty : public UNumericProperty
 {
-	DECLARE_CASTED_CLASS_INTRINSIC(UIntProperty, UNumericProperty, 0, TEXT("/Script/CoreUObject"), CASTCLASS_FIntProperty)
+	GENERATED_BODY()
+
+public:
+ 	UIntProperty(const FObjectInitializer& ObjectInitializer = FObjectInitializer::Get())
+	: Super(ObjectInitializer)
+	{
+	}
 
 	UIntProperty(ECppProperty, int32 InOffset, EPropertyFlags InFlags)
 	: UNumericProperty(FObjectInitializer::Get(), EC_CppProperty, InOffset, InFlags)
@@ -264,9 +300,16 @@ class UIntProperty : public UNumericProperty
 	}
 };
 
+UCLASS(Config = Engine)
 class UInt64Property : public UNumericProperty
 {
-	DECLARE_CASTED_CLASS_INTRINSIC(UInt64Property, UNumericProperty, 0, TEXT("/Script/CoreUObject"), CASTCLASS_FInt64Property)
+	GENERATED_BODY()
+
+public:
+ 	UInt64Property(const FObjectInitializer& ObjectInitializer = FObjectInitializer::Get())
+	: Super(ObjectInitializer)
+	{
+	}
 
 	UInt64Property(ECppProperty, int32 InOffset, EPropertyFlags InFlags)
 	: UNumericProperty(FObjectInitializer::Get(), EC_CppProperty, InOffset, InFlags)
@@ -279,9 +322,16 @@ class UInt64Property : public UNumericProperty
 	}
 };
 
+UCLASS(Config = Engine)
 class UUInt16Property : public UNumericProperty
 {
-	DECLARE_CASTED_CLASS_INTRINSIC(UUInt16Property, UNumericProperty, 0, TEXT("/Script/CoreUObject"), CASTCLASS_FUInt16Property)
+	GENERATED_BODY()
+
+public:
+ 	UUInt16Property(const FObjectInitializer& ObjectInitializer = FObjectInitializer::Get())
+	: Super(ObjectInitializer)
+	{
+	}
 
 	UUInt16Property(ECppProperty, int32 InOffset, EPropertyFlags InFlags)
 	: UNumericProperty(FObjectInitializer::Get(), EC_CppProperty, InOffset, InFlags)
@@ -294,9 +344,16 @@ class UUInt16Property : public UNumericProperty
 	}
 };
 
+UCLASS(Config = Engine)
 class UUInt32Property : public UNumericProperty
 {
-	DECLARE_CASTED_CLASS_INTRINSIC(UUInt32Property, UNumericProperty, 0, TEXT("/Script/CoreUObject"), CASTCLASS_FUInt32Property)
+	GENERATED_BODY()
+
+public:
+	UUInt32Property(const FObjectInitializer& ObjectInitializer = FObjectInitializer::Get())
+	: Super(ObjectInitializer)
+	{
+	}
 
 	UUInt32Property( const FObjectInitializer& ObjectInitializer, ECppProperty, int32 InOffset, EPropertyFlags InFlags )
 	: UNumericProperty( ObjectInitializer, EC_CppProperty, InOffset, InFlags )
@@ -304,9 +361,16 @@ class UUInt32Property : public UNumericProperty
 	}
 };
 
+UCLASS(Config = Engine)
 class UUInt64Property : public UNumericProperty
 {
-	DECLARE_CASTED_CLASS_INTRINSIC(UUInt64Property, UNumericProperty, 0, TEXT("/Script/CoreUObject"), CASTCLASS_FUInt64Property)
+	GENERATED_BODY()
+
+public:
+	UUInt64Property(const FObjectInitializer& ObjectInitializer = FObjectInitializer::Get())
+	: Super(ObjectInitializer)
+	{
+	}
 
 	UUInt64Property(ECppProperty, int32 InOffset, EPropertyFlags InFlags)
 	: UNumericProperty(FObjectInitializer::Get(), EC_CppProperty, InOffset, InFlags)
@@ -319,9 +383,16 @@ class UUInt64Property : public UNumericProperty
 	}
 };
 
+UCLASS(Config = Engine)
 class UFloatProperty : public UNumericProperty
 {
-	DECLARE_CASTED_CLASS_INTRINSIC(UFloatProperty, UNumericProperty, 0, TEXT("/Script/CoreUObject"), CASTCLASS_FFloatProperty)
+	GENERATED_BODY()
+
+public:
+	UFloatProperty(const FObjectInitializer& ObjectInitializer = FObjectInitializer::Get())
+	: Super(ObjectInitializer)
+	{
+	}
 
 	UFloatProperty(ECppProperty, int32 InOffset, EPropertyFlags InFlags)
 	: UNumericProperty(FObjectInitializer::Get(), EC_CppProperty, InOffset, InFlags)
@@ -334,9 +405,16 @@ class UFloatProperty : public UNumericProperty
 	}
 };
 
+UCLASS(Config=Engine)
 class UDoubleProperty : public UNumericProperty
 {
-	DECLARE_CASTED_CLASS_INTRINSIC(UDoubleProperty, UNumericProperty, 0, TEXT("/Script/CoreUObject"), CASTCLASS_FDoubleProperty)
+	GENERATED_BODY()
+
+public:
+	UDoubleProperty(const FObjectInitializer& ObjectInitializer = FObjectInitializer::Get())
+	: Super(ObjectInitializer)
+	{
+	}
 
 	UDoubleProperty(ECppProperty, int32 InOffset, EPropertyFlags InFlags)
 	: UNumericProperty(FObjectInitializer::Get(), EC_CppProperty, InOffset, InFlags)
@@ -349,9 +427,10 @@ class UDoubleProperty : public UNumericProperty
 	}
 };
 
+UCLASS(Config = Engine)
 class UBoolProperty : public UProperty
 {
-	DECLARE_CASTED_CLASS_INTRINSIC_NO_CTOR(UBoolProperty, UProperty, 0, TEXT("/Script/CoreUObject"), CASTCLASS_FBoolProperty, NO_API)
+	GENERATED_BODY()
 
 	// Variables.
 public:
@@ -416,15 +495,22 @@ public:
 	}
 };
 
+UCLASS(Abstract, Config = Engine)
 class UObjectPropertyBase : public UProperty
 {
-	DECLARE_CASTED_CLASS_INTRINSIC(UObjectPropertyBase, UProperty, CLASS_Abstract, TEXT("/Script/CoreUObject"), CASTCLASS_FObjectPropertyBase)
+	GENERATED_BODY()
 
 public:
 
 	// Variables.
+	UPROPERTY(SkipSerialization)
 	TObjectPtr<class UClass> PropertyClass;
 
+	UObjectPropertyBase(const FObjectInitializer& ObjectInitializer = FObjectInitializer::Get())
+		: Super(ObjectInitializer)
+	{
+	}
+
 	UObjectPropertyBase(ECppProperty, int32 InOffset, EPropertyFlags InFlags, UClass* InClass = NULL)
 		: UProperty(FObjectInitializer::Get(), EC_CppProperty, InOffset, InFlags)
 		, PropertyClass(InClass)
@@ -437,7 +523,6 @@ public:
 
 	// UObject interface
 	COREUOBJECT_API virtual void Serialize( FArchive& Ar ) override;
-	static COREUOBJECT_API void AddReferencedObjects(UObject* InThis, FReferenceCollector& Collector);
 	COREUOBJECT_API virtual void BeginDestroy() override;
 	// End of UObject interface
 
@@ -448,9 +533,16 @@ public:
 #endif // USE_UPROPERTY_LOAD_DEFERRING
 };
 
+UCLASS(Config = Engine)
 class UObjectProperty : public UObjectPropertyBase
 {
-	DECLARE_CASTED_CLASS_INTRINSIC(UObjectProperty, UObjectPropertyBase, 0, TEXT("/Script/CoreUObject"), CASTCLASS_FObjectProperty)
+	GENERATED_BODY()
+
+public:
+	UObjectProperty(const FObjectInitializer& ObjectInitializer = FObjectInitializer::Get())
+		: Super(ObjectInitializer)
+	{
+	}
 
 	UObjectProperty(ECppProperty, int32 InOffset, EPropertyFlags InFlags, UClass* InClass)
 		: UObjectPropertyBase(FObjectInitializer::Get(), EC_CppProperty, InOffset, InFlags, InClass)
@@ -463,9 +555,16 @@ class UObjectProperty : public UObjectPropertyBase
 	}
 };
 
+UCLASS(Config = Engine)
 class UWeakObjectProperty : public UObjectPropertyBase
 {
-	DECLARE_CASTED_CLASS_INTRINSIC(UWeakObjectProperty, UObjectPropertyBase, 0, TEXT("/Script/CoreUObject"), CASTCLASS_FWeakObjectProperty)
+	GENERATED_BODY()
+
+public:
+	UWeakObjectProperty(const FObjectInitializer& ObjectInitializer = FObjectInitializer::Get())
+	: Super(ObjectInitializer)
+	{
+	}
 
 	UWeakObjectProperty(ECppProperty, int32 InOffset, EPropertyFlags InFlags, UClass* InClass)
 	: UObjectPropertyBase(FObjectInitializer::Get(), EC_CppProperty, InOffset, InFlags, InClass)
@@ -478,9 +577,16 @@ class UWeakObjectProperty : public UObjectPropertyBase
 	}
 };
 
+UCLASS(Config = Engine)
 class ULazyObjectProperty : public UObjectPropertyBase
 {
-	DECLARE_CASTED_CLASS_INTRINSIC(ULazyObjectProperty, UObjectPropertyBase, 0, TEXT("/Script/CoreUObject"), CASTCLASS_FLazyObjectProperty)
+	GENERATED_BODY()
+
+public:
+	ULazyObjectProperty(const FObjectInitializer& ObjectInitializer = FObjectInitializer::Get())
+	: Super(ObjectInitializer)
+	{
+	}
 
 	ULazyObjectProperty(ECppProperty, int32 InOffset, EPropertyFlags InFlags, UClass* InClass)
 	: UObjectPropertyBase(FObjectInitializer::Get(), EC_CppProperty, InOffset, InFlags, InClass)
@@ -493,9 +599,16 @@ class ULazyObjectProperty : public UObjectPropertyBase
 	}
 };
 
+UCLASS(Config = Engine)
 class USoftObjectProperty : public UObjectPropertyBase
 {
-	DECLARE_CASTED_CLASS_INTRINSIC(USoftObjectProperty, UObjectPropertyBase, 0, TEXT("/Script/CoreUObject"), CASTCLASS_FSoftObjectProperty)
+	GENERATED_BODY()
+
+public:
+	USoftObjectProperty(const FObjectInitializer& ObjectInitializer = FObjectInitializer::Get())
+	: Super(ObjectInitializer)
+	{
+	}
 
 	USoftObjectProperty(ECppProperty, int32 InOffset, EPropertyFlags InFlags, UClass* InClass)
 	: UObjectPropertyBase(FObjectInitializer::Get(), EC_CppProperty, InOffset, InFlags, InClass)
@@ -506,15 +619,22 @@ class USoftObjectProperty : public UObjectPropertyBase
 	{}
 };
 
+UCLASS(Config = Engine)
 class UClassProperty : public UObjectProperty
 {
-	DECLARE_CASTED_CLASS_INTRINSIC(UClassProperty, UObjectProperty, 0, TEXT("/Script/CoreUObject"), CASTCLASS_FClassProperty)
+	GENERATED_BODY()
 
 public:
 
 	// Variables.
+	UPROPERTY(SkipSerialization)
 	TObjectPtr<class UClass> MetaClass;
 
+	UClassProperty(const FObjectInitializer& ObjectInitializer = FObjectInitializer::Get())
+		: Super(ObjectInitializer)
+	{
+	}
+
 	UClassProperty(ECppProperty, int32 InOffset, EPropertyFlags InFlags, UClass* InMetaClass, UClass* InClassType)
 		: UObjectProperty(FObjectInitializer::Get(), EC_CppProperty, InOffset, InFlags, InClassType ? InClassType : UClass::StaticClass())
 		, MetaClass(InMetaClass)
@@ -529,7 +649,6 @@ public:
 
 	// UObject interface
 	COREUOBJECT_API virtual void Serialize( FArchive& Ar ) override;
-	static COREUOBJECT_API void AddReferencedObjects(UObject* InThis, FReferenceCollector& Collector);
 	COREUOBJECT_API virtual void BeginDestroy() override;
 	// End of UObject interface
 
@@ -549,15 +668,22 @@ public:
 #endif // USE_UPROPERTY_LOAD_DEFERRING
 };
 
+UCLASS(Config = Engine)
 class USoftClassProperty : public USoftObjectProperty
 {
-	DECLARE_CASTED_CLASS_INTRINSIC(USoftClassProperty, USoftObjectProperty, 0, TEXT("/Script/CoreUObject"), CASTCLASS_FSoftClassProperty)
+	GENERATED_BODY()
 
 public:
 
 	// Variables.
+	UPROPERTY(SkipSerialization)
 	TObjectPtr<class UClass> MetaClass;
 
+	USoftClassProperty(const FObjectInitializer& ObjectInitializer = FObjectInitializer::Get())
+		: Super(ObjectInitializer)
+	{
+	}
+
 	USoftClassProperty(ECppProperty, int32 InOffset, EPropertyFlags InFlags, UClass* InMetaClass)
 		: Super(FObjectInitializer::Get(), EC_CppProperty, InOffset, InFlags, UClass::StaticClass())
 		, MetaClass(InMetaClass)
@@ -570,7 +696,6 @@ public:
 
 	// UObject interface
 	COREUOBJECT_API virtual void Serialize( FArchive& Ar ) override;
-	static COREUOBJECT_API void AddReferencedObjects(UObject* InThis, FReferenceCollector& Collector);
 	COREUOBJECT_API virtual void BeginDestroy() override;
 	// End of UObject interface
 
@@ -590,15 +715,22 @@ public:
 #endif // USE_UPROPERTY_LOAD_DEFERRING
 };
 
+UCLASS(Config = Engine)
 class UInterfaceProperty : public UProperty
 {
-	DECLARE_CASTED_CLASS_INTRINSIC(UInterfaceProperty, UProperty, 0, TEXT("/Script/CoreUObject"), CASTCLASS_FInterfaceProperty)
+	GENERATED_BODY()
 
 public:
 
 	/** The native interface class that this interface property refers to */
+	UPROPERTY(SkipSerialization)
 	TObjectPtr<class	UClass>		InterfaceClass;
 
+	UInterfaceProperty(const FObjectInitializer& ObjectInitializer = FObjectInitializer::Get())
+	: Super(ObjectInitializer)
+	{
+	}
+
 	UInterfaceProperty(ECppProperty, int32 InOffset, EPropertyFlags InFlags, UClass* InInterfaceClass)
 	: UProperty(FObjectInitializer::Get(), EC_CppProperty, InOffset, (InFlags & ~CPF_InterfaceClearMask))
 	, InterfaceClass(InInterfaceClass)
@@ -614,7 +746,6 @@ public:
 	// UObject interface
 	COREUOBJECT_API virtual void Serialize( FArchive& Ar ) override;
 	COREUOBJECT_API virtual void BeginDestroy() override;
-	static COREUOBJECT_API void AddReferencedObjects(UObject* InThis, FReferenceCollector& Collector);
 	// End of UObject interface
 
 	/**
@@ -633,11 +764,17 @@ public:
 #endif // USE_UPROPERTY_LOAD_DEFERRING
 };
 
+UCLASS(Config = Engine)
 class UNameProperty : public UProperty
 {
-	DECLARE_CASTED_CLASS_INTRINSIC(UNameProperty, UProperty, 0, TEXT("/Script/CoreUObject"), CASTCLASS_FNameProperty)
+	GENERATED_BODY()
 public:
 
+	UNameProperty(const FObjectInitializer& ObjectInitializer = FObjectInitializer::Get())
+	: Super(ObjectInitializer)
+	{
+	}
+
 	UNameProperty(ECppProperty, int32 InOffset, EPropertyFlags InFlags)
 	: UProperty(FObjectInitializer::Get(), EC_CppProperty, InOffset, InFlags)
 	{
@@ -649,11 +786,17 @@ public:
 	}
 };
 
+UCLASS(Config = Engine)
 class UStrProperty : public UProperty
 {
-	DECLARE_CASTED_CLASS_INTRINSIC(UStrProperty, UProperty, 0, TEXT("/Script/CoreUObject"), CASTCLASS_FStrProperty)
+	GENERATED_BODY()
 public:
 
+	UStrProperty(const FObjectInitializer& ObjectInitializer = FObjectInitializer::Get())
+	: Super(ObjectInitializer)
+	{
+	}
+
 	UStrProperty(ECppProperty, int32 InOffset, EPropertyFlags InFlags)
 	: UProperty(FObjectInitializer::Get(), EC_CppProperty, InOffset, InFlags)
 	{
@@ -665,15 +808,22 @@ public:
 	}
 };
 
+UCLASS(Config = Engine)
 class UArrayProperty : public UProperty
 {
-	DECLARE_CASTED_CLASS_INTRINSIC(UArrayProperty, UProperty, 0, TEXT("/Script/CoreUObject"), CASTCLASS_FArrayProperty)
+	GENERATED_BODY()
 
 public:
 
 	// Variables.
+	UPROPERTY(SkipSerialization)
 	TObjectPtr<UProperty> Inner;
 
+	UArrayProperty(const FObjectInitializer& ObjectInitializer = FObjectInitializer::Get())
+	: Super(ObjectInitializer)
+	{
+	}
+
 	UArrayProperty(ECppProperty, int32 InOffset, EPropertyFlags InFlags)
 	: UProperty(FObjectInitializer::Get(), EC_CppProperty, InOffset, InFlags)
 	{
@@ -686,73 +836,95 @@ public:
 
 	// UObject interface
 	COREUOBJECT_API virtual void Serialize( FArchive& Ar ) override;
-	static COREUOBJECT_API void AddReferencedObjects(UObject* InThis, FReferenceCollector& Collector);
 	// End of UObject interface
 };
 
+UCLASS(Config = Engine)
 class UMapProperty : public UProperty
 {
-	DECLARE_CASTED_CLASS_INTRINSIC(UMapProperty, UProperty, 0, TEXT("/Script/CoreUObject"), CASTCLASS_FMapProperty)
+	GENERATED_BODY()
 
 public:
 
 	// Properties representing the key type and value type of the contained pairs
+	UPROPERTY(SkipSerialization)
 	TObjectPtr<UProperty>       KeyProp;
+	UPROPERTY(SkipSerialization)
 	TObjectPtr<UProperty>       ValueProp;
 	FScriptMapLayout MapLayout;
 
+	UMapProperty(const FObjectInitializer& ObjectInitializer = FObjectInitializer::Get())
+	: Super(ObjectInitializer)
+	{
+	}
 	COREUOBJECT_API UMapProperty(const FObjectInitializer& ObjectInitializer, ECppProperty, int32 InOffset, EPropertyFlags InFlags);
 
 	// UObject interface
 	COREUOBJECT_API virtual void Serialize(FArchive& Ar) override;
-	static COREUOBJECT_API void AddReferencedObjects(UObject* InThis, FReferenceCollector& Collector);
 	// End of UObject interface
 };
 
+UCLASS(Config = Engine)
 class USetProperty : public UProperty
 {
-	DECLARE_CASTED_CLASS_INTRINSIC(USetProperty, UProperty, 0, TEXT("/Script/CoreUObject"), CASTCLASS_FSetProperty)
+	GENERATED_BODY()
 
 public:
 
 	// Properties representing the key type and value type of the contained pairs
+	UPROPERTY(SkipSerialization)
 	TObjectPtr<UProperty>       ElementProp;
 	FScriptSetLayout SetLayout;
 
+	USetProperty(const FObjectInitializer& ObjectInitializer = FObjectInitializer::Get())
+	: Super(ObjectInitializer)
+	{
+	}
 	COREUOBJECT_API USetProperty(const FObjectInitializer& ObjectInitializer, ECppProperty, int32 InOffset, EPropertyFlags InFlags);
 
 	// UObject interface
 	COREUOBJECT_API virtual void Serialize(FArchive& Ar) override;
-	static COREUOBJECT_API void AddReferencedObjects(UObject* InThis, FReferenceCollector& Collector);
 	// End of UObject interface
 };
 
+UCLASS(Config = Engine)
 class UStructProperty : public UProperty
 {
-	DECLARE_CASTED_CLASS_INTRINSIC(UStructProperty, UProperty, 0, TEXT("/Script/CoreUObject"), CASTCLASS_FStructProperty)
+	GENERATED_BODY()
 
 public:
 
 	// Variables.
+	UPROPERTY(SkipSerialization)
 	TObjectPtr<class UScriptStruct> Struct;
 
+	UStructProperty(const FObjectInitializer& ObjectInitializer = FObjectInitializer::Get())
+	: Super(ObjectInitializer)
+	{
+	}
 	COREUOBJECT_API UStructProperty(ECppProperty, int32 InOffset, EPropertyFlags InFlags, UScriptStruct* InStruct);
 	COREUOBJECT_API UStructProperty( const FObjectInitializer& ObjectInitializer, ECppProperty, int32 InOffset, EPropertyFlags InFlags, UScriptStruct* InStruct );
 
 	// UObject interface
 	COREUOBJECT_API virtual void Serialize( FArchive& Ar ) override;
-	static COREUOBJECT_API void AddReferencedObjects(UObject* InThis, FReferenceCollector& Collector);
 	// End of UObject interface
 };
 
+UCLASS(Config = Engine)
 class UDelegateProperty : public UProperty
 {
-	DECLARE_CASTED_CLASS_INTRINSIC(UDelegateProperty, UProperty, 0, TEXT("/Script/CoreUObject"), CASTCLASS_FDelegateProperty)
+	GENERATED_BODY()
 
 public:
 
 	/** Points to the source delegate function (the function declared with the delegate keyword) used in the declaration of this delegate property. */
-	UFunction* SignatureFunction;
+	UPROPERTY(SkipSerialization)
+	TObjectPtr<UFunction> SignatureFunction;
+
+	UDelegateProperty(const FObjectInitializer& ObjectInitializer = FObjectInitializer::Get())
+	: Super(ObjectInitializer)
+	{
+	}
 
 	UDelegateProperty(ECppProperty, int32 InOffset, EPropertyFlags InFlags, UFunction* InSignatureFunction = NULL)
 	: UProperty(FObjectInitializer::Get(), EC_CppProperty, InOffset, InFlags)
@@ -772,14 +944,21 @@ public:
 	// End of UObject interface
 };
 
+UCLASS(Abstract, Config = Engine)
 class UMulticastDelegateProperty : public UProperty
 {
-	DECLARE_CASTED_CLASS_INTRINSIC(UMulticastDelegateProperty, UProperty, CLASS_Abstract, TEXT("/Script/CoreUObject"), CASTCLASS_FMulticastDelegateProperty)
+	GENERATED_BODY()
 
 public:
 
 	/** Points to the source delegate function (the function declared with the delegate keyword) used in the declaration of this delegate property. */
-	UFunction* SignatureFunction;
+	UPROPERTY(SkipSerialization)
+	TObjectPtr<UFunction> SignatureFunction;
+
+	UMulticastDelegateProperty(const FObjectInitializer& ObjectInitializer = FObjectInitializer::Get())
+		: Super(ObjectInitializer)
+	{
+	}
 
 	UMulticastDelegateProperty(const FObjectInitializer& ObjectInitializer, ECppProperty, int32 InOffset, EPropertyFlags InFlags, UFunction* InSignatureFunction = nullptr)
 		: UProperty(ObjectInitializer, EC_CppProperty, InOffset, InFlags)
@@ -793,11 +972,16 @@ public:
 	// End of UObject interface
 };
 
+UCLASS(Config = Engine)
 class UMulticastInlineDelegateProperty : public UMulticastDelegateProperty
 {
-	DECLARE_CASTED_CLASS_INTRINSIC(UMulticastInlineDelegateProperty, UMulticastDelegateProperty, 0, TEXT("/Script/CoreUObject"), CASTCLASS_FMulticastInlineDelegateProperty)
+	GENERATED_BODY()
 
 public:
+	UMulticastInlineDelegateProperty(const FObjectInitializer& ObjectInitializer = FObjectInitializer::Get())
+	: Super(ObjectInitializer)
+	{
+	}
 
 	UMulticastInlineDelegateProperty(ECppProperty, int32 InOffset, EPropertyFlags InFlags, UFunction* InSignatureFunction = nullptr)
 	: UMulticastDelegateProperty(FObjectInitializer::Get(), EC_CppProperty, InOffset, InFlags, InSignatureFunction)
@@ -810,12 +994,18 @@ public:
 	}
 };
 
+UCLASS(Config = Engine)
 class UMulticastSparseDelegateProperty : public UMulticastDelegateProperty
 {
-	DECLARE_CASTED_CLASS_INTRINSIC(UMulticastSparseDelegateProperty, UMulticastDelegateProperty, 0, TEXT("/Script/CoreUObject"), CASTCLASS_FMulticastSparseDelegateProperty)
+	GENERATED_BODY()
 
 public:
 
+	UMulticastSparseDelegateProperty(const FObjectInitializer& ObjectInitializer = FObjectInitializer::Get())
+	: Super(ObjectInitializer)
+	{
+	}
+
 	UMulticastSparseDelegateProperty(ECppProperty, int32 InOffset, EPropertyFlags InFlags, UFunction* InSignatureFunction = nullptr)
 	: UMulticastDelegateProperty(FObjectInitializer::Get(), EC_CppProperty, InOffset, InFlags, InSignatureFunction)
 	{
@@ -827,29 +1017,42 @@ public:
 	}
 };
 
+
+UCLASS(Config = Engine)
 class UEnumProperty : public UProperty
 {
-	DECLARE_CASTED_CLASS_INTRINSIC(UEnumProperty, UProperty, 0, TEXT("/Script/CoreUObject"), CASTCLASS_FEnumProperty)
+	GENERATED_BODY()
 
 public:
+	UEnumProperty(const FObjectInitializer& ObjectInitializer = FObjectInitializer::Get()) 
+	: Super(ObjectInitializer)
+	{
+	}
 	COREUOBJECT_API UEnumProperty(const FObjectInitializer& ObjectInitializer, UEnum* InEnum);
 	COREUOBJECT_API UEnumProperty(const FObjectInitializer& ObjectInitializer, ECppProperty, int32 InOffset, EPropertyFlags InFlags, UEnum* InEnum);
 
 	// UObject interface
 	COREUOBJECT_API virtual void Serialize(FArchive& Ar) override;
-	static COREUOBJECT_API void AddReferencedObjects(UObject* InThis, FReferenceCollector& Collector);
 	// End of UObject interface
 
+	UPROPERTY(SkipSerialization)
 	TObjectPtr<UNumericProperty> UnderlyingProp; // The property which represents the underlying type of the enum
+	UPROPERTY(SkipSerialization)
 	TObjectPtr<UEnum> Enum; // The enum represented by this property
 };
 
+UCLASS(Config = Engine)
 class UTextProperty : public UProperty
 {
-	DECLARE_CASTED_CLASS_INTRINSIC(UTextProperty, UProperty, 0, TEXT("/Script/CoreUObject"), CASTCLASS_FTextProperty)
+	GENERATED_BODY()
 
 public:
 
+	UTextProperty(const FObjectInitializer& ObjectInitializer = FObjectInitializer::Get())
+	: UProperty(ObjectInitializer)
+	{
+	}
+
 	UTextProperty(ECppProperty, int32 InOffset, EPropertyFlags InFlags)
 	: UProperty(FObjectInitializer::Get(), EC_CppProperty, InOffset, InFlags)
 	{
diff --git a/Engine/Source/Runtime/Engine/Classes/Intrinsic/Model.h b/Engine/Source/Runtime/Engine/Classes/Intrinsic/Model.h
--- a/Engine/Source/Runtime/Engine/Classes/Intrinsic/Model.h
+++ b/Engine/Source/Runtime/Engine/Classes/Intrinsic/Model.h
@@ -1,17 +1,3 @@
 // Copyright Epic Games, Inc. All Rights Reserved.
 
 #pragma once
-
-// HEADER_UNIT_SKIP
-
-#include "UObject/ObjectMacros.h"
-
-#if !CPP && !UE_BUILD_DOCS      //noexport class
-UCLASS(noexport, Intrinsic)
-class UModel : public UObject
-{
-	GENERATED_UCLASS_BODY()
-
-};
-
-#endif
diff --git a/Engine/Source/Runtime/Engine/Private/Model.cpp b/Engine/Source/Runtime/Engine/Private/Model.cpp
--- a/Engine/Source/Runtime/Engine/Private/Model.cpp
+++ b/Engine/Source/Runtime/Engine/Private/Model.cpp
@@ -601,21 +601,6 @@ void UModel::GetResourceSizeEx(FResourceSizeEx& CumulativeResourceSize)
 	}
 }
 
-#if WITH_EDITOR
-#define UMODEL_EDITOR_GC_MEMBERS UE_GC_MEMBER(UModel, Polys), 
-#else
-#define UMODEL_EDITOR_GC_MEMBERS
-#endif
-
-IMPLEMENT_INTRINSIC_CLASS(UModel, ENGINE_API, UObject, CORE_API, "/Script/Engine",
-	{
-		Class->CppClassStaticFunctions = UOBJECT_CPPCLASS_STATICFUNCTIONS_FORCLASS(UModel);
-
-		UE::GC::TSchemaBuilder<FBspSurf> SurfSchema({ UE_GC_MEMBER(FBspSurf, Material), UE_GC_MEMBER(FBspSurf, Actor) });
-		UE::GC::DeclareIntrinsicMembers(Class, { UMODEL_EDITOR_GC_MEMBERS UE_GC_MEMBER(UModel, Surfs, SurfSchema) });
-	}
-);
-
 /*---------------------------------------------------------------------------------------
 	UModel implementation.
 ---------------------------------------------------------------------------------------*/
diff --git a/Engine/Source/Runtime/Engine/Public/Model.h b/Engine/Source/Runtime/Engine/Public/Model.h
--- a/Engine/Source/Runtime/Engine/Public/Model.h
+++ b/Engine/Source/Runtime/Engine/Public/Model.h
@@ -7,6 +7,7 @@
 #pragma once
 
 #include "CoreMinimal.h"
+#include "UObject/StructOpsTypeTraits.h"
 #include "UObject/ObjectMacros.h"
 #include "UObject/UObjectGlobals.h"
 #include "UObject/Object.h"
@@ -21,6 +22,8 @@
 #include "Templates/UniquePtr.h"
 #include "StaticMeshResources.h"
 
+#include "Model.generated.h"
+
 class AActor;
 class ABrush;
 class FBlake3;
@@ -203,8 +206,12 @@ struct FLeaf
 // are stored along with Bsp nodes, since several nodes which lie in the
 // same plane may reference the same poly.
 //
+USTRUCT()
 struct FBspSurf
 {
+	GENERATED_BODY()
+
+	UPROPERTY()
 	TObjectPtr<UMaterialInterface>	Material;		// 4 Material.
 	uint32							PolyFlags;		// 4 Polygon flags.
 	int32							pBase;			// 4 Polygon & texture base point index (where U,V==0,0).
@@ -212,6 +219,7 @@ struct FBspSurf
 	int32							vTextureU;		// 4 Texture U-vector index.
 	int32							vTextureV;		// 4 Texture V-vector index.
 	int32							iBrushPoly;		// 4 Editor brush polygon index.
+	UPROPERTY()
 	TObjectPtr<ABrush>				Actor;			// 4 Brush actor owning this Bsp surface.
 	FPlane4f						Plane;			// 16 The plane this surface lies on.
 	float							LightMapScale;	// 4 The number of units/lightmap texel on this surface.
@@ -388,15 +396,16 @@ struct FNodeGroup
 
 enum {MAX_NODES  = 65536};
 enum {MAX_POINTS = 128000};
+UCLASS(MinimalAPI)
 class UModel : public UObject
 {
-	DECLARE_CASTED_CLASS_INTRINSIC_NO_CTOR_NO_VTABLE_CTOR(UModel, UObject, 0, TEXT("/Script/Engine"), CASTCLASS_None, ENGINE_API)
-
-	/** DO NOT USE. This constructor is for internal usage only for hot-reload purposes. */
+	GENERATED_BODY()
+public:
 	UModel(FVTableHelper& Helper);
 
-#if WITH_EDITOR
+#if WITH_EDITORONLY_DATA
 	// Arrays and subobjects.
+	UPROPERTY(SkipSerialization)
 	TObjectPtr<UPolys>						Polys;
 #endif // WITH_EDITOR
 
@@ -404,6 +413,8 @@ class UModel : public UObject
 	TArray<FVert>			Verts;
 	TArray<FVector3f>			Vectors;
 	TArray<FVector3f>			Points;
+
+	UPROPERTY(SkipSerialization)
 	TArray<FBspSurf>		Surfs;
 
 #if WITH_EDITOR
