Fix static analysis warning V568: It's odd that the argument of sizeof() operator is the 'Hash->GetBytes()' expression. The address of GetBytes is the same as the value of GetBytes since it is a c-style array, and passing the address of GetBytes looks more normal to static analysis.
#rnx

diff --git a/Engine/Source/Editor/UnrealEd/Private/EditorDomain/EditorDomainUtils.cpp b/Engine/Source/Editor/UnrealEd/Private/EditorDomain/EditorDomainUtils.cpp
--- a/Engine/Source/Editor/UnrealEd/Private/EditorDomain/EditorDomainUtils.cpp
+++ b/Engine/Source/Editor/UnrealEd/Private/EditorDomain/EditorDomainUtils.cpp
@@ -526,7 +526,7 @@ bool TryAppendClassDigests(const TArrayView<UE::Cook::FCookDependency>& Dependen
 		{
 			ClassDependency.SetValue(*Hash);
 		}
-		Writer.Update(Hash->GetBytes(), sizeof(Hash->GetBytes()));
+		Writer.Update(&Hash->GetBytes(), sizeof(Hash->GetBytes()));
 	}
 
 	return true;
diff --git a/Engine/Source/Runtime/CoreUObject/Private/Cooker/CookDependency.cpp b/Engine/Source/Runtime/CoreUObject/Private/Cooker/CookDependency.cpp
--- a/Engine/Source/Runtime/CoreUObject/Private/Cooker/CookDependency.cpp
+++ b/Engine/Source/Runtime/CoreUObject/Private/Cooker/CookDependency.cpp
@@ -438,7 +438,7 @@ void FCookDependency::UpdateHash(FCookDependencyContext& Context, FBlake3Hash* O
 		{
 			*OutCurrentValue = Hash;
 		}
-		Context.Update(Hash.GetBytes(), sizeof(Hash.GetBytes()));
+		Context.Update(&Hash.GetBytes(), sizeof(Hash.GetBytes()));
 
 		break;
 	}
@@ -476,7 +476,7 @@ void FCookDependency::UpdateHash(FCookDependencyContext& Context, FBlake3Hash* O
 		{
 			*OutCurrentValue = Hash;
 		}
-		Context.Update(Hash.GetBytes(), sizeof(Hash.GetBytes()));
+		Context.Update(&Hash.GetBytes(), sizeof(Hash.GetBytes()));
 
 		return;
 	}
@@ -518,7 +518,7 @@ void FCookDependency::UpdateHash(FCookDependencyContext& Context, FBlake3Hash* O
 		{
 			*OutCurrentValue = Hash;
 		}
-		Context.Update(Hash.GetBytes(), sizeof(Hash.GetBytes()));
+		Context.Update(&Hash.GetBytes(), sizeof(Hash.GetBytes()));
 
 		return;
 	}
@@ -578,7 +578,7 @@ void FCookDependency::UpdateHash(FCookDependencyContext& Context, FBlake3Hash* O
 				{
 					*OutCurrentValue = Hash;
 				}
-				Context.Update(Hash.GetBytes(), sizeof(Hash.GetBytes()));
+				Context.Update(&Hash.GetBytes(), sizeof(Hash.GetBytes()));
 			}
 		}
 		return;
