Fix static analysis warning V568: It's odd that the argument of sizeof() operator is the 'Hash->GetBytes()' expression. The address of GetBytes is the same as the value of GetBytes since it is a c-style array, and passing the address of GetBytes looks more normal to static analysis.
#rnx

diff --git a/Engine/Source/Runtime/CoreUObject/Private/Cooker/CookDependency.cpp b/Engine/Source/Runtime/CoreUObject/Private/Cooker/CookDependency.cpp
--- a/Engine/Source/Runtime/CoreUObject/Private/Cooker/CookDependency.cpp
+++ b/Engine/Source/Runtime/CoreUObject/Private/Cooker/CookDependency.cpp
@@ -608,11 +608,11 @@ FBlake3Hash FCookDependency::ConvertToHash(const FIoHash& Hash)
 	FBlake3Hash OutHash;
 	constexpr size_t OutSize = sizeof(OutHash.GetBytes());
 	static_assert(sizeof(Hash) <= OutSize);
-	FMemory::Memcpy(OutHash.GetBytes(), &Hash, OutSize);
+	FMemory::Memcpy(&OutHash.GetBytes(), &Hash, sizeof(Hash));
 
 	if constexpr (sizeof(Hash) < OutSize)
 	{
-		FMemory::Memzero(OutHash.GetBytes() + sizeof(Hash), OutSize - sizeof(Hash));
+		FMemory::Memzero(((uint8*)&OutHash.GetBytes()) + sizeof(Hash), OutSize - sizeof(Hash));
 	}
 	return OutHash;
 }
@@ -632,11 +632,11 @@ FBlake3Hash FCookDependency::ConvertToHash(const FUtf8String& String)
 	}
 	else
 	{
-		FMemory::Memcpy(OutHash.GetBytes(), String.GetCharArray().GetData(), StringNumBytes);
+		FMemory::Memcpy(&OutHash.GetBytes(), String.GetCharArray().GetData(), StringNumBytes);
 
 		if (StringNumBytes < OutSize)
 		{
-			FMemory::Memzero(OutHash.GetBytes() + StringNumBytes, OutSize - StringNumBytes);
+			FMemory::Memzero(((uint8*) &OutHash.GetBytes()) + StringNumBytes, OutSize - StringNumBytes);
 		}
 	}
 	return OutHash;
