Adding remote object specific versions of FObjectPtr::GetFName(), FObjectPtr::GetPathName(), FObjectPtr:GetOuter() and a few others

diff --git a/Engine/Source/Runtime/CoreUObject/Private/UObject/ObjectPtr.cpp b/Engine/Source/Runtime/CoreUObject/Private/UObject/ObjectPtr.cpp
--- a/Engine/Source/Runtime/CoreUObject/Private/UObject/ObjectPtr.cpp
+++ b/Engine/Source/Runtime/CoreUObject/Private/UObject/ObjectPtr.cpp
@@ -2,6 +2,8 @@
 
 #include "UObject/ObjectPtr.h"
 #include "UObject/Class.h"
+#include "UObject/RemoteObjectPrivate.h"
+#include "UObject/RemoteObjectPathName.h"
 
 #if UE_WITH_OBJECT_HANDLE_LATE_RESOLVE
 
@@ -89,6 +91,162 @@ bool FObjectPtr::IsIn(FObjectPtr SomeOuter) const
 	return ResolveObjectHandleNoRead(LocalHandle)->IsIn(ResolveObjectHandleNoRead(SomeOuterHandle));
 }
 
+#elif UE_WITH_REMOTE_OBJECT_HANDLE
+
+FString FObjectPtr::GetPathName() const
+{
+	FObjectHandle LocalHandle = Handle;
+	if (!IsObjectHandleNull(LocalHandle))
+	{
+		if (IsObjectHandleResolved_ForGC(LocalHandle))
+		{
+			return ReadObjectHandlePointerNoCheck(LocalHandle)->GetPathName();
+		}
+		else
+		{
+			FRemoteObjectPathName RemotePathName(LocalHandle.ToStub()->Id);
+			if (RemotePathName.Num())
+			{
+				return RemotePathName.ToString();
+			}
+			else if (const UObject* ResolvedObject = ResolveObjectHandleNoRead(LocalHandle))
+			{
+				return ResolvedObject->GetPathName();
+			}
+		}
+	}
+	return FString(TEXT("None"));
+}
+
+FName FObjectPtr::GetFName() const
+{
+	using namespace UE::RemoteObject::Private;
+	using namespace UE::RemoteObject::Handle;
+
+	FObjectHandle LocalHandle = Handle;
+	if (!IsObjectHandleNull(LocalHandle))
+	{
+		if (IsObjectHandleResolved_ForGC(LocalHandle))
+		{
+			return ReadObjectHandlePointerNoCheck(LocalHandle)->GetFName();
+		}
+		else if (const FRemoteObjectStub* Stub = LocalHandle.ToStub())
+		{
+			if (!Stub->Name.IsNone())
+			{
+				return Stub->Name;
+			}
+		}
+		// If the above two branches failed to return the name, we need to resolve the object
+		if (const UObject* ResolvedObject = ResolveObjectHandleNoRead(LocalHandle))
+		{
+			return ResolvedObject->GetFName();
+		}
+	}
+	return NAME_None;
+}
+
+FString FObjectPtr::GetFullName(EObjectFullNameFlags Flags) const
+{
+	FObjectHandle LocalHandle = Handle;
+	if (!IsObjectHandleNull(LocalHandle) && IsObjectHandleResolved_ForGC(LocalHandle))
+	{
+		return ReadObjectHandlePointerNoCheck(LocalHandle)->GetFullName(nullptr, Flags);
+	}
+	return UE::CoreUObject::Private::GetFullName(ResolveObjectHandleNoRead(LocalHandle), nullptr, Flags);
+}
+
+FObjectPtr FObjectPtr::GetOuter() const
+{
+	using namespace UE::RemoteObject::Private;
+	using namespace UE::RemoteObject::Handle;
+
+	FObjectHandle LocalHandle = Handle;
+	if (!IsObjectHandleNull(LocalHandle))
+	{
+		if (IsObjectHandleResolved_ForGC(LocalHandle))
+		{
+			return FObjectPtr(ReadObjectHandlePointerNoCheck(LocalHandle)->GetOuter());
+		}
+		else if (const FRemoteObjectStub* Stub = LocalHandle.ToStub())
+		{
+			// Only rely on the OuterId if it's actually valid. Otherwise we have no way of distinguishing between a stub that represnts a package or an incomplete stub
+			if (Stub->OuterId.IsValid())
+			{
+				return FObjectPtr(Stub->OuterId);
+			}
+		}
+		if (const UObject* ResolvedObject = ResolveObjectHandleNoRead(LocalHandle))
+		{
+			return FObjectPtr(ResolvedObject->GetOuter());
+		}
+	}
+	return FObjectPtr();
+}
+
+FObjectPtr FObjectPtr::GetPackage() const
+{
+	using namespace UE::RemoteObject::Private;
+	using namespace UE::RemoteObject::Handle;
+
+	FObjectHandle LocalHandle = Handle;
+	if (!IsObjectHandleNull(LocalHandle))
+	{
+		if (IsObjectHandleResolved_ForGC(LocalHandle))
+		{
+			return FObjectPtr(ReadObjectHandlePointerNoCheck(LocalHandle)->GetPackage());
+		}
+		else if (const FRemoteObjectStub* Stub = LocalHandle.ToStub())
+		{
+			for (const FRemoteObjectStub* StubIt = Stub; StubIt;)
+			{
+				if (StubIt->OuterId.IsValid())
+				{
+					// Object represented by the stub has an outer so try and find it
+					if (UObject* Outer = StaticFindObjectFastInternal(StubIt->OuterId))
+					{
+						return FObjectPtr(Outer->GetPackage());
+					}
+					else
+					{
+						// Or try and find the outer's stub
+						StubIt = FindRemoteObjectStub(StubIt->OuterId);
+					}
+				}
+				else
+				{
+					// Currently we can't assume that all stubs have the OuterId info valid so we're going to have to resolve the object
+					break;
+				}
+			}
+		}
+
+		// If we got here then we were not able to resolve the package from a remote object or its stub so we need to resolve the object
+		if (const UObject* ResolvedObject = ResolveObjectHandleNoRead(LocalHandle))
+		{
+			return FObjectPtr(ResolvedObject->GetPackage());
+		}
+	}
+	return FObjectPtr();
+}
+
+bool FObjectPtr::IsIn(FObjectPtr SomeOuter) const
+{
+	FObjectHandle SomeOuterHandle = SomeOuter.Handle;
+	FObjectHandle LocalHandle = Handle;
+	if (IsObjectHandleNull(LocalHandle) || IsObjectHandleNull(SomeOuterHandle))
+	{
+		return false;
+	}
+	if (IsObjectHandleResolved_ForGC(LocalHandle) && IsObjectHandleResolved_ForGC(SomeOuterHandle))
+	{
+		UObject* Object = ReadObjectHandlePointerNoCheck(LocalHandle);
+		UObject* Outer = ReadObjectHandlePointerNoCheck(SomeOuterHandle);
+		return Object->IsIn(Outer);
+	}
+	return ResolveObjectHandleNoRead(LocalHandle)->IsIn(ResolveObjectHandleNoRead(SomeOuterHandle));
+}
+
 #endif
 
 bool FObjectPtr::IsA(const UClass* SomeBase) const
diff --git a/Engine/Source/Runtime/CoreUObject/Private/UObject/RemoteObjectPathName.cpp b/Engine/Source/Runtime/CoreUObject/Private/UObject/RemoteObjectPathName.cpp
--- a/Engine/Source/Runtime/CoreUObject/Private/UObject/RemoteObjectPathName.cpp
+++ b/Engine/Source/Runtime/CoreUObject/Private/UObject/RemoteObjectPathName.cpp
@@ -131,14 +131,53 @@ FString FPackedRemoteObjectPathName::ToString(const FRemoteObjectTables& InTable
 	return RemotePathNameToString(*this, InTables, InMinPathSegmentIndex);
 }
 
-
+void AppendObjectToPathName(UObject* InObject, FRemoteObjectPathName& OutPathName)
+{
+	for (UObject* PathNameObject = InObject; PathNameObject; PathNameObject = PathNameObject->GetOuter())
+	{
+		OutPathName.Names.Add(PathNameObject->GetFName());
+		OutPathName.RemoteIds.Add(FRemoteObjectId(PathNameObject));
+	}
+}
 
 FRemoteObjectPathName::FRemoteObjectPathName(UObject* InObject)
 {
-	for (UObject* PathNameObject = InObject; PathNameObject; PathNameObject = PathNameObject->GetOuter())
+	AppendObjectToPathName(InObject, *this);
+}
+
+FRemoteObjectPathName::FRemoteObjectPathName(FRemoteObjectId RemoteId)
+{
+	using namespace UE::RemoteObject::Private;
+	using namespace UE::RemoteObject::Handle;
+
+	while (RemoteId.IsValid())
 	{
-		Names.Add(PathNameObject->GetFName());
-		RemoteIds.Add(FRemoteObjectId(PathNameObject));
+		UObject* Object = StaticFindObjectFastInternal(RemoteId);
+		if (Object)
+		{
+			AppendObjectToPathName(Object, *this);
+			// We can abort now because once we encounter an existing object in the outer chain we can fully construct the pathname
+			break;
+		}
+		else
+		{
+			FRemoteObjectStub* Stub = FindRemoteObjectStub(RemoteId);
+			if (Stub && !Stub->Name.IsNone() && Stub->OuterId.IsValid())
+			{
+				Names.Add(Stub->Name);
+				RemoteIds.Add(Stub->Id);
+				RemoteId = Stub->OuterId;
+			}
+			else
+			{
+				// If the object could not be found by its id and the object does not have a stub, we can't fully construct the pathname so reset everything and abort
+				// We also need to early out if the OuterId of a stub is invalid because at the moment we can distinguish between a stub that represents a package or an incomplete stub
+				// and since we don't normally migrate packages (yet) we must assume it's the latter
+				Names.Empty();
+				RemoteIds.Empty();
+				break;
+			}
+		}
 	}
 }
 
diff --git a/Engine/Source/Runtime/CoreUObject/Public/UObject/ObjectPtr.h b/Engine/Source/Runtime/CoreUObject/Public/UObject/ObjectPtr.h
--- a/Engine/Source/Runtime/CoreUObject/Public/UObject/ObjectPtr.h
+++ b/Engine/Source/Runtime/CoreUObject/Public/UObject/ObjectPtr.h
@@ -192,7 +192,7 @@ public:
 	// @TODO OBJPTR: Deprecate this.
 	[[nodiscard]] FString GetPath() const { return GetPathName(); }
 
-#if UE_WITH_OBJECT_HANDLE_LATE_RESOLVE
+#if UE_WITH_OBJECT_HANDLE_LATE_RESOLVE || UE_WITH_REMOTE_OBJECT_HANDLE
 	// Gets the PathName of the object without resolving the object reference.
 	[[nodiscard]] COREUOBJECT_API FString GetPathName() const;
 
diff --git a/Engine/Source/Runtime/CoreUObject/Public/UObject/RemoteObjectPathName.h b/Engine/Source/Runtime/CoreUObject/Public/UObject/RemoteObjectPathName.h
--- a/Engine/Source/Runtime/CoreUObject/Public/UObject/RemoteObjectPathName.h
+++ b/Engine/Source/Runtime/CoreUObject/Public/UObject/RemoteObjectPathName.h
@@ -74,6 +74,7 @@ struct FRemoteObjectPathName : public FRemoteObjectTables
 	FRemoteObjectPathName() = default;
 	FRemoteObjectPathName(FRemoteObjectPathName&&) = default;
 	COREUOBJECT_API explicit FRemoteObjectPathName(UObject* InObject);
+	COREUOBJECT_API explicit FRemoteObjectPathName(FRemoteObjectId RemoteId);
 	FRemoteObjectPathName& operator=(FRemoteObjectPathName&&) = default;
 
 	FRemoteObjectPathName(const FRemoteObjectPathName&) = default;
