Converted some of DelegatePropertyTools::ImportDelegateFromText() to use string views instead of C pointers.

diff --git a/Engine/Source/Runtime/CoreUObject/Private/UObject/PropertyHelper.cpp b/Engine/Source/Runtime/CoreUObject/Private/UObject/PropertyHelper.cpp
--- a/Engine/Source/Runtime/CoreUObject/Private/UObject/PropertyHelper.cpp
+++ b/Engine/Source/Runtime/CoreUObject/Private/UObject/PropertyHelper.cpp
@@ -3,6 +3,8 @@
 #include "UObject/PropertyHelper.h"
 
 #include "Logging/StructuredLog.h"
+#include "Containers/StringView.h"
+#include "Misc/AsciiSet.h"
 #include "String/Find.h"
 #include "String/ParseTokens.h"
 #include "UObject/Class.h"
@@ -134,48 +136,32 @@ const TCHAR* DelegatePropertyTools::ImportDelegateFromText( FScriptDelegate& Del
 {
 	SkipWhitespace(Buffer);
 
-	TCHAR ObjName[NAME_SIZE];
-	TCHAR FuncName[NAME_SIZE];
-
 	// Get object name
-	int32 i;
 	while (*Buffer == TEXT('('))
 	{
 		++Buffer;
 	}
-	for( i=0; *Buffer && *Buffer != TCHAR('.') && *Buffer != TCHAR(')') && *Buffer != TCHAR(',') && i < NAME_SIZE; Buffer++ )
-	{
-		ObjName[i++] = *Buffer;
-	}
-	if (i == NAME_SIZE)
-	{
-		return nullptr;
-	}
-	ObjName[i] = TCHAR('\0');
-	UClass* Cls = nullptr;
-	UObject* Object = nullptr;
-	bool bIsNullClass = false;
+	const TCHAR* NameEnd = FAsciiSet::FindFirstOrEnd(Buffer, ".),");
+	FStringView ObjNameView(Buffer, UE_PTRDIFF_TO_INT32(NameEnd - Buffer));
+	Buffer += ObjNameView.Len();
 
 	// Handle (null).None syntax used for empty delegates in ExportText
-	if (*Buffer == TCHAR(')') && FCString::Strcmp(ObjName, TEXT("null")) == 0)
+	bool bIsNullClass = *Buffer == TCHAR(')') && ObjNameView == TEXT("null");
+	if (bIsNullClass)
 	{
 		++Buffer;
-		bIsNullClass = true;
 	}
 
 	// Get function name
+	FStringView FuncNameView;
+	UClass* Cls = nullptr;
+	UObject* Object = nullptr;
 	if (*Buffer == TCHAR('.'))
 	{
-		Buffer++;
-		for( i=0; *Buffer && *Buffer != TCHAR(')') && *Buffer != TCHAR(',') && i < NAME_SIZE; Buffer++ )
-		{
-			FuncName[i++] = *Buffer;
-		}
-		if (i == NAME_SIZE)
-		{
-			return nullptr;
-		}
-		FuncName[i] = TEXT('\0');
+		++Buffer;
+		const TCHAR* FuncNameEnd = FAsciiSet::FindFirstOrEnd(Buffer, "),");
+		FuncNameView = FStringView(Buffer, UE_PTRDIFF_TO_INT32(FuncNameEnd - Buffer));
+		Buffer += FuncNameView.Len();
 	}
 	else
 	{
@@ -191,42 +177,41 @@ const TCHAR* DelegatePropertyTools::ImportDelegateFromText( FScriptDelegate& Del
 				*DelegateString, *SignatureFunctionString);
 			return nullptr;
 		}
+
 		// since we don't support nested components, we only need to check one level deep
-		else if (!Parent->HasAnyFlags(RF_ClassDefaultObject) && Parent->GetOuter()->HasAnyFlags(RF_ClassDefaultObject))
+		if (!Parent->HasAnyFlags(RF_ClassDefaultObject) && Parent->GetOuter() && Parent->GetOuter()->HasAnyFlags(RF_ClassDefaultObject))
 		{
 			Object = Parent->GetOuter();
-			Cls = Parent->GetOuter()->GetClass();
 		}
 		else
 		{
 			Object = Parent;
-			Cls = Parent->GetClass();
 		}
-		FCString::Strncpy(FuncName, ObjName, i + 1);
+
+		Cls = Object->GetClass();
+		FuncNameView = ObjNameView;
 	}
 
 	if (bIsNullClass)
 	{
-		FName ParsedName(FuncName);
+		FName ParsedName(FuncNameView);
 		if (ParsedName == NAME_None)
 		{
 			// Deliberately null
 			Delegate.Unbind();
 			return Buffer;
 		}
-		else
-		{
-			FString DelegateString = Delegate.GetFunctionName().ToString();
-			FString SignatureFunctionString = GetNameSafe(SignatureFunction);
-			ErrorText->Logf(ELogVerbosity::Warning, TEXT("Cannot import delegate with function name %s but null class. Delegate=%s SignatureFunction=%s"),
-				*FuncName, *DelegateString, *SignatureFunctionString);
-			return nullptr;
-		}
+
+		FString DelegateString = Delegate.GetFunctionName().ToString();
+		FString SignatureFunctionString = GetNameSafe(SignatureFunction);
+		ErrorText->Logf(ELogVerbosity::Warning, TEXT("Cannot import delegate with function name %.*s but null class. Delegate=%s SignatureFunction=%s"),
+			FuncNameView.Len(), FuncNameView.GetData(), *DelegateString, *SignatureFunctionString);
+		return nullptr;
 	}
 
 	if (Cls == nullptr)
 	{
-		Cls = UClass::TryFindTypeSlow<UClass>(ObjName);
+		Cls = UClass::TryFindTypeSlow<UClass>(ObjNameView);
 		if (Cls != nullptr)
 		{
 			Object = Cls->GetDefaultObject();
@@ -237,13 +222,13 @@ const TCHAR* DelegatePropertyTools::ImportDelegateFromText( FScriptDelegate& Del
 			UObject* OuterToCheck = Parent;
 			while (Object == nullptr && OuterToCheck)
 			{
-				Object = StaticFindObject(UObject::StaticClass(), OuterToCheck, ObjName);
+				Object = StaticFindObject(UObject::StaticClass(), OuterToCheck, ObjNameView);
 				OuterToCheck = OuterToCheck->GetOuter();
 			}
 			
 			if (Object == nullptr)
 			{
-				Object = StaticFindFirstObject(UObject::StaticClass(), ObjName, EFindFirstObjectOptions::NativeFirst | EFindFirstObjectOptions::EnsureIfAmbiguous);
+				Object = StaticFindFirstObject(UObject::StaticClass(), ObjNameView, EFindFirstObjectOptions::NativeFirst | EFindFirstObjectOptions::EnsureIfAmbiguous);
 			}
 			if (Object != nullptr)
 			{
@@ -251,9 +236,10 @@ const TCHAR* DelegatePropertyTools::ImportDelegateFromText( FScriptDelegate& Del
 			}
 		}
 	}
-	UFunction *Func = FindUField<UFunction>( Cls, FuncName );
+
 	// Check function params.
-	if(Func != nullptr)
+	UFunction* Func = FindUField<UFunction>( Cls, *FString(FuncNameView) );
+	if (Func)
 	{
 		// Find the delegate UFunction to check params
 		check(SignatureFunction != nullptr && "Invalid delegate property");
@@ -280,10 +266,10 @@ const TCHAR* DelegatePropertyTools::ImportDelegateFromText( FScriptDelegate& Del
 	}
 	else 
 	{
-		ErrorText->Logf(ELogVerbosity::Warning,TEXT("Unable to find function %s in object %s for delegate (found class: %s)"),FuncName,ObjName,*GetNameSafe(Cls));
+		ErrorText->Logf(ELogVerbosity::Warning,TEXT("Unable to find function %.*s in object %.*s for delegate (found class: %s)"),FuncNameView.Len(),FuncNameView.GetData(),ObjNameView.Len(),ObjNameView.GetData(),*GetNameSafe(Cls));
 	}
 
-	//UE_LOG(LogProperty, Log, TEXT("... importing delegate FunctionName:'%s'(%s)   Object:'%s'(%s)"),Func != nullptr ? *Func->GetName() : TEXT("nullptr"), FuncName, Object != nullptr ? *Object->GetFullName() : TEXT("nullptr"), ObjName);
+	//UE_LOG(LogProperty, Log, TEXT("... importing delegate FunctionName:'%s'(%.*s)   Object:'%s'(%.*s)"),Func != nullptr ? *Func->GetName() : TEXT("nullptr"), FuncNameView.Len(), FuncNameView.GetData(), Object != nullptr ? *Object->GetFullName() : TEXT("nullptr"), ObjNameView.Len(), ObjNameView.GetData());
 
 	Delegate.BindUFunction( Func ? Object : nullptr, Func ? Func->GetFName() : NAME_None );
 
