Suppress static analysis warning V568: "It's odd that the argument of sizeof() operator is the 'OutHash.GetByutes()' expression.
#rnx

diff --git a/Engine/Source/Runtime/CoreUObject/Private/Cooker/CookDependency.cpp b/Engine/Source/Runtime/CoreUObject/Private/Cooker/CookDependency.cpp
--- a/Engine/Source/Runtime/CoreUObject/Private/Cooker/CookDependency.cpp
+++ b/Engine/Source/Runtime/CoreUObject/Private/Cooker/CookDependency.cpp
@@ -606,7 +606,7 @@ void FCookDependency::SetValue(const FBlake3Hash& Hash)
 FBlake3Hash FCookDependency::ConvertToHash(const FIoHash& Hash)
 {
 	FBlake3Hash OutHash;
-	constexpr size_t OutSize = sizeof(OutHash.GetBytes());
+	constexpr size_t OutSize = sizeof(OutHash.GetBytes()); // -V568
 	static_assert(sizeof(Hash) <= OutSize);
 	FMemory::Memcpy(&OutHash.GetBytes(), &Hash, sizeof(Hash));
 
@@ -620,7 +620,7 @@ FBlake3Hash FCookDependency::ConvertToHash(const FIoHash& Hash)
 FBlake3Hash FCookDependency::ConvertToHash(const FUtf8String& String)
 {
 	FBlake3Hash OutHash;
-	constexpr size_t OutSize = sizeof(OutHash.GetBytes());
+	constexpr size_t OutSize = sizeof(OutHash.GetBytes()); // -V568
 
 	SIZE_T StringNumBytes = String.NumBytesWithoutNull();
 	if (StringNumBytes > OutSize)
