Flattened nesting in DelegatePropertyTools::ImportDelegateFromText, and ensured the delegate is cleared on error.

diff --git a/Engine/Source/Runtime/CoreUObject/Private/UObject/PropertyHelper.cpp b/Engine/Source/Runtime/CoreUObject/Private/UObject/PropertyHelper.cpp
--- a/Engine/Source/Runtime/CoreUObject/Private/UObject/PropertyHelper.cpp
+++ b/Engine/Source/Runtime/CoreUObject/Private/UObject/PropertyHelper.cpp
@@ -136,6 +136,8 @@ const TCHAR* DelegatePropertyTools::ImportDelegateFromText( FScriptDelegate& Del
 {
 	SkipWhitespace(Buffer);
 
+	Delegate.Unbind();
+
 	// Get object name
 	while (*Buffer == TEXT('('))
 	{
@@ -198,7 +200,6 @@ const TCHAR* DelegatePropertyTools::ImportDelegateFromText( FScriptDelegate& Del
 		if (ParsedName == NAME_None)
 		{
 			// Deliberately null
-			Delegate.Unbind();
 			return Buffer;
 		}
 
@@ -209,10 +210,10 @@ const TCHAR* DelegatePropertyTools::ImportDelegateFromText( FScriptDelegate& Del
 		return nullptr;
 	}
 
-	if (Cls == nullptr)
+	if (!Cls)
 	{
 		Cls = UClass::TryFindTypeSlow<UClass>(ObjNameView);
-		if (Cls != nullptr)
+		if (Cls)
 		{
 			Object = Cls->GetDefaultObject();
 		}
@@ -225,55 +226,54 @@ const TCHAR* DelegatePropertyTools::ImportDelegateFromText( FScriptDelegate& Del
 				Object = StaticFindObject(UObject::StaticClass(), OuterToCheck, ObjNameView);
 				OuterToCheck = OuterToCheck->GetOuter();
 			}
-			
-			if (Object == nullptr)
+
+			if (!Object)
 			{
 				Object = StaticFindFirstObject(UObject::StaticClass(), ObjNameView, EFindFirstObjectOptions::NativeFirst | EFindFirstObjectOptions::EnsureIfAmbiguous);
+				if (!Object)
+				{
+					ErrorText->Logf(ELogVerbosity::Warning,TEXT("Unable to find object %.*s for delegate"), ObjNameView.Len(), ObjNameView.GetData());
+					return nullptr;
+				}
 			}
-			if (Object != nullptr)
-			{
-				Cls = Object->GetClass();
-			}
+			Cls = Object->GetClass();
 		}
 	}
 
 	// Check function params.
 	UFunction* Func = FindUField<UFunction>( Cls, *FString(FuncNameView) );
-	if (Func)
+	if (!Func)
 	{
-		// Find the delegate UFunction to check params
-		check(SignatureFunction != nullptr && "Invalid delegate property");
+		ErrorText->Logf(ELogVerbosity::Warning,TEXT("Unable to find function %.*s in object %.*s for delegate (found class: %s)"),FuncNameView.Len(),FuncNameView.GetData(),ObjNameView.Len(),ObjNameView.GetData(),*GetNameSafe(Cls));
+		return nullptr;
+	}
 
-		// check return type and params
-		if(	Func->NumParms == SignatureFunction->NumParms )
-		{
-			int32 Count=0;
-			for( TFieldIterator<FProperty> It1(Func),It2(SignatureFunction); Count<SignatureFunction->NumParms; ++It1,++It2,++Count )
-			{
-				if( It1->GetClass()!=It2->GetClass() || (It1->PropertyFlags&CPF_OutParm)!=(It2->PropertyFlags&CPF_OutParm) )
-				{
-					ErrorText->Logf(ELogVerbosity::Warning,TEXT("Function %s does not match param types with delegate signature %s"), *Func->GetName(), *SignatureFunction->GetName());
-					Func = nullptr;
-					break;
-				}
-			}
-		}
-		else
+	// Find the delegate UFunction to check params
+	check(SignatureFunction != nullptr && "Invalid delegate property");
+
+	// check return type and params
+	if (Func->NumParms != SignatureFunction->NumParms)
+	{
+		ErrorText->Logf(ELogVerbosity::Warning,TEXT("Function %s does not match number of params with delegate signature %s"),*Func->GetName(), *SignatureFunction->GetName());
+		return nullptr;
+	}
+
+	int32 Count=0;
+	for( TFieldIterator<FProperty> It1(Func),It2(SignatureFunction); Count<SignatureFunction->NumParms; ++It1,++It2,++Count )
+	{
+		if( It1->GetClass()!=It2->GetClass() || (It1->PropertyFlags&CPF_OutParm)!=(It2->PropertyFlags&CPF_OutParm) )
 		{
-			ErrorText->Logf(ELogVerbosity::Warning,TEXT("Function %s does not match number of params with delegate signature %s"),*Func->GetName(), *SignatureFunction->GetName());
+			ErrorText->Logf(ELogVerbosity::Warning,TEXT("Function %s does not match param types with delegate signature %s"), *Func->GetName(), *SignatureFunction->GetName());
 			Func = nullptr;
+			break;
 		}
 	}
-	else 
-	{
-		ErrorText->Logf(ELogVerbosity::Warning,TEXT("Unable to find function %.*s in object %.*s for delegate (found class: %s)"),FuncNameView.Len(),FuncNameView.GetData(),ObjNameView.Len(),ObjNameView.GetData(),*GetNameSafe(Cls));
-	}
 
 	//UE_LOG(LogProperty, Log, TEXT("... importing delegate FunctionName:'%s'(%.*s)   Object:'%s'(%.*s)"),Func != nullptr ? *Func->GetName() : TEXT("nullptr"), FuncNameView.Len(), FuncNameView.GetData(), Object != nullptr ? *Object->GetFullName() : TEXT("nullptr"), ObjNameView.Len(), ObjNameView.GetData());
 
-	Delegate.BindUFunction( Func ? Object : nullptr, Func ? Func->GetFName() : NAME_None );
+	Delegate.BindUFunction(Object, Func->GetFName());
 
-	return ( Func != nullptr && Object != nullptr ) ? Buffer : nullptr;
+	return Buffer;
 }
 
 namespace UE
