Static analysis fix.

#rnx

diff --git a/Engine/Source/Runtime/CoreUObject/Private/UObject/PropertyHelper.cpp b/Engine/Source/Runtime/CoreUObject/Private/UObject/PropertyHelper.cpp
--- a/Engine/Source/Runtime/CoreUObject/Private/UObject/PropertyHelper.cpp
+++ b/Engine/Source/Runtime/CoreUObject/Private/UObject/PropertyHelper.cpp
@@ -264,8 +264,7 @@ const TCHAR* DelegatePropertyTools::ImportDelegateFromText( FScriptDelegate& Del
 		if( It1->GetClass()!=It2->GetClass() || (It1->PropertyFlags&CPF_OutParm)!=(It2->PropertyFlags&CPF_OutParm) )
 		{
 			ErrorText->Logf(ELogVerbosity::Warning,TEXT("Function %s does not match param types with delegate signature %s"), *Func->GetName(), *SignatureFunction->GetName());
-			Func = nullptr;
-			break;
+			return nullptr;
 		}
 	}
 
