[ObjectOverrides] Added missing PreEditContainer removal for arrays.

diff --git a/Engine/Source/Runtime/CoreUObject/Private/UObject/OverriddenPropertySet.cpp b/Engine/Source/Runtime/CoreUObject/Private/UObject/OverriddenPropertySet.cpp
--- a/Engine/Source/Runtime/CoreUObject/Private/UObject/OverriddenPropertySet.cpp
+++ b/Engine/Source/Runtime/CoreUObject/Private/UObject/OverriddenPropertySet.cpp
@@ -1043,6 +1043,14 @@ void FOverriddenPropertySet::NotifyPropertyChange(FOverriddenPropertyNode* Paren
 					}
 					return;
 				}
+				
+				ON_SCOPE_EXIT
+				{
+					if (Notification == EPropertyNotificationType::PostEdit)
+					{
+						SavedPreEditContainers.Free(Owner, ArrayProperty);
+					}
+				};
 
 				FScriptArrayHelper PreEditArrayHelper(ArrayProperty, SavedPreEditContainers.Find(Owner, ArrayProperty));
 
@@ -1275,7 +1283,10 @@ void FOverriddenPropertySet::NotifyPropertyChange(FOverriddenPropertyNode* Paren
 
 			ON_SCOPE_EXIT
 			{
-				SavedPreEditContainers.Free(Owner, MapProperty);
+				if (Notification == EPropertyNotificationType::PostEdit)
+				{
+					SavedPreEditContainers.Free(Owner, MapProperty);
+				}
 			};
 
 			auto MapReplace = [&]()
