[Core] Allow override max printed reference chains

#tests run the engine, trigger reference chain search

diff --git a/Engine/Source/Runtime/CoreUObject/Private/UObject/ReferenceChainSearch.cpp b/Engine/Source/Runtime/CoreUObject/Private/UObject/ReferenceChainSearch.cpp
--- a/Engine/Source/Runtime/CoreUObject/Private/UObject/ReferenceChainSearch.cpp
+++ b/Engine/Source/Runtime/CoreUObject/Private/UObject/ReferenceChainSearch.cpp
@@ -1698,11 +1698,18 @@ int32 FReferenceChainSearch::PrintResults(bool bDumpAllChains /*= false*/, UObje
 	return PrintResults([](FCallbackParams& Params) { return true; }, bDumpAllChains, TargetObject, InVerbosityForPrint);
 }
 
+static TAutoConsoleVariable CVarMaxReferenceChainsToPrint(
+	TEXT("gc.MaxReferenceChainsToPrint"),
+	100,
+	TEXT("Sets the limit of how many reference chains ReferenceChainSearch utility will print"),
+	ECVF_Default
+);
+
 int32 FReferenceChainSearch::PrintResults(TFunctionRef<bool(FCallbackParams& Params)> ReferenceCallback, bool bDumpAllChains /*= false*/, UObject* TargetObject /*= nullptr*/, ELogVerbosity::Type InVerbosityForPrint /*= ELogVerbosity::Log*/) const
 {
 	FSlowHeartBeatScope DisableHangDetection; // This function can be very slow
 
-	const int32 MaxChainsToPrint = 100;
+	const int32 MaxChainsToPrint = CVarMaxReferenceChainsToPrint.GetValueOnAnyThread();
 	int32 NumPrintedChains = 0;
 
 	TMap<uint64, FString> CallstackCache;
