[StructUtils] Change the hash function for UPropertyBag to have a more deterministic value. Helps with deterministic cooking.

#rnx

diff --git a/Engine/Source/Runtime/CoreUObject/Private/StructUtils/PropertyBag.cpp b/Engine/Source/Runtime/CoreUObject/Private/StructUtils/PropertyBag.cpp
--- a/Engine/Source/Runtime/CoreUObject/Private/StructUtils/PropertyBag.cpp
+++ b/Engine/Source/Runtime/CoreUObject/Private/StructUtils/PropertyBag.cpp
@@ -3,6 +3,7 @@
 #include "StructUtils/PropertyBag.h"
 
 #include "HAL/IConsoleManager.h"
+#include "Hash/xxhash.h"
 #include "Misc/ScopeLock.h"
 #include "Misc/TransactionallySafeCriticalSection.h"
 #include "UObject/EnumProperty.h"
@@ -57,22 +58,25 @@ FCustomVersionRegistration GPropertyBagCustomVersion(FPropertyBagCustomVersion::
 
 namespace UE::StructUtils::Private
 {
-	uint64 CalcPropertyDescArrayHash(const TConstArrayView<FPropertyBagPropertyDesc> Descs);
 
 	bool CanCastTo(const UStruct* From, const UStruct* To)
 	{
 		return From != nullptr && To != nullptr && From->IsChildOf(To);
 	}
 
-	uint64 GetObjectHash(const UObject* Object)
+	void AppendPropertyDescs(FXxHash64Builder& HashBuilder, const FPropertyBagPropertyDesc& Desc);
+
+	void AppendPathName(FStringBuilderBase& StringBuilder, const UObject* Object)
 	{
-		const FString PathName = GetPathNameSafe(Object);
-		return CityHash64(reinterpret_cast<const char*>(GetData(PathName)), PathName.Len() * sizeof(TCHAR));
+		if (Object)
+		{
+			Object->GetPathName(nullptr, StringBuilder);
+		}
 	}
 
-	uint64 CalcValueObjectHash(const UObject* InValueObject)
+	void AppendValueObjectHash(FXxHash64Builder& HashBuilder, FStringBuilderBase& StringBuilder, const UObject* InValueObject)
 	{
-		uint64 ValueTypeObjectHash = GetObjectHash(InValueObject);
+		AppendPathName(StringBuilder, InValueObject);
 #if WITH_EDITOR
 		// for user defined structures we need to hash each property. after changing a user defined structure
 		// we may find the same property bag again if the hash is not sufficient, even though the memory layout
@@ -81,12 +85,13 @@ namespace UE::StructUtils::Private
 		{
 			for(const FField* Property = UserDefinedStruct->ChildProperties; Property != nullptr; Property = Property->Next)
 			{
-				ValueTypeObjectHash = CityHash128to64({ValueTypeObjectHash, GetTypeHash(Property)});
+				const uint32 FieldHash = GetTypeHash(Property);
+				HashBuilder.Update(&FieldHash, sizeof(FieldHash));
 				if(const FStructProperty* StructProperty = CastField<FStructProperty>(Property))
 				{
 					if(StructProperty->Struct)
 					{
-						ValueTypeObjectHash = CityHash128to64({ValueTypeObjectHash, CalcValueObjectHash(StructProperty->Struct)});
+						AppendValueObjectHash(HashBuilder, StringBuilder, StructProperty->Struct);
 					}
 				}
 			}
@@ -94,30 +99,40 @@ namespace UE::StructUtils::Private
 		// for property bags we recurse
 		else if(const UPropertyBag* PropertyBag = Cast<UPropertyBag>(InValueObject))
 		{
-			ValueTypeObjectHash = CityHash128to64({ValueTypeObjectHash, CalcPropertyDescArrayHash(PropertyBag->GetPropertyDescs())});
+			for (const FPropertyBagPropertyDesc& Desc : PropertyBag->GetPropertyDescs())
+			{
+				AppendPropertyDescs(HashBuilder, Desc);
+			}
 		}
 #endif
-		return ValueTypeObjectHash;
 	}
 
-	uint64 CalcPropertyDescHash(const FPropertyBagPropertyDesc& Desc)
+	void AppendPropertyDescs(FXxHash64Builder& HashBuilder, const FPropertyBagPropertyDesc& Desc)
 	{
+		TStringBuilder<256> StringBuilder;
+		HashBuilder.Update(&Desc.ID, sizeof(Desc.ID));
+		StringBuilder << Desc.Name;
+		HashBuilder.Update(&Desc.ValueType, sizeof(Desc.ValueType));
+		HashBuilder.Update(&Desc.PropertyFlags, sizeof(Desc.PropertyFlags));
 #if WITH_EDITORONLY_DATA
-		const uint32 Hashes[] = { GetTypeHash(Desc.ID), GetTypeHash(Desc.Name), GetTypeHash(Desc.ValueType), GetTypeHash(Desc.ContainerTypes), GetTypeHash(Desc.PropertyFlags), GetTypeHash(Desc.MetaData) };
-#else
-		const uint32 Hashes[] = { GetTypeHash(Desc.ID), GetTypeHash(Desc.Name), GetTypeHash(Desc.ValueType), GetTypeHash(Desc.ContainerTypes), GetTypeHash(Desc.PropertyFlags) };
+		for (const FPropertyBagPropertyDescMetaData& MetaData : Desc.MetaData)
+		{
+			StringBuilder << MetaData.Key;
+			HashBuilder.Update(*MetaData.Value, MetaData.Value.Len() * sizeof(FString::ElementType));
+		}
 #endif
-		return CityHash64WithSeed(reinterpret_cast<const char*>(Hashes), sizeof(Hashes), CalcValueObjectHash(Desc.ValueTypeObject));
+		AppendValueObjectHash(HashBuilder, StringBuilder, Desc.ValueTypeObject);
+		HashBuilder.Update(StringBuilder.ToString(), StringBuilder.Len() * sizeof(TStringBuilder<256>::ElementType));
 	}
 
 	uint64 CalcPropertyDescArrayHash(const TConstArrayView<FPropertyBagPropertyDesc> Descs)
 	{
-		uint64 Hash = 0;
+		FXxHash64Builder HashBuilder;
 		for (const FPropertyBagPropertyDesc& Desc : Descs)
 		{
-			Hash = CityHash128to64(Uint128_64(Hash, CalcPropertyDescHash(Desc)));
+			AppendPropertyDescs(HashBuilder, Desc);
 		}
-		return Hash;
+		return HashBuilder.Finalize().Hash;
 	}
 
 	FPropertyBagContainerTypes GetContainerTypesFromProperty(const FProperty* InSourceProperty)
