Allow StopOuter to be passed to GetPathNameSafe
#rnx

diff --git a/Engine/Source/Runtime/CoreUObject/Private/UObject/Field.cpp b/Engine/Source/Runtime/CoreUObject/Private/UObject/Field.cpp
--- a/Engine/Source/Runtime/CoreUObject/Private/UObject/Field.cpp
+++ b/Engine/Source/Runtime/CoreUObject/Private/UObject/Field.cpp
@@ -1126,11 +1126,11 @@ FString GetFullNameSafe(const FField* InField)
 	}
 }
 
-FString GetPathNameSafe(const FField* InField)
+FString GetPathNameSafe(const FField* InField, const UObject* StopOuter /*= nullptr*/)
 {
 	if (InField)
 	{
-		return InField->GetPathName();
+		return InField->GetPathName(StopOuter);
 	}
 	else
 	{
diff --git a/Engine/Source/Runtime/CoreUObject/Public/UObject/Field.h b/Engine/Source/Runtime/CoreUObject/Public/UObject/Field.h
--- a/Engine/Source/Runtime/CoreUObject/Public/UObject/Field.h
+++ b/Engine/Source/Runtime/CoreUObject/Public/UObject/Field.h
@@ -1108,7 +1108,7 @@ COREUOBJECT_API FString GetFullNameSafe(const FField* InField);
 /**
  * Gets the path name of the provided field. If the field pointer is null, the result is "none"
  */
-COREUOBJECT_API FString GetPathNameSafe(const FField* InField);
+COREUOBJECT_API FString GetPathNameSafe(const FField* InField, const UObject* StopOuter = nullptr);
 /** 
  * Finds a field given a path to the field (Package.Class[:Subobject:...]:FieldName)
  */
