Temporary workaround for clearing the initialization state of object properties that were initialized but their referenced object has since been deleted.

#rnx

diff --git a/Engine/Source/Runtime/CoreUObject/Internal/UObject/PropertyStateTracking.h b/Engine/Source/Runtime/CoreUObject/Internal/UObject/PropertyStateTracking.h
--- a/Engine/Source/Runtime/CoreUObject/Internal/UObject/PropertyStateTracking.h
+++ b/Engine/Source/Runtime/CoreUObject/Internal/UObject/PropertyStateTracking.h
@@ -10,10 +10,10 @@
 #include "Serialization/StructuredArchive.h"
 #include "Templates/SharedPointerFwd.h"
 #include "UObject/Class.h"
+#include "UObject/UnrealType.h"
 
 class FBlake3;
 class FName;
-class FProperty;
 class FStructuredArchiveRecord;
 class UEnum;
 class UObject;
@@ -63,9 +63,26 @@ public:
 	}
 
 	/** Query whether the property has the property value flag set. */
-	inline bool IsSet(const FProperty* Property, int32 ArrayIndex = 0) const
+	bool IsSet(const FProperty* Property, int32 ArrayIndex = 0) const
 	{
-		return Struct->HasPropertyValueFlag(Flag, Data, Property, ArrayIndex);
+		bool bIsSet = Struct->HasPropertyValueFlag(Flag, Data, Property, ArrayIndex);
+
+		// This is a temporary workaround to handle the case of an object reference being
+		// garbage collected and set to NULL (see FGarbageEliminationScopeGuard). This if
+		// should be removed once the Editor handles the clearing of deleted object references.
+		if (bIsSet && (Flag == EPropertyValueFlags::Initialized) && Property->HasAnyPropertyFlags(CPF_NonNullable))
+		{
+			if (const FObjectPropertyBase* ObjectProperty = CastField<FObjectPropertyBase>(Property))
+			{
+				const UObject* PropertyValue = ObjectProperty->LoadObjectPropertyValue_InContainer(Data);
+				if (!PropertyValue)
+				{
+					bIsSet = false;
+				}
+			}
+		}
+		
+		return bIsSet;
 	}
 
 	/** Set the property value flag for the property to bValue. */
