Fixed text import of delegate properties with full object paths.
Made FDelegateProperty consistent with FPropertyMulticastDelegate on text export.

diff --git a/Engine/Source/Runtime/CoreUObject/Private/UObject/PropertyDelegate.cpp b/Engine/Source/Runtime/CoreUObject/Private/UObject/PropertyDelegate.cpp
--- a/Engine/Source/Runtime/CoreUObject/Private/UObject/PropertyDelegate.cpp
+++ b/Engine/Source/Runtime/CoreUObject/Private/UObject/PropertyDelegate.cpp
@@ -141,10 +141,10 @@ void FDelegateProperty::ExportText_Internal( FString& ValueStr, const void* Prop
 		check(ScriptDelegate != NULL);
 		bool bDelegateHasValue = ScriptDelegate->GetFunctionName() != NAME_None;
 		ValueStr += FString::Printf(TEXT("%s.%s"),
-			ScriptDelegate->GetUObject() != NULL ? *ScriptDelegate->GetUObject()->GetName() : TEXT("(null)"),
+			ScriptDelegate->GetUObject() != NULL ? *ScriptDelegate->GetUObject()->GetPathName() : TEXT("(null)"),
 			*ScriptDelegate->GetFunctionName().ToString());
 	};
-	
+
 	if (PropertyPointerType == EPropertyPointerType::Container && HasGetter())
 	{
 		FScriptDelegate LocalScriptDelegate;
diff --git a/Engine/Source/Runtime/CoreUObject/Private/UObject/PropertyHelper.cpp b/Engine/Source/Runtime/CoreUObject/Private/UObject/PropertyHelper.cpp
--- a/Engine/Source/Runtime/CoreUObject/Private/UObject/PropertyHelper.cpp
+++ b/Engine/Source/Runtime/CoreUObject/Private/UObject/PropertyHelper.cpp
@@ -143,7 +143,17 @@ const TCHAR* DelegatePropertyTools::ImportDelegateFromText( FScriptDelegate& Del
 	{
 		++Buffer;
 	}
-	const TCHAR* NameEnd = FAsciiSet::FindFirstOrEnd(Buffer, ".),");
+
+	// Find the function name separator before finding the rest of the object name.
+	const TCHAR* DelegatePathEnd = FAsciiSet::FindFirstOrEnd(Buffer, ",)");
+	const TCHAR* NameEnd = FAsciiSet::FindLastOrEnd(Buffer, ".");
+
+	if (NameEnd > DelegatePathEnd || *NameEnd != TEXT('.'))
+	{
+		// Didn't find the separator before the end - probably malformed text
+		return nullptr;
+	}
+
 	FStringView ObjNameView(Buffer, UE_PTRDIFF_TO_INT32(NameEnd - Buffer));
 	Buffer += ObjNameView.Len();
 
