Changing Object.IsResolved() to IsObjectHandleResolved_ForGC in a few more FReferenceCollector::AddReferencedObject(s) overrides to properly handle remote object references

diff --git a/Engine/Source/Runtime/CoreUObject/Public/UObject/UObjectGlobals.h b/Engine/Source/Runtime/CoreUObject/Public/UObject/UObjectGlobals.h
--- a/Engine/Source/Runtime/CoreUObject/Public/UObject/UObjectGlobals.h
+++ b/Engine/Source/Runtime/CoreUObject/Public/UObject/UObjectGlobals.h
@@ -2689,7 +2689,7 @@ public:
 	template<class UObjectType>
 	void AddReferencedObject(TObjectPtr<UObjectType>& Object, const UObject* ReferencingObject = nullptr, const FProperty* ReferencingProperty = nullptr)
 	{
-		if (IsObjectHandleResolved_ForGC(*reinterpret_cast<FObjectHandle*>(&Object)))
+		if (IsObjectHandleResolved_ForGC(Object.GetHandle()))
 		{
 			HandleObjectReference(*reinterpret_cast<UObject**>(&Object), ReferencingObject, ReferencingProperty);
 		}
@@ -2707,7 +2707,7 @@ public:
 	{
 		static_assert(sizeof(UObjectType) > 0, "AddReferencedObject: Element must be a pointer to a fully-defined type");
 		static_assert(TPointerIsConvertibleFromTo<UObjectType, const UObjectBase>::Value, "AddReferencedObject: Element must be a pointer to a type derived from UObject");
-		if (Object.IsResolved())
+		if (IsObjectHandleResolved_ForGC(Object.GetHandle()))
 		{
 			HandleObjectReference(*reinterpret_cast<UObject**>(&Object), ReferencingObject, ReferencingProperty);
 		}
@@ -2782,7 +2782,7 @@ public:
 		static_assert(TPointerIsConvertibleFromTo<UObjectType, const UObjectBase>::Value, "AddReferencedObjects: Elements must be pointers to a type derived from UObject");
 		for (auto& Object : ObjectSet)
 		{
-			if (Object.IsResolved())
+			if (IsObjectHandleResolved_ForGC(Object.GetHandle()))
 			{
 				HandleObjectReference(*reinterpret_cast<UObject**>(&Object), ReferencingObject, ReferencingProperty);
 			}
