[StructUtils] Add missing property type on property bag hash. Bug introduced with CL 43547678

#rnx

diff --git a/Engine/Source/Runtime/CoreUObject/Private/StructUtils/PropertyBag.cpp b/Engine/Source/Runtime/CoreUObject/Private/StructUtils/PropertyBag.cpp
--- a/Engine/Source/Runtime/CoreUObject/Private/StructUtils/PropertyBag.cpp
+++ b/Engine/Source/Runtime/CoreUObject/Private/StructUtils/PropertyBag.cpp
@@ -113,6 +113,7 @@ namespace UE::StructUtils::Private
 		HashBuilder.Update(&Desc.ID, sizeof(Desc.ID));
 		StringBuilder << Desc.Name;
 		HashBuilder.Update(&Desc.ValueType, sizeof(Desc.ValueType));
+		HashBuilder.Update(&Desc.ContainerTypes, sizeof(Desc.ContainerTypes));
 		HashBuilder.Update(&Desc.PropertyFlags, sizeof(Desc.PropertyFlags));
 #if WITH_EDITORONLY_DATA
 		for (const FPropertyBagPropertyDescMetaData& MetaData : Desc.MetaData)
