PR #13371: [Iris] Iris server crash fix in FObjectProperty::Identical

diff --git a/Engine/Source/Runtime/CoreUObject/Private/UObject/PropertyObject.cpp b/Engine/Source/Runtime/CoreUObject/Private/UObject/PropertyObject.cpp
--- a/Engine/Source/Runtime/CoreUObject/Private/UObject/PropertyObject.cpp
+++ b/Engine/Source/Runtime/CoreUObject/Private/UObject/PropertyObject.cpp
@@ -119,8 +119,9 @@ bool FObjectProperty::Identical(const void* A, const void* B, uint32 PortFlags)
 		return false;
 	}
 
-	TObjectPtr<UObject> ObjectA = A ? GetPropertyValue(A) : nullptr;
-	TObjectPtr<UObject> ObjectB = B ? GetPropertyValue(B) : nullptr;
+	const TObjectPtr<UObject> NullValue = nullptr;
+	const TObjectPtr<UObject>& ObjectA = A ? GetPropertyValue(A) : NullValue;
+	const TObjectPtr<UObject>& ObjectB = B ? GetPropertyValue(B) : NullValue;
 
 	if (ObjectA == ObjectB)
 	{
