[ZenLoader]
- Fix wrong lock being used to access packages

diff --git a/Engine/Source/Runtime/CoreUObject/Private/Serialization/AsyncLoading2.cpp b/Engine/Source/Runtime/CoreUObject/Private/Serialization/AsyncLoading2.cpp
--- a/Engine/Source/Runtime/CoreUObject/Private/Serialization/AsyncLoading2.cpp
+++ b/Engine/Source/Runtime/CoreUObject/Private/Serialization/AsyncLoading2.cpp
@@ -2116,7 +2116,7 @@ public:
 
 	inline UObject* FindPublicExportObjectUnchecked(const FPublicExportKey& Key)
 	{
-		FReadScopeLock ScopeLock(ScriptObjectsLock);
+		FReadScopeLock ScopeLock(PackagesLock);
 		return FindPublicExportObjectUncheckedNoLock(Key);
 	}
 
