Tracking the initialization state of default-initialized Verse struct properties.

#rnx

diff --git a/Engine/Source/Runtime/CoreUObject/Private/UObject/InstanceDataObjectUtils.cpp b/Engine/Source/Runtime/CoreUObject/Private/UObject/InstanceDataObjectUtils.cpp
--- a/Engine/Source/Runtime/CoreUObject/Private/UObject/InstanceDataObjectUtils.cpp
+++ b/Engine/Source/Runtime/CoreUObject/Private/UObject/InstanceDataObjectUtils.cpp
@@ -1347,6 +1347,78 @@ namespace UE
 			CLASS_EditInlineNew | CLASS_CollapseCategories | CLASS_Const | CLASS_CompiledFromBlueprint | CLASS_HasInstancedReference);
 	}
 
+	static void TrackDefaultInitializedPropertiesInCDO(UObject* CDO, UClass* Class, const UVerseClass* OwnerClass)
+	{
+		TRACE_CPUPROFILER_EVENT_SCOPE(TrackDefaultInitializedPropertiesInCDO)
+
+		// If a property has been initialized by the Verse class' CDO then we consider everything
+		// under it (i.e., any sub-objects) as initialized.
+		for (const TFieldPath<FProperty>& OwnerFieldPath : OwnerClass->PropertiesWrittenByInitCDO)
+		{
+			FProperty* OwnerProperty = OwnerFieldPath.Get();
+
+			if (FProperty* Property = Class->FindPropertyByName(OwnerProperty->GetFName()))
+			{
+				if (Property->HasAnyPropertyFlags(CPF_RequiredParm))
+				{
+					UE::FInitializedPropertyValueState(CDO).Set(Property);
+				}
+
+				// Keep track of visited property-owner pairs to avoid referencing cycles.
+				TSet<TTuple<const FProperty*, void*>> VisitedPropOwners;
+
+				FPropertyVisitorPath Path;
+				FPropertyVisitorData VisitorData(CDO, /*ParentStructData*/CDO);
+				FPropertyVisitorContext VisitorContext(Path, VisitorData, FPropertyVisitorContext::EScope::All);
+				PropertyVisitorHelpers::VisitProperty(Class, Property, VisitorContext, [CDO, &VisitedPropOwners](const FPropertyVisitorContext& Context) -> EPropertyVisitorControlFlow
+				{
+					const FPropertyVisitorPath& PropertyPath = Context.Path;
+					const FPropertyVisitorData& Data = Context.Data;
+					const FProperty* Property = PropertyPath.Top().Property;
+					void* Owner = Data.ParentStructData;
+					TTuple<const FProperty*, void*> PropOwner(Property, Owner);
+
+					if (!Property || VisitedPropOwners.Contains(PropOwner))
+					{
+						return EPropertyVisitorControlFlow::StepOver;
+					}
+
+					bool bIsInCDO = true;
+
+					const UStruct* OwnerType = PropertyPath.Top().ParentStructType;
+
+					if (OwnerType && OwnerType->IsChildOf<UObject>())
+					{
+						const UObject* OwnerObject = (const UObject*)Owner;
+
+						if (OwnerObject && (OwnerObject != CDO))
+						{
+							bIsInCDO = OwnerObject->IsInOuter(CDO);
+						}
+					}
+
+					// It is possible for the property and owner types to differ during re-instancing
+					// when a new CDO is created. Skip tracking in this case.
+					if (bIsInCDO && OwnerType && Property->HasAnyPropertyFlags(CPF_RequiredParm) && OwnerType->IsChildOf(Property->GetOwnerStruct()))
+					{
+						UE::FInitializedPropertyValueState(OwnerType, Owner).Set(Property);
+					}
+
+					VisitedPropOwners.Add(PropOwner);
+
+					if (bIsInCDO)
+					{
+						return EPropertyVisitorControlFlow::StepInto;
+					}
+					else
+					{
+						return EPropertyVisitorControlFlow::StepOver;
+					}
+				});
+			}
+		}
+	}
+
 	UClass* CreateInstanceDataObjectClass(const FPropertyPathNameTree* PropertyTree, const FUnknownEnumNames* EnumNames, UClass* OwnerClass, UObject* Outer)
 	{
 		PropertyTree = bEnableIDOUnknownProperties ? PropertyTree : nullptr;
@@ -1398,6 +1470,17 @@ namespace UE
 
 		CleanUpInstanceDataObjectTypeCache(IDOClassCache);
 
+		// We need to copy the initialization state of the owner's properties for properties initialized in Verse code:
+		// e.g.: MyStruct<public>:test_struct = test_struct{ MyInt := 1 }
+		// This is necessary for Verse struct properties, because struct types do not store their initialization state 
+		// in any way outside of IDOs (in contrast to classes which use an annotation system).
+		// NOTE: if IDOs become the only system for tracking initialization state then we should deprecate the use of:
+		// TrackDefaultInitializedProperties in UStruct.
+		if (const UVerseClass* OwnerVerseClass = Cast<UVerseClass>(OwnerClass))
+		{
+			TrackDefaultInitializedPropertiesInCDO(NewClass->GetDefaultObject(), NewClass, OwnerVerseClass);
+		}
+
 		IDOClassCache.Add(Key, NewClass);
 		return NewClass;
 	}
