Allow the FProperty container more control over adding the property to the Destructor and PostConstruct link lists

#rnx

diff --git a/Engine/Source/Runtime/CoreUObject/Private/UObject/Class.cpp b/Engine/Source/Runtime/CoreUObject/Private/UObject/Class.cpp
--- a/Engine/Source/Runtime/CoreUObject/Private/UObject/Class.cpp
+++ b/Engine/Source/Runtime/CoreUObject/Private/UObject/Class.cpp
@@ -1063,6 +1063,30 @@ void UStruct::Link(FArchive& Ar, bool bRelinkExistingProperties)
 	{
 		FProperty* Property = *It;
 
+		// Determine which lists we add the property to.   
+		//
+		// The relationship between a property and the owning class is more complex that might first 
+		// appear.  
+		// 
+		// For UHT generated types, UFunction are owned by a UClass.  UScriptStructs are not 
+		// owned by a class and will default to needing a LinkDestructor. For blueprint types, 
+		// UScriptStructs can be owned and will get link flags from the owning UClass.
+		EStructPropertyLinkFlags PropertyLinkFlags = EStructPropertyLinkFlags::LinkDestructor;
+		{
+			if (UClass* OwnerClass = Property->GetOwnerClass())
+			{
+				PropertyLinkFlags = OwnerClass->GetPropertyLinkFlags(this, Property);
+			}
+			if (Property->ContainsFinishDestroy(EncounteredStructProps))
+			{
+				PropertyLinkFlags |= EStructPropertyLinkFlags::LinkDestructor;
+			}
+			else if (Property->HasAnyPropertyFlags(CPF_IsPlainOldData | CPF_NoDestructor))
+			{
+				PropertyLinkFlags &= ~EStructPropertyLinkFlags::LinkDestructor;
+			}
+		}
+
 		// Ref link contains any properties which contain object references including types with user-defined serializers which don't explicitly specify whether they 
 		// contain object references
 		if (Property->ContainsObjectReference(EncounteredStructProps, EPropertyObjectReferenceType::Any))
@@ -1070,23 +1094,15 @@ void UStruct::Link(FArchive& Ar, bool bRelinkExistingProperties)
 			RefLinkBuilder.AppendNoTerminate(*Property);
 		}
 
-		// We are checking for Script content to cover the case where a UFunction using script VM has been added
-		// to a native class.  This is the case for UHT generated Verse classes.  The Verse compiler will add the
-		// Verse script implementations to the UHT generated class.
-		const UClass* OwnerClass = Property->GetOwnerClass();
-		bool bOwnedByNativeClass = Script.IsEmpty() && OwnerClass && OwnerClass->HasAnyClassFlags(CLASS_Native | CLASS_Intrinsic);
-		bool bShouldHandleFinishDestroy = Property->ContainsFinishDestroy(EncounteredStructProps);
-
-		if ((!Property->HasAnyPropertyFlags(CPF_IsPlainOldData | CPF_NoDestructor)
-			&& !bOwnedByNativeClass) // these would be covered by the native destructor
-			|| bShouldHandleFinishDestroy)
+		// These properties will be destructed via the property implementation.  Things in a struct that need a destructor will still be in here, 
+		// even though in many cases they will also be destroyed by a native destructor on the whole struct
+		if (EnumHasAnyFlags(PropertyLinkFlags, EStructPropertyLinkFlags::LinkDestructor))
 		{
-			// things in a struct that need a destructor will still be in here, even though in many cases they will also be destroyed by a native destructor on the whole struct
 			DestructorLinkBuilder.AppendNoTerminate(*Property);
 		}
 
 		// Link references to properties that require their values to be initialized and/or copied from CDO post-construction. Note that this includes all non-native-class-owned properties.
-		if (OwnerClass && (!bOwnedByNativeClass || (Property->HasAnyPropertyFlags(CPF_Config) && !OwnerClass->HasAnyClassFlags(CLASS_PerObjectConfig))))
+		if (EnumHasAnyFlags(PropertyLinkFlags, EStructPropertyLinkFlags::LinkPostConstruct))
 		{
 			PostConstructLinkBuilder.AppendNoTerminate(*Property);
 		}
@@ -5372,6 +5388,24 @@ void UClass::Link(FArchive& Ar, bool bRelinkExistingProperties)
 	Super::Link(Ar, bRelinkExistingProperties);
 }
 
+EStructPropertyLinkFlags UClass::GetPropertyLinkFlags(UStruct* ContainerStruct, FProperty* Property) const
+{
+	// If we aren't a native type or we have a VM script, then we will need destructor and post construct
+	if (!ContainerStruct->Script.IsEmpty() || !HasAnyClassFlags(CLASS_Native | CLASS_Intrinsic))
+	{
+		return EStructPropertyLinkFlags::LinkDestructor | EStructPropertyLinkFlags::LinkPostConstruct;
+	}
+
+	// For native types, config properties also need post construct
+	if (Property->HasAnyPropertyFlags(CPF_Config) && !HasAnyClassFlags(CLASS_PerObjectConfig))
+	{
+		return EStructPropertyLinkFlags::LinkPostConstruct;
+	}
+
+	// All other properties of native types assumes constructor/destructor will handle everything
+	return EStructPropertyLinkFlags::None;
+}
+
 #if VALIDATE_CLASS_REPS
 namespace UE::Net::Private
 {
diff --git a/Engine/Source/Runtime/CoreUObject/Public/UObject/Class.h b/Engine/Source/Runtime/CoreUObject/Public/UObject/Class.h
--- a/Engine/Source/Runtime/CoreUObject/Public/UObject/Class.h
+++ b/Engine/Source/Runtime/CoreUObject/Public/UObject/Class.h
@@ -408,6 +408,20 @@ private:
 -----------------------------------------------------------------------------*/
 enum class EStructStateFlags : uint16;
 
+/** Provides more explicit control over if a property is added to PostConstruct & Destructor collections */
+enum class EStructPropertyLinkFlags : uint8
+{
+	/** Property should not be added to any link lists */
+	None = 0,
+
+	/** If set, add the property to the post constructor link list */
+	LinkPostConstruct = 1 << 0,
+
+	/** If set, add the property to the destructor link list */
+	LinkDestructor = 1 << 1,
+};
+ENUM_CLASS_FLAGS(EStructPropertyLinkFlags)
+
 /**
  * Base class for all UObject types that contain fields.
  */
@@ -3380,6 +3394,7 @@ public:
 	COREUOBJECT_API virtual void Link(FArchive& Ar, bool bRelinkExistingProperties) override;
 	COREUOBJECT_API virtual void SetSuperStruct(UStruct* NewSuperStruct) override;
 	COREUOBJECT_API virtual bool IsStructTrashed() const override;
+	UE_INTERNAL COREUOBJECT_API virtual EStructPropertyLinkFlags GetPropertyLinkFlags(UStruct* ContainerStruct, FProperty* Property) const;
 	// End of UStruct interface.
 
 #if WITH_EDITORONLY_DATA
