Again, changing more Object.IsResolved() calls to IsObjectHandleResolved_ForGC in FReferenceCollector::AddReferencedObject(s) overrides to properly handle remote object references. Hopefully I got all of them this time

diff --git a/Engine/Source/Runtime/CoreUObject/Public/UObject/ObjectHandle.h b/Engine/Source/Runtime/CoreUObject/Public/UObject/ObjectHandle.h
--- a/Engine/Source/Runtime/CoreUObject/Public/UObject/ObjectHandle.h
+++ b/Engine/Source/Runtime/CoreUObject/Public/UObject/ObjectHandle.h
@@ -174,12 +174,11 @@ inline bool IsObjectHandleResolved(FObjectHandle Handle)
  */
 inline bool IsObjectHandleResolved_ForGC(FObjectHandle Handle)
 {
-#if UE_WITH_OBJECT_HANDLE_LATE_RESOLVE
-	return !(Handle.PointerOrRef & 1);
-#elif UE_WITH_REMOTE_OBJECT_HANDLE
+#if UE_WITH_REMOTE_OBJECT_HANDLE
+	// Unlike the non-GC version we don't check if the object (if still exists locally) is marked as Remote
 	return !(Handle.PointerOrHandle & 1);
 #else
-	return true;
+	return IsObjectHandleResolved(Handle);
 #endif
 }
 
diff --git a/Engine/Source/Runtime/CoreUObject/Public/UObject/UObjectGlobals.h b/Engine/Source/Runtime/CoreUObject/Public/UObject/UObjectGlobals.h
--- a/Engine/Source/Runtime/CoreUObject/Public/UObject/UObjectGlobals.h
+++ b/Engine/Source/Runtime/CoreUObject/Public/UObject/UObjectGlobals.h
@@ -2805,7 +2805,7 @@ public:
 									!(std::is_pointer_v<ValueType> && std::is_convertible_v<ValueType, const UObjectBase*>));
 		for (auto& It : Map)
 		{
-			if (It.Key.IsResolved())
+			if (IsObjectHandleResolved_ForGC(It.Key.GetHandle()))
 			{
 				HandleObjectReference(*reinterpret_cast<UObject**>(&It.Key), ReferencingObject, ReferencingProperty);
 			}
@@ -2820,7 +2820,7 @@ public:
 									!(std::is_pointer_v<KeyType> && std::is_convertible_v<KeyType, const UObjectBase*>));
 		for (auto& It : Map)
 		{
-			if (It.Value.IsResolved())
+			if (IsObjectHandleResolved_ForGC(It.Value.GetHandle()))
 			{
 				HandleObjectReference(*reinterpret_cast<UObject**>(&It.Value), ReferencingObject, ReferencingProperty);
 			}
@@ -2835,11 +2835,11 @@ public:
 		static_assert(TPointerIsConvertibleFromTo<ValueType, const UObjectBase>::Value, "AddReferencedObjects: Values must be pointers to a type derived from UObject");
 		for (auto& It : Map)
 		{
-			if (It.Key.IsResolved())
+			if (IsObjectHandleResolved_ForGC(It.Key.GetHandle()))
 			{
 				HandleObjectReference(*reinterpret_cast<UObject**>(&It.Key), ReferencingObject, ReferencingProperty);
 			}
-			if (It.Value.IsResolved())
+			if (IsObjectHandleResolved_ForGC(It.Value.GetHandle()))
 			{
 				HandleObjectReference(*reinterpret_cast<UObject**>(&It.Value), ReferencingObject, ReferencingProperty);
 			}
