[ZenLoader]
- Fix wrong lock being used to access packages

diff --git a/Engine/Source/Runtime/CoreUObject/Private/Serialization/AsyncLoading2.cpp b/Engine/Source/Runtime/CoreUObject/Private/Serialization/AsyncLoading2.cpp
--- a/Engine/Source/Runtime/CoreUObject/Private/Serialization/AsyncLoading2.cpp
+++ b/Engine/Source/Runtime/CoreUObject/Private/Serialization/AsyncLoading2.cpp
@@ -2014,7 +2014,7 @@ public:
 			FLoadedPackageRef* PackageRef = FindPackageRefNoLock(PackageId);
 			if (PackageRef)
 			{
-				UObject* ExistingObject = FindPublicExportObjectUnchecked(*PublicExportKey);
+				UObject* ExistingObject = FindPublicExportObjectUncheckedNoLock(*PublicExportKey);
 				UE_CLOG(!ExistingObject, LogStreaming, Fatal,
 					TEXT("FGlobalImportStore::VerifyObjectForRemoval: The loaded public export object '%s' with flags (ObjectFlags=%x, InternalObjectFlags=%x) and id %s:0x%llX is missing in GlobalImportStore. ")
 					TEXT("Reason unknown. Double delete? Bug or hash collision?"),
@@ -2116,7 +2116,7 @@ public:
 
 	inline UObject* FindPublicExportObjectUnchecked(const FPublicExportKey& Key)
 	{
-		FReadScopeLock ScopeLock(ScriptObjectsLock);
+		FReadScopeLock ScopeLock(PackagesLock);
 		return FindPublicExportObjectUncheckedNoLock(Key);
 	}
 
