Added more trace regions for dstm

diff --git a/Engine/Source/Runtime/CoreUObject/Private/UObject/RemoteExecutor.cpp b/Engine/Source/Runtime/CoreUObject/Private/UObject/RemoteExecutor.cpp
--- a/Engine/Source/Runtime/CoreUObject/Private/UObject/RemoteExecutor.cpp
+++ b/Engine/Source/Runtime/CoreUObject/Private/UObject/RemoteExecutor.cpp
@@ -121,6 +121,14 @@ public:
 
 		double LastStallPrintTime = FPlatformTime::Seconds();
 
+#if !UE_BUILD_SHIPPING
+		ELogVerbosity::Type InitialLogVerbosity = LogRemoteExec.GetVerbosity();
+		ON_SCOPE_EXIT
+		{
+			LogRemoteExec.SetVerbosity(InitialLogVerbosity);
+		};
+#endif
+
 		for (int32 LocalIterationNumber = 0; ; LocalIterationNumber++)
 		{
 			TRACE_CPUPROFILER_EVENT_SCOPE(ResumeExecutePendingWork)
@@ -128,6 +136,15 @@ public:
 			double Now = FPlatformTime::Seconds();
 			bool bPrintStall = false;
 
+#if !UE_BUILD_SHIPPING
+			constexpr int32 MaxIterationsBeforeStall = 100000;
+			if (LocalIterationNumber >= MaxIterationsBeforeStall && (LocalIterationNumber % MaxIterationsBeforeStall))
+			{
+				bPrintStall = true;
+				UE_SET_LOG_VERBOSITY(LogRemoteExec, VeryVerbose);
+			}
+#endif
+
 			if ((Now - LastStallPrintTime) > 1.0)
 			{
 				LastStallPrintTime = Now;
@@ -908,11 +925,14 @@ void ExecutePendingWork()
 
 void ExecuteTransactionalWithExplicitPriority(FName WorkName, FRemoteWorkPriority WorkPriority, const TFunctionRef<void(void)>& Work)
 {
+	static const FName Name = "ExecuteTransactional";
+	UE::RemoteExecutor::OnRegionBeginDelegate.ExecuteIfBound(Name);
 #if UE_WITH_REMOTE_OBJECT_HANDLE
 	ExecuteTransactionalInternal(WorkName, WorkPriority, Work);
 #else
 	Work();
 #endif
+	UE::RemoteExecutor::OnRegionEndDelegate.ExecuteIfBound(FString());
 }
 
 void ExecuteTransactional(FName WorkName, const TFunctionRef<void(void)>& Work)
