[ZenLoader]
- Remove editor hybrid loader since we now fully rely on zenloader to mix cooked and uncooked content together

[RN] The hybrid loader was never intended to be used outside of Epic and has been removed without deprecation

diff --git a/Engine/Source/Runtime/CoreUObject/Private/Serialization/AsyncLoading2.cpp b/Engine/Source/Runtime/CoreUObject/Private/Serialization/AsyncLoading2.cpp
--- a/Engine/Source/Runtime/CoreUObject/Private/Serialization/AsyncLoading2.cpp
+++ b/Engine/Source/Runtime/CoreUObject/Private/Serialization/AsyncLoading2.cpp
@@ -4110,7 +4110,7 @@ class FAsyncLoadingThread2 final
 	friend struct FAsyncPackage2;
 	friend class FAsyncLoadingVerseRoot;
 public:
-	FAsyncLoadingThread2(FIoDispatcher& IoDispatcher, IAsyncPackageLoader* InUncookedPackageLoader);
+	FAsyncLoadingThread2(FIoDispatcher& IoDispatcher);
 	virtual ~FAsyncLoadingThread2();
 
 	virtual ELoaderType GetLoaderType() const override
@@ -4294,8 +4294,6 @@ private:
 	/** I/O Dispatcher */
 	FIoDispatcher& IoDispatcher;
 
-	IAsyncPackageLoader* UncookedPackageLoader;
-
 	FPackageStore& PackageStore;
 	FGlobalImportStore GlobalImportStore;
 	TMpscQueue<FPackageRequest> PackageRequestQueue;
@@ -6583,72 +6581,6 @@ void FAsyncPackage2::ImportPackagesRecursiveInner(FAsyncLoadingThreadState2& Thr
 		checkf(ImportedPackageStatus == EPackageStoreEntryStatus::Missing || ImportedPackageRef.GetPackageLoader() == ImportedPackageLoader, TEXT("ImportPackageStatus(%d) != EPackageStoreEntryStatus::Missing(%d) || ImportedPackageRef.GetPackageLoader()(%s) != ImportedPackageLoader(%s)"), ImportedPackageStatus, EPackageStoreEntryStatus::Missing, LexToString(ImportedPackageRef.GetPackageLoader()), LexToString(ImportedPackageLoader));
 		checkf(ImportedPackageStatus == EPackageStoreEntryStatus::Missing || ImportedPackageRef.GetPackageHeaderExtension() == ImportedPackagePath.GetHeaderExtension(), TEXT("ImportPackageStatus(%d) != EPackageStoreEntryStatus::Missing(%d) || ImportedPackageRef.GetPackageHeaderExtension()(%d) != ImportedPackagePath.GetHeaderExtension()(%d)"), ImportedPackageStatus, EPackageStoreEntryStatus::Missing, ImportedPackageRef.GetPackageHeaderExtension(), ImportedPackagePath.GetHeaderExtension());
 
-#if WITH_EDITOR
-		if (AsyncLoadingThread.UncookedPackageLoader && ImportedPackageStatus == EPackageStoreEntryStatus::Ok && ImportedPackageLoader == EPackageLoader::LinkerLoad)
-		{
-			UPackage* UncookedPackage = ImportedPackageRef.GetPackage();
-			if (!ImportedPackageRef.AreAllPublicExportsLoaded())
-			{
-				UE_ASYNC_PACKAGE_LOG(Verbose, Desc, TEXT("ImportPackages: LoadUncookedImport"), TEXT("Loading imported uncooked package '%s' '0x%s'"), *ImportedPackageNameToLoad.ToString(), *LexToString(ImportedPackageId));
-				check(IsInGameThread());
-				IoBatch.Issue(); // The batch might already contain requests for packages being imported from the uncooked one we're going to load so make sure that those are started before blocking
-				check(!ImportedPackagePath.IsEmpty());
-				int32 ImportRequestId = AsyncLoadingThread.UncookedPackageLoader->LoadPackage(ImportedPackagePath, NAME_None, FLoadPackageAsyncDelegate(), PKG_None, INDEX_NONE, 0, nullptr, LOAD_None);
-				AsyncLoadingThread.UncookedPackageLoader->FlushLoading({ImportRequestId});
-				UncookedPackage = FindObjectFast<UPackage>(nullptr, ImportedPackagePath.GetPackageFName());
-				ImportedPackageRef.SetPackage(UncookedPackage);
-				if (UncookedPackage)
-				{
-					UncookedPackage->SetCanBeImportedFlag(true);
-					UncookedPackage->SetPackageId(ImportedPackageId);
-					UncookedPackage->SetInternalFlags(EInternalObjectFlags::LoaderImport);
-					const bool bSnapshotExportCount = true;
-					ImportedPackageRef.SetAllPublicExportsLoaded(bSnapshotExportCount);
-				}
-			}
-			if (UncookedPackage)
-			{
-				ForEachObjectWithOuter(UncookedPackage, [this, ImportedPackageId](UObject* Object)
-				{
-					if (Object->HasAllFlags(RF_Public))
-					{
-						Object->SetInternalFlags(EInternalObjectFlags::LoaderImport);
-
-						TArray<FName, TInlineAllocator<64>> FullPath;
-						FullPath.Add(Object->GetFName());
-						UObject* Outer = Object->GetOuter();
-						while (Outer)
-						{
-							FullPath.Add(Outer->GetFName());
-							Outer = Outer->GetOuter();
-						}
-						TStringBuilder<256> PackageRelativeExportPath;
-						for (int32 PathIndex = FullPath.Num() - 2; PathIndex >= 0; --PathIndex)
-						{
-							TCHAR NameStr[FName::StringBufferSize];
-							uint32 NameLen = FullPath[PathIndex].ToString(NameStr);
-							for (uint32 I = 0; I < NameLen; ++I)
-							{
-								NameStr[I] = TChar<TCHAR>::ToLower(NameStr[I]);
-							}
-							PackageRelativeExportPath.AppendChar('/');
-							PackageRelativeExportPath.Append(FStringView(NameStr, NameLen));
-						}
-						uint64 ExportHash = CityHash64(reinterpret_cast<const char*>(PackageRelativeExportPath.GetData() + 1), (PackageRelativeExportPath.Len() - 1) * sizeof(TCHAR));
-						ImportStore.StoreGlobalObject(ImportedPackageId, ExportHash, Object);
-					}
-				}, /* bIncludeNestedObjects*/ true);
-			}
-			else
-			{
-				ImportedPackageRef.SetHasFailed();
-				UE_ASYNC_PACKAGE_LOG(Warning, Desc, TEXT("ImportPackages: SkipPackage"),
-					TEXT("Failed to load uncooked imported package with id '0x%s' ('%s')"), *LexToString(ImportedPackageId), *ImportedPackageNameToLoad.ToString());
-			}
-			continue;
-		}
-#endif
-
 		FAsyncPackage2* ImportedPackage = nullptr;
 		bool bInserted = false;
 		bool bIsFullyLoaded = ImportedPackageRef.AreAllPublicExportsLoaded();
@@ -10093,10 +10025,9 @@ EAsyncPackageState::Type FAsyncLoadingThread2::TickAsyncLoadingFromGameThread(FA
 	return Result;
 }
 
-FAsyncLoadingThread2::FAsyncLoadingThread2(FIoDispatcher& InIoDispatcher, IAsyncPackageLoader* InUncookedPackageLoader)
+FAsyncLoadingThread2::FAsyncLoadingThread2(FIoDispatcher& InIoDispatcher)
 	: Thread(nullptr)
 	, IoDispatcher(InIoDispatcher)
-	, UncookedPackageLoader(InUncookedPackageLoader)
 	, PackageStore(FPackageStore::Get())
 	, GlobalImportStore(*this)
 {
@@ -12181,9 +12112,9 @@ void FAsyncPackage2::DumpStateImpl(TSet<FAsyncPackage2*>& Set, int32 Indent, TMu
 
 #endif // ALT2_DUMP_STATE_ON_HANG
 
-IAsyncPackageLoader* MakeAsyncPackageLoader2(FIoDispatcher& InIoDispatcher, IAsyncPackageLoader* InUncookedPackageLoader)
+IAsyncPackageLoader* MakeAsyncPackageLoader2(FIoDispatcher& InIoDispatcher)
 {
-	return new FAsyncLoadingThread2(InIoDispatcher, InUncookedPackageLoader);
+	return new FAsyncLoadingThread2(InIoDispatcher);
 }
 
 #if WITH_DEV_AUTOMATION_TESTS 
diff --git a/Engine/Source/Runtime/CoreUObject/Private/Serialization/AsyncPackageLoader.cpp b/Engine/Source/Runtime/CoreUObject/Private/Serialization/AsyncPackageLoader.cpp
--- a/Engine/Source/Runtime/CoreUObject/Private/Serialization/AsyncPackageLoader.cpp
+++ b/Engine/Source/Runtime/CoreUObject/Private/Serialization/AsyncPackageLoader.cpp
@@ -4,7 +4,6 @@
 #include "HAL/IConsoleManager.h"
 #include "Serialization/AsyncLoadingThread.h"
 #include "Serialization/AsyncLoading2.h"
-#include "Serialization/EditorPackageLoader.h"
 #include "Serialization/TransactionallySafeAsyncLoading.h"
 #include "UObject/GCObject.h"
 #include "UObject/LinkerLoad.h"
@@ -247,17 +246,10 @@ void InitAsyncThread()
 		{
 			GPackageLoader.Reset(MakeAsyncPackageLoader2(IoDispatcher));
 		}
-#if WITH_EDITOR
-		else if (bHasScriptObjectsChunk || bHasUseIoStoreParamInEditor)
-		{
-			GPackageLoader = MakeEditorPackageLoader(IoDispatcher);
-		}
-#else
 		else if (bHasScriptObjectsChunk)
 		{
 			GPackageLoader.Reset(MakeAsyncPackageLoader2(IoDispatcher));
 		}
-#endif
 	}
 	if (!GPackageLoader.IsValid())
 	{
diff --git a/Engine/Source/Runtime/CoreUObject/Private/Serialization/EditorPackageLoader.cpp b/Engine/Source/Runtime/CoreUObject/Private/Serialization/EditorPackageLoader.cpp
deleted file mode 100644
diff --git a/Engine/Source/Runtime/CoreUObject/Private/Serialization/EditorPackageLoader.h b/Engine/Source/Runtime/CoreUObject/Private/Serialization/EditorPackageLoader.h
deleted file mode 100644
diff --git a/Engine/Source/Runtime/CoreUObject/Public/Serialization/AsyncLoading2.h b/Engine/Source/Runtime/CoreUObject/Public/Serialization/AsyncLoading2.h
--- a/Engine/Source/Runtime/CoreUObject/Public/Serialization/AsyncLoading2.h
+++ b/Engine/Source/Runtime/CoreUObject/Public/Serialization/AsyncLoading2.h
@@ -434,4 +434,4 @@ COREUOBJECT_API void FindAllRuntimeScriptPackages(FRuntimeScriptPackages& OutPac
  *
  * @return The async package loader.
  */
-IAsyncPackageLoader* MakeAsyncPackageLoader2(FIoDispatcher& InIoDispatcher, IAsyncPackageLoader* UncookedPackageLoader = nullptr);
+IAsyncPackageLoader* MakeAsyncPackageLoader2(FIoDispatcher& InIoDispatcher);
