UScriptStruct::ExportText: Reduce number of string allocations

#rnx

diff --git a/Engine/Source/Runtime/CoreUObject/Private/UObject/Class.cpp b/Engine/Source/Runtime/CoreUObject/Private/UObject/Class.cpp
--- a/Engine/Source/Runtime/CoreUObject/Private/UObject/Class.cpp
+++ b/Engine/Source/Runtime/CoreUObject/Private/UObject/Class.cpp
@@ -3512,6 +3512,8 @@ void UScriptStruct::ExportText(FString& ValueStr, const void* Value, const void*
 
 	const bool bT3DOverrideEnabled = FOverridableSerializationLogic::HasCapabilities(FOverridableSerializationLogic::ECapabilities::T3DSerialization);
 
+	FString InnerValue;
+	InnerValue.Reserve(512);
 	for (TFieldIterator<FProperty> It(this); It; ++It)
 	{
 		if (It->ShouldPort(PortFlags))
@@ -3521,7 +3523,6 @@ void UScriptStruct::ExportText(FString& ValueStr, const void* Value, const void*
 				const bool bStaticArray = It->ArrayDim > 1;
 				FOverridableTextPortPropertyPathScope ScopePath(*It, bStaticArray ? Index : INDEX_NONE, bStaticArray ? EPropertyVisitorInfoType::StaticArrayIndex : EPropertyVisitorInfoType::None);
 
-				FString InnerValue;
 				if (It->ExportText_InContainer(Index, InnerValue, Value, Defaults, OwnerObject, PPF_Delimited | PortFlags, ExportRootScope))
 				{
 					Count++;
@@ -3538,9 +3539,14 @@ void UScriptStruct::ExportText(FString& ValueStr, const void* Value, const void*
 						ValueStr += TEXT(",\n");
 					}
 
-					const FString PropertyName = (PortFlags & (PPF_ExternalEditor | PPF_BlueprintDebugView)) != 0 ? It->GetAuthoredName() : It->GetName();
+					// Add property name
+					ValueStr += (PortFlags & (PPF_ExternalEditor | PPF_BlueprintDebugView)) != 0 ? It->GetAuthoredName() : It->GetName();
 
-					FString OverridableOperation;
+					if (bStaticArray)
+					{
+						ValueStr += FString::Printf(TEXT("[%i]"), Index);
+					}
+					
 					if (bT3DOverrideEnabled)
 					{
 						if(const FOverriddenPropertySet* OverriddenProperties = FOverridableSerializationLogic::GetOverriddenProperties())
@@ -3552,21 +3558,15 @@ void UScriptStruct::ExportText(FString& ValueStr, const void* Value, const void*
 							const EOverriddenPropertyOperation Operation = OverriddenProperties->GetOverriddenPropertyOperation(&Chain, /*Property*/nullptr);
 							if (Operation != EOverriddenPropertyOperation::None)
 							{
-								OverridableOperation = FString::Printf(TEXT("<%s>"), *GetOverriddenOperationString(Operation));
+								ValueStr += FString::Printf(TEXT("<%s>"), *GetOverriddenOperationString(Operation));
 							}
 						}
 					}
 
-					if (!bStaticArray)
-					{
-						ValueStr += FString::Printf(TEXT("%s%s="), *PropertyName, *OverridableOperation);
-					}
-					else
-					{
-						ValueStr += FString::Printf(TEXT("%s[%i]%s="), *PropertyName, Index, *OverridableOperation);
-					}
+					ValueStr += TCHAR('=');;
 					ValueStr += InnerValue;
 				}
+				InnerValue.Reset();
 			}
 		}
 	}
diff --git a/Engine/Source/Runtime/CoreUObject/Private/UObject/Obj.cpp b/Engine/Source/Runtime/CoreUObject/Private/UObject/Obj.cpp
--- a/Engine/Source/Runtime/CoreUObject/Private/UObject/Obj.cpp
+++ b/Engine/Source/Runtime/CoreUObject/Private/UObject/Obj.cpp
@@ -2262,6 +2262,7 @@ static void GetAssetRegistryTagFromProperty(const void* BaseMemoryLocation, cons
 		else
 		{
 			FString PropertyStr;
+			PropertyStr.Reserve(8192);
 			Prop->ExportTextItem_Direct(PropertyStr, Bundles, Bundles, nullptr, PPF_None);
 			Context.AddTag(UObject::FAssetRegistryTag(GAssetBundleDataName, MoveTemp(PropertyStr), UObject::FAssetRegistryTag::ETagType::TT_Alphabetical));
 		}
@@ -2310,6 +2311,7 @@ static void GetAssetRegistryTagFromProperty(const void* BaseMemoryLocation, cons
 		}
 		else
 		{
+			PropertyStr.Reserve(8192);
 			const uint8* PropertyAddr = Prop->ContainerPtrToValuePtr<uint8>(BaseMemoryLocation);
 			Prop->ExportTextItem_Direct(PropertyStr, PropertyAddr, PropertyAddr, nullptr, PPF_None);
 		}
