[AutoRTFM] Remove unnecessary open around critical section locking and unlocking.

This was added in March 2024, before we had a transactionally-safe critical section class. This lock has already been updated to use FTransactionallySafeCriticalSection so it's unnecessary to use open blocks here.

diff --git a/Engine/Source/Runtime/CoreUObject/Private/UObject/UObjectHash.cpp b/Engine/Source/Runtime/CoreUObject/Private/UObject/UObjectHash.cpp
--- a/Engine/Source/Runtime/CoreUObject/Private/UObject/UObjectHash.cpp
+++ b/Engine/Source/Runtime/CoreUObject/Private/UObject/UObjectHash.cpp
@@ -826,8 +826,7 @@ public:
 		if (!(IsGarbageCollectingAndLockingUObjectHashTables() && IsInGameThread()))
 		{
 			Tables = &InTables;
-			UE_AUTORTFM_OPEN{ InTables.Lock(); };
-			AutoRTFM::PushOnAbortHandler(this, [this](){ if (this->Tables){ this->Tables->Unlock(); }});
+			InTables.Lock();
 		}
 		else
 		{
@@ -842,8 +841,7 @@ public:
 #if THREADSAFE_UOBJECTS
 		if (Tables)
 		{
-			UE_AUTORTFM_OPEN{ Tables->Unlock(); };
-			AutoRTFM::PopOnAbortHandler(this);
+			Tables->Unlock();
 		}
 #endif
 	}
