[ZenLoader]
- Add support for reading parallel load exclusions from ini settings
- Time spent recomputing non native classes is negligible when present (3.6ms total on 4 minute map load).

diff --git a/Engine/Config/BaseEngine.ini b/Engine/Config/BaseEngine.ini
--- a/Engine/Config/BaseEngine.ini
+++ b/Engine/Config/BaseEngine.ini
@@ -1652,6 +1652,20 @@ s.FlushStreamingOnExit=True
 +FixedBootOrder=/Script/Engine/Default__MaterialInterface
 +FixedBootOrder=/Script/Engine/Default__DeviceProfileManager
 
+[/Script/Engine.EditorStreamingSettings]
+; DataTable are loading other datatable from under a mutex
+; which is a deadlock hazard.
++ParallelPreloadExclusions=/Script/Engine.DataTable
+
+; UserDefinedStruct often try to load softobject ptr from text while
+; they are preloading. This is very deadlock prone so exclude those
+; types from running in parallel for now.
++ParallelPreloadExclusions=/Script/CoreUObject.UserDefinedStruct
+
+; Blueprint can recurse into other assets preload, which is a deadlock
+; hazard, exclude since the majority of their time is spent in create exports anyway.
++ParallelPreloadExclusions=/Script/Engine.Blueprint
+
 [/Script/Engine.GarbageCollectionSettings]
 gc.MaxObjectsNotConsideredByGC=1
 gc.FlushStreamingOnGC=0
diff --git a/Engine/Source/Runtime/CoreUObject/Private/Serialization/AsyncLoading2.cpp b/Engine/Source/Runtime/CoreUObject/Private/Serialization/AsyncLoading2.cpp
--- a/Engine/Source/Runtime/CoreUObject/Private/Serialization/AsyncLoading2.cpp
+++ b/Engine/Source/Runtime/CoreUObject/Private/Serialization/AsyncLoading2.cpp
@@ -25,6 +25,7 @@
 #include "Misc/CoreDelegates.h"
 #include "Misc/CommandLine.h"
 #include "Misc/App.h"
+#include "Misc/ConfigCacheIni.h"
 #include "Misc/MessageDialog.h"
 #include "Misc/PackageName.h"
 #include "Misc/StringBuilder.h"
@@ -8331,28 +8332,133 @@ FString FAsyncPackage2::GetNameFromPackageObjectIndex(FPackageObjectIndex Index)
 	return Out;
 }
 
-bool FAsyncPackage2::CanRunNodeAsync(const FAsyncLoadEventSpec* Spec) const
+namespace AsyncLoading2::Private
 {
-	if (Spec->Func == &FAsyncPackage2::Event_PreloadLinkerLoadExports)
+
+class FExcludedClassSet
+{
+	FName           Name;
+	TSet<FString>   ExcludedClassNames;
+	TSet<UClass*>   ExcludedClasses;
+	TSet<UClass*>   ExcludedNativeClasses;
+	uint64          ClassesVersionNumber = 0;
+	FDelegateHandle ConfigDelegate;
+
+	void EnsureClassesUpdated()
 	{
-		TRACE_CPUPROFILER_EVENT_SCOPE(FAsyncPackage2::CanRunNodeAsync);
+		if (ExcludedClassNames.Num())
+		{
+			uint64 CurrentVersion = GetRegisteredClassesVersionNumber();
+			if (ClassesVersionNumber != CurrentVersion)
+			{
+				TRACE_CPUPROFILER_EVENT_SCOPE(FExcludedClassSet::EnsureClassesUpdated);
+
+				ClassesVersionNumber = CurrentVersion;
+				ExcludedClasses.Reset();
+
+				for (auto It = ExcludedClassNames.CreateIterator(); It; ++It)
+				{
+					if (UClass* Class = FindObject<UClass>(nullptr, *It))
+					{
+						if (Class->IsNative())
+						{
+							ExcludedNativeClasses.Add(Class);
+
+							// We don't need to look at native classes ever again
+							It.RemoveCurrent();
+						}
+						else
+						{
+							ExcludedClasses.Add(Class);
+						}
+					}
+				}
+			}
+		}
+	}
+
+	void Reset()
+	{
+		ExcludedClassNames.Reset();
+		ExcludedClasses.Reset();
+		ExcludedNativeClasses.Reset();
+		ClassesVersionNumber = 0;
+	}
+
+	void OnConfigSectionChanged()
+	{
+		TRACE_CPUPROFILER_EVENT_SCOPE(FExcludedClassSet::OnConfigSectionChanged);
+
+		Reset();
+
+		if (GConfig && GConfig->IsReadyForUse())
+		{
+			const TCHAR* SectionName = FPlatformProperties::RequiresCookedData() ? TEXT("/Script/Engine.StreamingSettings") : TEXT("/Script/Engine.EditorStreamingSettings");
+
+			if (const FConfigSection* ConfigSection = GConfig->GetSection(SectionName, false, GEngineIni))
+			{
+				for (const auto& It : *ConfigSection)
+				{
+					if (It.Key == Name)
+					{
+						ExcludedClassNames.Add(It.Value.GetValue());
+					}
+				}
+			}
+		}
+	}
+public:
+	FExcludedClassSet(const TCHAR* InName)
+		: Name(InName)
+	{
+		ConfigDelegate = FCoreDelegates::TSOnConfigSectionsChanged().AddLambda(
+			[this](const FString&, const TSet<FString>&)
+			{
+				OnConfigSectionChanged();
+			}
+		);
+
+		OnConfigSectionChanged();
+	}
+
+	~FExcludedClassSet()
+	{
+		FCoreDelegates::TSOnConfigSectionsChanged().Remove(ConfigDelegate);
+	}
 
-		// This will be moved to .ini settings soon.
-		static TArray<UClass*> ExcludedClasses =
+	bool IsExcluded(UClass* Class)
+	{
+		EnsureClassesUpdated();
+
+		for (UClass* ExcludedClass : ExcludedClasses)
 		{
-			// DataTable are loading other datatable from under a mutex
-			// which is a deadlock hazard. Skip for now.
-			FindObject<UClass>(nullptr, TEXT("/Script/Engine.DataTable")),
+			if (Class->IsChildOf(ExcludedClass))
+			{
+				return true;
+			}
+		}
 
-			// UserDefinedStruct often try to load softobject ptr from text while
-			// they are preloading. This is very deadlock prone so exclude those
-			// types from going wide for now.
-			FindObject<UClass>(nullptr, TEXT("/Script/CoreUObject.UserDefinedStruct")),
+		for (UClass* ExcludedClass : ExcludedNativeClasses)
+		{
+			if (Class->IsChildOf(ExcludedClass))
+			{
+				return true;
+			}
+		}
+		
+		return false;
+	}
+};
 
-			// Blueprint can recurse into other assets preload, which is a deadlock
-			// hazard, exclude since the majority of their time is spent in create exports anyway.
-			FindObject<UClass>(nullptr, TEXT("/Script/Engine.Blueprint")),
-		};
+} // namespace AsyncLoading2::Private
+
+bool FAsyncPackage2::CanRunNodeAsync(const FAsyncLoadEventSpec* Spec) const
+{
+	if (Spec->Func == &FAsyncPackage2::Event_PreloadLinkerLoadExports)
+	{
+		TRACE_CPUPROFILER_EVENT_SCOPE(FAsyncPackage2::CanRunNodeAsync);
+
+		static AsyncLoading2::Private::FExcludedClassSet ParallelPreloadExclusions(TEXT("ParallelPreloadExclusions"));
 
 		const int32 ExportCount = Data.Exports.Num();
 		for (int32 ExportIndex = 0; ExportIndex < ExportCount; ++ExportIndex)
@@ -8361,12 +8467,9 @@ bool FAsyncPackage2::CanRunNodeAsync(const FAsyncLoadEventSpec* Spec) const
 
 			if (Export.Object)
 			{
-				for (UClass* ExcludedClass : ExcludedClasses)
+				if (ParallelPreloadExclusions.IsExcluded(Export.Object->GetClass()))
 				{
-					if (Export.Object->GetClass()->IsChildOf(ExcludedClass))
-					{
-						return false;
-					}
+					return false;
 				}
 			}
 		}
