Fix for compilation issue in RemoteExecutor when logging is disabled.

diff --git a/Engine/Source/Runtime/CoreUObject/Private/UObject/RemoteExecutor.cpp b/Engine/Source/Runtime/CoreUObject/Private/UObject/RemoteExecutor.cpp
--- a/Engine/Source/Runtime/CoreUObject/Private/UObject/RemoteExecutor.cpp
+++ b/Engine/Source/Runtime/CoreUObject/Private/UObject/RemoteExecutor.cpp
@@ -121,7 +121,7 @@ public:
 
 		double LastStallPrintTime = FPlatformTime::Seconds();
 
-#if !UE_BUILD_SHIPPING
+#if !UE_BUILD_SHIPPING && !NO_LOGGING
 		ELogVerbosity::Type InitialLogVerbosity = LogRemoteExec.GetVerbosity();
 		ON_SCOPE_EXIT
 		{
@@ -136,7 +136,7 @@ public:
 			double Now = FPlatformTime::Seconds();
 			bool bPrintStall = false;
 
-#if !UE_BUILD_SHIPPING
+#if !UE_BUILD_SHIPPING && !NO_LOGGING
 			constexpr int32 MaxIterationsBeforeStall = 100000;
 			if (LocalIterationNumber >= MaxIterationsBeforeStall && !(LocalIterationNumber % MaxIterationsBeforeStall))
 			{
