[CoreUObject]
* Added GetTransientPackageAsObject to prevent all the casts at callsites causing additional AST entries and also makes IWYU think it needs Package.h

diff --git a/Engine/Source/Runtime/CoreUObject/Private/UObject/Obj.cpp b/Engine/Source/Runtime/CoreUObject/Private/UObject/Obj.cpp
--- a/Engine/Source/Runtime/CoreUObject/Private/UObject/Obj.cpp
+++ b/Engine/Source/Runtime/CoreUObject/Private/UObject/Obj.cpp
@@ -5659,6 +5659,11 @@ UPackage* GetTransientPackage()
 	return GObjTransientPkg;
 }
 
+UObject* GetTransientPackageAsObject()
+{
+	return GObjTransientPkg;
+}
+
 UObject* GetTransientOuterForRename(UClass* ForClass)
 {
 	// if someone has tautologically placed themself within their own hierarchy then we'll
diff --git a/Engine/Source/Runtime/CoreUObject/Public/UObject/ObjectMacros.h b/Engine/Source/Runtime/CoreUObject/Public/UObject/ObjectMacros.h
--- a/Engine/Source/Runtime/CoreUObject/Public/UObject/ObjectMacros.h
+++ b/Engine/Source/Runtime/CoreUObject/Public/UObject/ObjectMacros.h
@@ -1836,7 +1836,7 @@ public: \
 		return TStaticCastFlags; \
 	} \
 	/** For internal use only; use StaticConstructObject() to create new objects. */ \
-	inline void* operator new(const size_t InSize, EInternal InInternalOnly, UObject* InOuter = (UObject*)GetTransientPackage(), FName InName = NAME_None, EObjectFlags InSetFlags = RF_NoFlags) \
+	inline void* operator new(const size_t InSize, EInternal InInternalOnly, UObject* InOuter = GetTransientPackageAsObject(), FName InName = NAME_None, EObjectFlags InSetFlags = RF_NoFlags) \
 	{ \
 		return StaticAllocateObject(StaticClass(), InOuter, InName, InSetFlags); \
 	} \
@@ -1894,7 +1894,7 @@ public: \
 	#define DEFINE_VTABLE_PTR_HELPER_CTOR_CALLER(TClass) \
 		static UObject* __VTableCtorCaller(FVTableHelper& Helper) \
 		{ \
-			return new (EC_InternalUseOnlyConstructor, (UObject*)GetTransientPackage(), NAME_None, RF_NeedLoad | RF_ClassDefaultObject | RF_TagGarbageTemp) TClass(Helper); \
+			return new (EC_InternalUseOnlyConstructor, GetTransientPackageAsObject(), NAME_None, RF_NeedLoad | RF_ClassDefaultObject | RF_TagGarbageTemp) TClass(Helper); \
 		}
 #else // WITH_RELOAD && !CHECK_PUREVIRTUALS
 	#define DEFINE_VTABLE_PTR_HELPER_CTOR_CALLER(TClass) \
@@ -1909,7 +1909,7 @@ public: \
 	static void IntrinsicClassInit(UClass* Class); \
 	static UObject* __VTableCtorCaller(FVTableHelper& Helper) \
 	{ \
-		return new (EC_InternalUseOnlyConstructor, (UObject*)GetTransientPackage(), NAME_None, RF_NeedLoad | RF_ClassDefaultObject | RF_TagGarbageTemp) TClass(Helper); \
+		return new (EC_InternalUseOnlyConstructor, GetTransientPackageAsObject(), NAME_None, RF_NeedLoad | RF_ClassDefaultObject | RF_TagGarbageTemp) TClass(Helper); \
 	}
 
 #define DECLARE_CLASS_INTRINSIC(TClass,TSuperClass,TStaticFlags,TPackage) \
@@ -1926,7 +1926,7 @@ public: \
 	static void IntrinsicClassInit(UClass* Class); \
 	static UObject* __VTableCtorCaller(FVTableHelper& Helper) \
 	{ \
-		return new (EC_InternalUseOnlyConstructor, (UObject*)GetTransientPackage(), NAME_None, RF_NeedLoad | RF_ClassDefaultObject | RF_TagGarbageTemp) TClass(Helper); \
+		return new (EC_InternalUseOnlyConstructor, GetTransientPackageAsObject(), NAME_None, RF_NeedLoad | RF_ClassDefaultObject | RF_TagGarbageTemp) TClass(Helper); \
 	}
 
 #define DECLARE_CASTED_CLASS_INTRINSIC_WITH_API_NO_CTOR( TClass, TSuperClass, TStaticFlags, TPackage, TStaticCastFlags, TRequiredAPI ) \
@@ -1937,7 +1937,7 @@ public: \
 	static void IntrinsicClassInit(UClass* Class); \
 	static UObject* __VTableCtorCaller(FVTableHelper& Helper) \
 	{ \
-		return new (EC_InternalUseOnlyConstructor, (UObject*)GetTransientPackage(), NAME_None, RF_NeedLoad | RF_ClassDefaultObject | RF_TagGarbageTemp) TClass(Helper); \
+		return new (EC_InternalUseOnlyConstructor, GetTransientPackageAsObject(), NAME_None, RF_NeedLoad | RF_ClassDefaultObject | RF_TagGarbageTemp) TClass(Helper); \
 	}
 
 #define DECLARE_CASTED_CLASS_INTRINSIC_WITH_API( TClass, TSuperClass, TStaticFlags, TPackage, TStaticCastFlags, TRequiredAPI ) \
@@ -1951,7 +1951,7 @@ public: \
 	static void IntrinsicClassInit(UClass* Class); \
 	static UObject* __VTableCtorCaller(FVTableHelper& Helper) \
 	{ \
-		return new (EC_InternalUseOnlyConstructor, (UObject*)GetTransientPackage(), NAME_None, RF_NeedLoad | RF_ClassDefaultObject | RF_TagGarbageTemp) TClass(Helper); \
+		return new (EC_InternalUseOnlyConstructor, GetTransientPackageAsObject(), NAME_None, RF_NeedLoad | RF_ClassDefaultObject | RF_TagGarbageTemp) TClass(Helper); \
 	}
 
 #define DECLARE_CASTED_CLASS_INTRINSIC_NO_CTOR_NO_VTABLE_CTOR( TClass, TSuperClass, TStaticFlags, TPackage, TStaticCastFlags, TRequiredAPI ) \
@@ -1962,7 +1962,7 @@ public: \
 	static void IntrinsicClassInit(UClass* Class); \
 	static UObject* __VTableCtorCaller(FVTableHelper& Helper) \
 	{ \
-		return new (EC_InternalUseOnlyConstructor, (UObject*)GetTransientPackage(), NAME_None, RF_NeedLoad | RF_ClassDefaultObject | RF_TagGarbageTemp) TClass(Helper); \
+		return new (EC_InternalUseOnlyConstructor, GetTransientPackageAsObject(), NAME_None, RF_NeedLoad | RF_ClassDefaultObject | RF_TagGarbageTemp) TClass(Helper); \
 	}
 
 #define DECLARE_CASTED_CLASS_INTRINSIC_NO_CTOR( TClass, TSuperClass, TStaticFlags, TPackage, TStaticCastFlags, TRequiredAPI ) \
diff --git a/Engine/Source/Runtime/CoreUObject/Public/UObject/UObjectGlobals.h b/Engine/Source/Runtime/CoreUObject/Public/UObject/UObjectGlobals.h
--- a/Engine/Source/Runtime/CoreUObject/Public/UObject/UObjectGlobals.h
+++ b/Engine/Source/Runtime/CoreUObject/Public/UObject/UObjectGlobals.h
@@ -254,6 +254,7 @@ COREUOBJECT_API TArray<const TCHAR*> ParsePropertyFlags(EPropertyFlags Flags);
 
 /** Returns the transient top-level package, which is useful for temporarily storing objects that should never be saved */
 COREUOBJECT_API UPackage* GetTransientPackage();
+COREUOBJECT_API UObject* GetTransientPackageAsObject();
 
 /** Returns an object in the transient package which respects the rules of Within */
 COREUOBJECT_API UObject* GetTransientOuterForRename(UClass* ForClass);
@@ -1910,7 +1911,7 @@ FUNCTION_NON_NULL_RETURN_END
 
 template< class T >
 FUNCTION_NON_NULL_RETURN_START
-	T* NewObject(UObject* Outer = (UObject*)GetTransientPackage())
+	T* NewObject(UObject* Outer = GetTransientPackageAsObject())
 FUNCTION_NON_NULL_RETURN_END
 {
 	// Name is always None for this case
