[IDO]
- Make enum, struct and class creation thread-safe

diff --git a/Engine/Source/Runtime/CoreUObject/Private/UObject/InstanceDataObjectUtils.cpp b/Engine/Source/Runtime/CoreUObject/Private/UObject/InstanceDataObjectUtils.cpp
--- a/Engine/Source/Runtime/CoreUObject/Private/UObject/InstanceDataObjectUtils.cpp
+++ b/Engine/Source/Runtime/CoreUObject/Private/UObject/InstanceDataObjectUtils.cpp
@@ -4,6 +4,7 @@
 
 #if WITH_EDITORONLY_DATA
 
+#include "Async/ManualResetEvent.h"
 #include "Async/SharedLock.h"
 #include "Async/SharedMutex.h"
 #include "Async/UniqueLock.h"
@@ -261,7 +262,25 @@ namespace UE
 	static const FName NAME_IDOMapKey(ANSITEXTVIEW("Key"));
 	static const FName NAME_IDOMapValue(ANSITEXTVIEW("Value"));
 
-	static TMap<FBlake3Hash, TWeakObjectPtr<UInstanceDataObjectClass>> IDOClassCache;
+	template <typename T>
+	struct TCacheItem
+	{
+		TWeakObjectPtr<T>     WeakPtr;
+		UE::FManualResetEvent ReadyEvent;
+
+		T* Extract()
+		{
+			if (!ReadyEvent.IsNotified())
+			{
+				TRACE_CPUPROFILER_EVENT_SCOPE(IDO_CacheItem_Wait);
+				ReadyEvent.Wait();
+			}
+
+			return WeakPtr.Get();
+		}
+	};
+
+	static TMap<FBlake3Hash, TSharedPtr<TCacheItem<UInstanceDataObjectClass>>> IDOClassCache;
 	static FSharedMutex IDOClassCacheMutex;
 
 	static void OnInstanceDataObjectSupportChanged(IConsoleVariable*);
@@ -495,9 +514,9 @@ namespace UE
 		{
 			TUniqueLock Lock(IDOClassCacheMutex);
 
-			for (const TPair<FBlake3Hash, TWeakObjectPtr<UInstanceDataObjectClass>>& Pair : IDOClassCache)
+			for (const TPair<FBlake3Hash, TSharedPtr<TCacheItem<UInstanceDataObjectClass>>>& Pair : IDOClassCache)
 			{
-				UInstanceDataObjectClass* Class = Pair.Value.Get();
+				UInstanceDataObjectClass* Class = Pair.Value->WeakPtr.Get();
 
 				if (Class)
 				{
@@ -708,13 +727,14 @@ namespace UE
 	}
 
 	template <typename T>
-	static void CleanUpInstanceDataObjectTypeCache(TMap<FBlake3Hash, TWeakObjectPtr<T>>& Cache)
+	static void CleanUpInstanceDataObjectTypeCache(TMap<FBlake3Hash, TSharedPtr<TCacheItem<T>>>& Cache)
 	{
 		if (Cache.Num() % 64 == 0)
 		{
 			for (auto It = Cache.CreateIterator(); It; ++It)
 			{
-				if (!It->Value.IsValid())
+				// If it's not notified, its construction is in progress, do not remove it.
+				if (!It->Value->WeakPtr.IsValid() && It->Value->ReadyEvent.IsNotified())
 				{
 					It.RemoveCurrent();
 				}
@@ -822,12 +842,31 @@ namespace UE
 		}
 
 		// Check if a cached enum exists for this key.
-		static TMap<FBlake3Hash, TWeakObjectPtr<UEnum>> EnumCache;
+		static TMap<FBlake3Hash, TSharedPtr<TCacheItem<UEnum>>> EnumCache;
 		static FSharedMutex EnumCacheMutex;
-		if (TSharedLock Lock(EnumCacheMutex); UEnum* Enum = EnumCache.FindRef(Key).Get())
+
+		TSharedPtr<TCacheItem<UEnum>> Item;
+		{
+			TSharedLock Lock(EnumCacheMutex);
+			Item = EnumCache.FindRef(Key);
+		}
+
+		if (UEnum* Enum = Item ? Item->Extract() : nullptr)
 		{
 			return Enum;
 		}
+		
+		{
+			TUniqueLock Lock(EnumCacheMutex);
+			CleanUpInstanceDataObjectTypeCache(EnumCache);
+			Item = EnumCache.FindRef(Key);
+			if (UEnum* Enum = Item ? Item->Extract() : nullptr)
+			{
+				return Enum;
+			}
+
+			Item = EnumCache.Add(Key, MakeShared<TCacheItem<UEnum>>());
+		}
 
 		// Construct a transient type that impersonates the original type.
 		const FName InstanceDataObjectName(WriteToString<128>(EnumName, TEXTVIEW("_InstanceDataObject")));
@@ -837,15 +876,9 @@ namespace UE
 
 		// TODO: Detect out-of-bounds values and increase the size of the underlying type accordingly.
 
-		TUniqueLock Lock(EnumCacheMutex);
-		if (UEnum* Enum = EnumCache.FindRef(Key).Get())
-		{
-			return Enum;
-		}
-
-		CleanUpInstanceDataObjectTypeCache(EnumCache);
+		Item->WeakPtr = NewEnum;
+		Item->ReadyEvent.Notify();
 
-		EnumCache.Add(Key, NewEnum);
 		return NewEnum;
 	}
 
@@ -1447,11 +1480,31 @@ namespace UE
 			Key = KeyBuilder.Finalize();
 		}
 
-		if (TSharedLock Lock(IDOClassCacheMutex); UClass* Class = IDOClassCache.FindRef(Key).Get())
+		TSharedPtr<TCacheItem<UInstanceDataObjectClass>> Item;
+		{
+			TSharedLock Lock(IDOClassCacheMutex);
+			Item = IDOClassCache.FindRef(Key);
+		}
+
+		if (UClass* Class = Item ? Item->Extract() : nullptr)
 		{
 			return Class;
 		}
 
+		{
+			TUniqueLock Lock(IDOClassCacheMutex);
+
+			CleanUpInstanceDataObjectTypeCache(IDOClassCache);
+
+			Item = IDOClassCache.FindRef(Key);
+			if (UClass* Class = Item ? Item->Extract() : nullptr)
+			{
+				return Class;
+			}
+
+			Item = IDOClassCache.Add(Key, MakeShared<TCacheItem<UInstanceDataObjectClass>>());
+		}
+
 		UInstanceDataObjectClass* NewClass = CreateInstanceDataObjectStructRec<UInstanceDataObjectClass>(OwnerClass, Outer, PropertyTree, EnumNames);
 		if (const FString& DisplayName = OwnerClass->GetMetaData(NAME_DisplayName); !DisplayName.IsEmpty())
 		{
@@ -1462,14 +1515,6 @@ namespace UE
 
 		CopyTaggedProperties(OwnerClass->GetDefaultObject(), NewClass->GetDefaultObject());
 
-		TUniqueLock Lock(IDOClassCacheMutex);
-		if (UClass* Class = IDOClassCache.FindRef(Key).Get())
-		{
-			return Class;
-		}
-
-		CleanUpInstanceDataObjectTypeCache(IDOClassCache);
-
 		// We need to copy the initialization state of the owner's properties for properties initialized in Verse code:
 		// e.g.: MyStruct<public>:test_struct = test_struct{ MyInt := 1 }
 		// This is necessary for Verse struct properties, because struct types do not store their initialization state 
@@ -1481,7 +1526,9 @@ namespace UE
 			TrackDefaultInitializedPropertiesInCDO(NewClass->GetDefaultObject(), NewClass, OwnerVerseClass);
 		}
 
-		IDOClassCache.Add(Key, NewClass);
+		Item->WeakPtr = NewClass;
+		Item->ReadyEvent.Notify();
+
 		return NewClass;
 	}
 
@@ -1504,27 +1551,41 @@ namespace UE
 			Key = KeyBuilder.Finalize();
 		}
 
-		static TMap<FBlake3Hash, TWeakObjectPtr<UInstanceDataObjectStruct>> StructCache;
+		static TMap<FBlake3Hash, TSharedPtr<TCacheItem<UInstanceDataObjectStruct>>> StructCache;
 		static FSharedMutex StructCacheMutex;
-		if (TSharedLock Lock(StructCacheMutex); UScriptStruct* Struct = StructCache.FindRef(Key).Get())
+		TSharedPtr<TCacheItem<UInstanceDataObjectStruct>> Item;
+		{
+			TSharedLock Lock(StructCacheMutex);
+			Item = StructCache.FindRef(Key);
+		}
+
+		if (UScriptStruct* Struct = Item ? Item->Extract() : nullptr)
 		{
 			return Struct;
 		}
 
+		{
+			TUniqueLock Lock(StructCacheMutex);
+
+			CleanUpInstanceDataObjectTypeCache(StructCache);
+
+			Item = StructCache.FindRef(Key);
+			if (UScriptStruct* Struct = Item ? Item->Extract() : nullptr)
+			{
+				return Struct;
+			}
+
+			Item = StructCache.Add(Key, MakeShared<TCacheItem<UInstanceDataObjectStruct>>());
+		}
+
 		UInstanceDataObjectStruct* NewStruct = CreateInstanceDataObjectStructRec<UInstanceDataObjectStruct>(OwnerStruct, Outer, PropertyTree, EnumNames);
 		NewStruct->Guid = Guid;
 		NewStruct->SetMetaData(NAME_OriginalType, OriginalName);
 		NewStruct->SetMetaData(NAME_PresentAsTypeMetadata, OriginalName);
 
-		TUniqueLock Lock(StructCacheMutex);
-		if (UScriptStruct* Struct = StructCache.FindRef(Key).Get())
-		{
-			return Struct;
-		}
-
-		CleanUpInstanceDataObjectTypeCache(StructCache);
+		Item->WeakPtr = NewStruct;
+		Item->ReadyEvent.Notify();
 
-		StructCache.Add(Key, NewStruct);
 		return NewStruct;
 	}
 
