Add a macro to enable / disable minsize and mark up a few notable functions.

- Also reworked RegisterNativeRedirects to avoid inlining.

diff --git a/Engine/Source/Runtime/Core/Public/Clang/ClangPlatformCompilerPreSetup.h b/Engine/Source/Runtime/Core/Public/Clang/ClangPlatformCompilerPreSetup.h
--- a/Engine/Source/Runtime/Core/Public/Clang/ClangPlatformCompilerPreSetup.h
+++ b/Engine/Source/Runtime/Core/Public/Clang/ClangPlatformCompilerPreSetup.h
@@ -72,6 +72,14 @@
 	#define PRAGMA_RESTORE_UNREACHABLE_CODE_WARNINGS
 #endif // PRAGMA_RESTORE_UNREACHABLE_CODE_WARNINGS
 
+#ifndef PRAGMA_PUSH_ATTRIBUTE_MINSIZE_FUNCTIONS
+	#define PRAGMA_PUSH_ATTRIBUTE_MINSIZE_FUNCTIONS _Pragma("clang attribute push(__attribute__((minsize)), apply_to=function)")
+#endif // PRAGMA_PUSH_ATTRIBUTE_MINSIZE_FUNCTIONS
+
+#ifndef PRAGMA_POP_ATTRIBUTE_MINSIZE_FUNCTIONS
+	#define PRAGMA_POP_ATTRIBUTE_MINSIZE_FUNCTIONS _Pragma("clang attribute pop")
+#endif // PRAGMA_POP_ATTRIBUTE_MINSIZE_FUNCTIONS
+
 #if __has_warning("-Wimplicit-float-conversion")
 #define DISABLE_IMPLICIT_FLOAT_CONVERSION_FRAGMENT _Pragma("clang diagnostic ignored \"-Wimplicit-float-conversion\"")
 #else
diff --git a/Engine/Source/Runtime/Core/Public/Microsoft/MSVCPlatformCompilerPreSetup.h b/Engine/Source/Runtime/Core/Public/Microsoft/MSVCPlatformCompilerPreSetup.h
--- a/Engine/Source/Runtime/Core/Public/Microsoft/MSVCPlatformCompilerPreSetup.h
+++ b/Engine/Source/Runtime/Core/Public/Microsoft/MSVCPlatformCompilerPreSetup.h
@@ -281,6 +281,14 @@
 	#define EMIT_CUSTOM_ERROR_AT_LINE(Line, Error) \
 		__pragma(message(WARNING_LOCATION(Line) ": error C4996: " Error))
 
+	#ifndef PRAGMA_PUSH_ATTRIBUTE_MINSIZE_FUNCTIONS
+		#define PRAGMA_PUSH_ATTRIBUTE_MINSIZE_FUNCTIONS // no-op for now
+	#endif // PRAGMA_PUSH_ATTRIBUTE_MINSIZE_FUNCTIONS
+
+	#ifndef PRAGMA_POP_ATTRIBUTE_MINSIZE_FUNCTIONS
+		#define PRAGMA_POP_ATTRIBUTE_MINSIZE_FUNCTIONS // no-op for now
+	#endif // PRAGMA_POP_ATTRIBUTE_MINSIZE_FUNCTIONS
+
 	#ifndef PGO_LINK_DISABLE_WARNINGS
 		#define PGO_LINK_DISABLE_WARNINGS PRAGMA_DISABLE_BUFFER_OVERRUN_WARNING
 	#endif
diff --git a/Engine/Source/Runtime/CoreUObject/Private/UObject/CoreRedirects.cpp b/Engine/Source/Runtime/CoreUObject/Private/UObject/CoreRedirects.cpp
--- a/Engine/Source/Runtime/CoreUObject/Private/UObject/CoreRedirects.cpp
+++ b/Engine/Source/Runtime/CoreUObject/Private/UObject/CoreRedirects.cpp
@@ -3161,7 +3161,7 @@ ECoreRedirectFlags FCoreRedirects::GetFlagsForTypeClass(UClass *TypeClass)
 
 #if UE_WITH_CORE_REDIRECTS
 
-UE_DISABLE_OPTIMIZATION_SHIP
+PRAGMA_PUSH_ATTRIBUTE_MINSIZE_FUNCTIONS
 
 template <typename T>
 FORCEINLINE static T& DiscardNoDiscard(T& Ref)
@@ -3169,14 +3169,49 @@ FORCEINLINE static T& DiscardNoDiscard(T& Ref)
 	return Ref;
 }
 
+static FORCENOINLINE FCoreRedirect& ClassRedirectImpl(TArray<FCoreRedirect>& Redirects, const TCHAR* OldName, const TCHAR* NewName)
+{
+	return DiscardNoDiscard(Redirects.Emplace_GetRef(ECoreRedirectFlags::Type_Class, OldName, NewName));
+}
+
+static FORCENOINLINE FCoreRedirect& ClassRedirectInstancesImpl(TArray<FCoreRedirect>& Redirects, const TCHAR* OldName, const TCHAR* NewName)
+{
+	return DiscardNoDiscard(Redirects.Emplace_GetRef(ECoreRedirectFlags::Type_Class | ECoreRedirectFlags::Category_InstanceOnly, OldName, NewName));
+}
+
+static FORCENOINLINE FCoreRedirect& StructRedirectImpl(TArray<FCoreRedirect>& Redirects, const TCHAR* OldName, const TCHAR* NewName)
+{
+	return DiscardNoDiscard(Redirects.Emplace_GetRef(ECoreRedirectFlags::Type_Struct, OldName, NewName));
+}
+
+static FORCENOINLINE FCoreRedirect& EnumRedirectImpl(TArray<FCoreRedirect>& Redirects, const TCHAR* OldName, const TCHAR* NewName)
+{
+	return DiscardNoDiscard(Redirects.Emplace_GetRef(ECoreRedirectFlags::Type_Enum, OldName, NewName));
+}
+
+static FORCENOINLINE FCoreRedirect& PropertyRedirectImpl(TArray<FCoreRedirect>& Redirects, const TCHAR* OldName, const TCHAR* NewName)
+{
+	return DiscardNoDiscard(Redirects.Emplace_GetRef(ECoreRedirectFlags::Type_Property, OldName, NewName));
+}
+
+static FORCENOINLINE FCoreRedirect& FunctionRedirectImpl(TArray<FCoreRedirect>& Redirects, const TCHAR* OldName, const TCHAR* NewName)
+{
+	return DiscardNoDiscard(Redirects.Emplace_GetRef(ECoreRedirectFlags::Type_Function, OldName, NewName));
+}
+
+static FORCENOINLINE FCoreRedirect& PackageRedirectImpl(TArray<FCoreRedirect>& Redirects, const TCHAR* OldName, const TCHAR* NewName)
+{
+	return DiscardNoDiscard(Redirects.Emplace_GetRef(ECoreRedirectFlags::Type_Package, OldName, NewName));
+}
+
 // The compiler doesn't like having a massive string table in a single function so split it up
-#define CLASS_REDIRECT(OldName, NewName) DiscardNoDiscard(Redirects.Emplace_GetRef(ECoreRedirectFlags::Type_Class, TEXT(OldName), TEXT(NewName)))
-#define CLASS_REDIRECT_INSTANCES(OldName, NewName) DiscardNoDiscard(Redirects.Emplace_GetRef(ECoreRedirectFlags::Type_Class | ECoreRedirectFlags::Category_InstanceOnly, TEXT(OldName), TEXT(NewName)))
-#define STRUCT_REDIRECT(OldName, NewName) DiscardNoDiscard(Redirects.Emplace_GetRef(ECoreRedirectFlags::Type_Struct, TEXT(OldName), TEXT(NewName)))
-#define ENUM_REDIRECT(OldName, NewName) DiscardNoDiscard(Redirects.Emplace_GetRef(ECoreRedirectFlags::Type_Enum, TEXT(OldName), TEXT(NewName)))
-#define PROPERTY_REDIRECT(OldName, NewName) DiscardNoDiscard(Redirects.Emplace_GetRef(ECoreRedirectFlags::Type_Property, TEXT(OldName), TEXT(NewName)))
-#define FUNCTION_REDIRECT(OldName, NewName) DiscardNoDiscard(Redirects.Emplace_GetRef(ECoreRedirectFlags::Type_Function, TEXT(OldName), TEXT(NewName)))
-#define PACKAGE_REDIRECT(OldName, NewName) DiscardNoDiscard(Redirects.Emplace_GetRef(ECoreRedirectFlags::Type_Package, TEXT(OldName), TEXT(NewName)))
+#define CLASS_REDIRECT(OldName, NewName) ClassRedirectImpl(Redirects, TEXT(OldName), TEXT(NewName))
+#define CLASS_REDIRECT_INSTANCES(OldName, NewName) ClassRedirectInstancesImpl(Redirects, TEXT(OldName), TEXT(NewName))
+#define STRUCT_REDIRECT(OldName, NewName) StructRedirectImpl(Redirects, TEXT(OldName), TEXT(NewName))
+#define ENUM_REDIRECT(OldName, NewName) EnumRedirectImpl(Redirects, TEXT(OldName), TEXT(NewName))
+#define PROPERTY_REDIRECT(OldName, NewName) PropertyRedirectImpl(Redirects, TEXT(OldName), TEXT(NewName))
+#define FUNCTION_REDIRECT(OldName, NewName) FunctionRedirectImpl(Redirects, TEXT(OldName), TEXT(NewName))
+#define PACKAGE_REDIRECT(OldName, NewName) PackageRedirectImpl(Redirects, TEXT(OldName), TEXT(NewName))
 
 static void RegisterNativeRedirects40(TArray<FCoreRedirect>& Redirects)
 {
@@ -3906,7 +3941,7 @@ static void RegisterNativeRedirects49(TArray<FCoreRedirect>& Redirects)
 	ENUM_REDIRECT("EInterpolationBlend", "/Script/AnimGraphRuntime.EInterpolationBlend");
 }
 
-UE_ENABLE_OPTIMIZATION_SHIP
+PRAGMA_POP_ATTRIBUTE_MINSIZE_FUNCTIONS
 
 void FCoreRedirects::RegisterNativeRedirectsUnderWriteLock(FScopeCoreRedirectsWriteLockedContext& LockedContext)
 {
diff --git a/Engine/Source/Runtime/SlateCore/Private/Styling/StarshipCoreStyle.cpp b/Engine/Source/Runtime/SlateCore/Private/Styling/StarshipCoreStyle.cpp
--- a/Engine/Source/Runtime/SlateCore/Private/Styling/StarshipCoreStyle.cpp
+++ b/Engine/Source/Runtime/SlateCore/Private/Styling/StarshipCoreStyle.cpp
@@ -149,6 +149,7 @@ void FStarshipCoreStyle::SetFocusBrush(FSlateBrush* NewBrush)
 	FSlateStyleRegistry::RegisterSlateStyle(Style.Get());
 }
 
+PRAGMA_PUSH_ATTRIBUTE_MINSIZE_FUNCTIONS
 
 TSharedRef<ISlateStyle> FStarshipCoreStyle::Create()
 {
@@ -1047,6 +1048,7 @@ TSharedRef<ISlateStyle> FStarshipCoreStyle::Create()
 	return Style;
 }
 
+PRAGMA_POP_ATTRIBUTE_MINSIZE_FUNCTIONS
 
 const TSharedPtr<FSlateDynamicImageBrush> FStarshipCoreStyle::GetDynamicImageBrush(FName BrushTemplate, FName TextureName, const ANSICHAR* Specifier)
 {
diff --git a/Engine/Source/Runtime/SlateCore/Private/Styling/UMGCoreStyle.cpp b/Engine/Source/Runtime/SlateCore/Private/Styling/UMGCoreStyle.cpp
--- a/Engine/Source/Runtime/SlateCore/Private/Styling/UMGCoreStyle.cpp
+++ b/Engine/Source/Runtime/SlateCore/Private/Styling/UMGCoreStyle.cpp
@@ -195,6 +195,8 @@ void FUMGCoreStyle::SetFocusBrush(FSlateBrush* NewBrush)
 	FSlateStyleRegistry::RegisterSlateStyle(Style.Get());
 }
 
+PRAGMA_PUSH_ATTRIBUTE_MINSIZE_FUNCTIONS
+
 TSharedRef<ISlateStyle> FUMGCoreStyle::Create()
 {
 	TSharedRef<FUMGStyleSet> Style = MakeShareable(new FUMGStyleSet("UMGCoreStyle"));
@@ -1423,6 +1425,7 @@ TSharedRef<ISlateStyle> FUMGCoreStyle::Create()
 	return Style;
 }
 
+PRAGMA_POP_ATTRIBUTE_MINSIZE_FUNCTIONS
 
 void FUMGCoreStyle::SetStyle(const TSharedRef< ISlateStyle >& NewStyle)
 {
