Replace some FORCEINLINE usage in Core/CoreUObject with UE_FORCEINLINE_HINT.

diff --git a/Engine/Source/Runtime/CoreUObject/Public/Serialization/PropertyLocalizationDataGathering.h b/Engine/Source/Runtime/CoreUObject/Public/Serialization/PropertyLocalizationDataGathering.h
--- a/Engine/Source/Runtime/CoreUObject/Public/Serialization/PropertyLocalizationDataGathering.h
+++ b/Engine/Source/Runtime/CoreUObject/Public/Serialization/PropertyLocalizationDataGathering.h
@@ -146,12 +146,12 @@ public:
 	static COREUOBJECT_API FLocalizationDataObjectGatheringCallbackMap& GetTypeSpecificLocalizationDataObjectGatheringCallbacks();
 	static COREUOBJECT_API FLocalizationDataStructGatheringCallbackMap& GetTypeSpecificLocalizationDataStructGatheringCallbacks();
 
-	FORCEINLINE TArray<FGatherableTextData>& GetGatherableTextDataArray() const
+	UE_FORCEINLINE_HINT TArray<FGatherableTextData>& GetGatherableTextDataArray() const
 	{
 		return GatherableTextDataArray;
 	}
 
-	FORCEINLINE bool IsObjectValidForGather(const UObject* Object) const
+	UE_FORCEINLINE_HINT bool IsObjectValidForGather(const UObject* Object) const
 	{
 		return AllObjectsInPackage.Contains(Object);
 	}
@@ -177,7 +177,7 @@ private:
 				&& GatherTextFlags == Other.GatherTextFlags;
 		}
 
-		FORCEINLINE bool operator!=(const FObjectAndGatherFlags& Other) const
+		UE_FORCEINLINE_HINT bool operator!=(const FObjectAndGatherFlags& Other) const
 		{
 			return !(*this == Other);
 		}
@@ -206,12 +206,12 @@ private:
 /** Struct to automatically register a callback when it's constructed */
 struct FAutoRegisterLocalizationDataGatheringCallback
 {
-	FORCEINLINE FAutoRegisterLocalizationDataGatheringCallback(const UClass* InClass, const FPropertyLocalizationDataGatherer::FLocalizationDataObjectGatheringCallback& InCallback)
+	UE_FORCEINLINE_HINT FAutoRegisterLocalizationDataGatheringCallback(const UClass* InClass, const FPropertyLocalizationDataGatherer::FLocalizationDataObjectGatheringCallback& InCallback)
 	{
 		FPropertyLocalizationDataGatherer::GetTypeSpecificLocalizationDataObjectGatheringCallbacks().Add(InClass, InCallback);
 	}
 
-	FORCEINLINE FAutoRegisterLocalizationDataGatheringCallback(const UScriptStruct* InStruct, const FPropertyLocalizationDataGatherer::FLocalizationDataStructGatheringCallback& InCallback)
+	UE_FORCEINLINE_HINT FAutoRegisterLocalizationDataGatheringCallback(const UScriptStruct* InStruct, const FPropertyLocalizationDataGatherer::FLocalizationDataStructGatheringCallback& InCallback)
 	{
 		FPropertyLocalizationDataGatherer::GetTypeSpecificLocalizationDataStructGatheringCallbacks().Add(InStruct, InCallback);
 	}
diff --git a/Engine/Source/Runtime/CoreUObject/Public/Serialization/TextReferenceCollector.h b/Engine/Source/Runtime/CoreUObject/Public/Serialization/TextReferenceCollector.h
--- a/Engine/Source/Runtime/CoreUObject/Public/Serialization/TextReferenceCollector.h
+++ b/Engine/Source/Runtime/CoreUObject/Public/Serialization/TextReferenceCollector.h
@@ -38,7 +38,7 @@ public:
 /** Struct to automatically register a callback when it's constructed */
 struct FAutoRegisterTextReferenceCollectorCallback
 {
-	FORCEINLINE FAutoRegisterTextReferenceCollectorCallback(const UClass* InClass, const FTextReferenceCollector::FTextReferenceCollectorCallback& InCallback)
+	UE_FORCEINLINE_HINT FAutoRegisterTextReferenceCollectorCallback(const UClass* InClass, const FTextReferenceCollector::FTextReferenceCollectorCallback& InCallback)
 	{
 		FTextReferenceCollector::GetTypeSpecificTextReferenceCollectorCallbacks().Add(InClass, InCallback);
 	}
diff --git a/Engine/Source/Runtime/CoreUObject/Public/StructUtils/PropertyBag.h b/Engine/Source/Runtime/CoreUObject/Public/StructUtils/PropertyBag.h
--- a/Engine/Source/Runtime/CoreUObject/Public/StructUtils/PropertyBag.h
+++ b/Engine/Source/Runtime/CoreUObject/Public/StructUtils/PropertyBag.h
@@ -144,12 +144,12 @@ struct FPropertyBagContainerTypes
 
 	UE_API bool operator == (const FPropertyBagContainerTypes& Other) const;
 
-	FORCEINLINE bool operator !=(const FPropertyBagContainerTypes& Other) const
+	UE_FORCEINLINE_HINT bool operator !=(const FPropertyBagContainerTypes& Other) const
 	{
 		return !(Other == *this);
 	}
 
-	friend FORCEINLINE uint32 GetTypeHash(const FPropertyBagContainerTypes& PropertyBagContainerTypes)
+	friend UE_FORCEINLINE_HINT uint32 GetTypeHash(const FPropertyBagContainerTypes& PropertyBagContainerTypes)
 	{
 		return GetArrayHash(PropertyBagContainerTypes.Types.GetData(), PropertyBagContainerTypes.NumContainers);
 	}
@@ -222,7 +222,7 @@ struct FPropertyBagPropertyDescMetaData
 		return Ar;
 	}
 
-	friend FORCEINLINE uint32 GetTypeHash(const FPropertyBagPropertyDescMetaData& PropertyDescMetaData)
+	friend UE_FORCEINLINE_HINT uint32 GetTypeHash(const FPropertyBagPropertyDescMetaData& PropertyDescMetaData)
 	{
 		return HashCombine(GetTypeHash(PropertyDescMetaData.Key), GetTypeHash(PropertyDescMetaData.Value));
 	}
@@ -237,7 +237,7 @@ struct FPropertyBagPropertyDescMetaData
 		return Hash;
 	}
 
-	friend FORCEINLINE uint32 GetTypeHash(const TArray<FPropertyBagPropertyDescMetaData>& MetaData)
+	friend UE_FORCEINLINE_HINT uint32 GetTypeHash(const TArray<FPropertyBagPropertyDescMetaData>& MetaData)
 	{
 		return GetTypeHash(TArrayView<const FPropertyBagPropertyDescMetaData>(MetaData.GetData(), MetaData.Num()));
 	}
@@ -1213,7 +1213,7 @@ public:
 	}
 
 	/** Returns number of elements in set. */
-	FORCEINLINE int32 Num() const
+	UE_FORCEINLINE_HINT int32 Num() const
 	{
 		return FScriptSetHelper::Num();
 	}
diff --git a/Engine/Source/Runtime/CoreUObject/Public/StructUtils/StructArrayView.h b/Engine/Source/Runtime/CoreUObject/Public/StructUtils/StructArrayView.h
--- a/Engine/Source/Runtime/CoreUObject/Public/StructUtils/StructArrayView.h
+++ b/Engine/Source/Runtime/CoreUObject/Public/StructUtils/StructArrayView.h
@@ -151,7 +151,7 @@ struct FStructArrayView
 	 *
 	 * @returns Size in bytes of array type.
 	 */
-	FORCEINLINE_DEBUGGABLE uint32 GetTypeSize() const
+	UE_FORCEINLINE_HINT uint32 GetTypeSize() const
 	{
 		return ElementSize;
 	}
@@ -359,7 +359,7 @@ struct FStructArrayView
 				&& Index == Other.Index;
 		}
 
-		FORCEINLINE bool operator != (const FIterator& Other) const
+		UE_FORCEINLINE_HINT bool operator != (const FIterator& Other) const
 		{
 			return !(*this == Other);
 		}
@@ -370,11 +370,11 @@ struct FStructArrayView
 	};
 
 	/** Ranged iteration support. DO NOT USE DIRECTLY. */
-	FORCEINLINE FIterator begin() const { return FIterator(*this, 0); }
-	FORCEINLINE FIterator end() const { return FIterator(*this, Num()); }
+	UE_FORCEINLINE_HINT FIterator begin() const { return FIterator(*this, 0); }
+	UE_FORCEINLINE_HINT FIterator end() const { return FIterator(*this, Num()); }
 
 	/** Swaps the elements at the specified Indicies. The indicies must be valid! */
-	FORCEINLINE void Swap(const int32 Index1, const int32 Index2)
+	UE_FORCEINLINE_HINT void Swap(const int32 Index1, const int32 Index2)
 	{
 		FMemory::Memswap(GetDataAt(Index1), GetDataAt(Index2), ElementSize); // GetDataAt() calls RangeCheck().
 	}
@@ -528,7 +528,7 @@ struct FConstStructArrayView
 	 *
 	 * @returns Size in bytes of array type.
 	 */
-	FORCEINLINE_DEBUGGABLE uint32 GetTypeSize()
+	UE_FORCEINLINE_HINT uint32 GetTypeSize()
 	{
 		return ElementSize;
 	}
@@ -697,7 +697,7 @@ struct FConstStructArrayView
 				&& Index == Other.Index;
 		}
 
-		FORCEINLINE bool operator != (const FIterator& Other) const
+		UE_FORCEINLINE_HINT bool operator != (const FIterator& Other) const
 		{
 			return !(*this == Other);
 		}
@@ -708,8 +708,8 @@ struct FConstStructArrayView
 	};
 
 	/** Ranged iteration support. DO NOT USE DIRECTLY. */
-	FORCEINLINE FIterator begin() const { return FIterator(*this, 0); }
-	FORCEINLINE FIterator end() const { return FIterator(*this, Num()); }
+	UE_FORCEINLINE_HINT FIterator begin() const { return FIterator(*this, 0); }
+	UE_FORCEINLINE_HINT FIterator end() const { return FIterator(*this, Num()); }
 
 private:
 	const void* DataPtr = nullptr;
diff --git a/Engine/Source/Runtime/CoreUObject/Public/StructUtils/StructTypeBitSet.h b/Engine/Source/Runtime/CoreUObject/Public/StructUtils/StructTypeBitSet.h
--- a/Engine/Source/Runtime/CoreUObject/Public/StructUtils/StructTypeBitSet.h
+++ b/Engine/Source/Runtime/CoreUObject/Public/StructUtils/StructTypeBitSet.h
@@ -589,7 +589,7 @@ struct TTypeBitSetBase
 	 * @param Other - The other bitset to intersect with.
 	 * @return A new bitset representing the intersection.
 	 */
-	FORCEINLINE TImplementation operator&(const TImplementation& Other) const
+	UE_FORCEINLINE_HINT TImplementation operator&(const TImplementation& Other) const
 	{
 		return TImplementation(TBitArray<>::BitwiseAND(StructTypesBitArray, Other.StructTypesBitArray, EBitwiseOperatorFlags::MinSize));
 	}
@@ -599,12 +599,12 @@ struct TTypeBitSetBase
 	 * @param Other - The other bitset to union with.
 	 * @return A new bitset representing the union.
 	 */
-	FORCEINLINE TImplementation operator|(const TImplementation& Other) const
+	UE_FORCEINLINE_HINT TImplementation operator|(const TImplementation& Other) const
 	{
 		return TImplementation(TBitArray<>::BitwiseOR(StructTypesBitArray, Other.StructTypesBitArray, EBitwiseOperatorFlags::MaxSize));
 	}
 
-	FORCEINLINE TImplementation GetOverlap(const TImplementation& Other) const
+	UE_FORCEINLINE_HINT TImplementation GetOverlap(const TImplementation& Other) const
 	{
 		return GetImplementation() & Other;
 	}
@@ -615,22 +615,22 @@ struct TTypeBitSetBase
 	 * @param Other - The other bitset to compare.
 	 * @return True if both bitsets are equivalent; false otherwise.
 	 */
-	FORCEINLINE bool IsEquivalent(const TImplementation& Other) const
+	UE_FORCEINLINE_HINT bool IsEquivalent(const TImplementation& Other) const
 	{
 		return StructTypesBitArray.CompareSetBits(Other.StructTypesBitArray, /*bMissingBitValue=*/false);
 	}
 
-	FORCEINLINE bool HasAll(const TImplementation& Other) const
+	UE_FORCEINLINE_HINT bool HasAll(const TImplementation& Other) const
 	{
 		return StructTypesBitArray.HasAll(Other.StructTypesBitArray);
 	}
 
-	FORCEINLINE bool HasAny(const TImplementation& Other) const
+	UE_FORCEINLINE_HINT bool HasAny(const TImplementation& Other) const
 	{
 		return StructTypesBitArray.HasAny(Other.StructTypesBitArray);
 	}
 
-	FORCEINLINE bool HasNone(const TImplementation& Other) const
+	UE_FORCEINLINE_HINT bool HasNone(const TImplementation& Other) const
 	{
 		return !StructTypesBitArray.HasAny(Other.StructTypesBitArray);
 	}
@@ -649,7 +649,7 @@ struct TTypeBitSetBase
 	 * @param BitIndex - The index of the bit to check.
 	 * @return True if the bit is set; false otherwise.
 	 */
-	FORCEINLINE bool IsBitSet(const int32 BitIndex) const
+	UE_FORCEINLINE_HINT bool IsBitSet(const int32 BitIndex) const
 	{
 		return StructTypesBitArray.Contains(BitIndex);
 	}
@@ -658,7 +658,7 @@ struct TTypeBitSetBase
 	 * Adds the bits from another bitset to this one (union).
 	 * @param Other - The other bitset to add.
 	 */
-	FORCEINLINE void operator+=(const TImplementation& Other)
+	UE_FORCEINLINE_HINT void operator+=(const TImplementation& Other)
 	{
 		StructTypesBitArray = TBitArray<>::BitwiseOR(StructTypesBitArray, Other.StructTypesBitArray, EBitwiseOperatorFlags::MaxSize);
 	}
@@ -667,7 +667,7 @@ struct TTypeBitSetBase
 	 * Removes the bits from another bitset from this one (difference).
 	 * @param Other - The other bitset to subtract.
 	 */
-	FORCEINLINE void operator-=(const TImplementation& Other)
+	UE_FORCEINLINE_HINT void operator-=(const TImplementation& Other)
 	{
 		StructTypesBitArray -= Other.StructTypesBitArray;
 	}
@@ -1016,7 +1016,7 @@ public:
 	 * Removes the bits from another bitset from this one (difference).
 	 * @param Other - The other bitset to subtract.
 	 */
-	FORCEINLINE void Remove(const TStructTypeBitSet& Other)
+	UE_FORCEINLINE_HINT void Remove(const TStructTypeBitSet& Other)
 	{
 		StructTypesBitArray -= Other.StructTypesBitArray;
 	}
@@ -1060,7 +1060,7 @@ public:
 	 * @param Other - The other bitset to compare
 	 * @return True if the bitsets are equal; false otherwise.
 	 */
-	FORCEINLINE bool operator==(const TStructTypeBitSet& Other) const
+	UE_FORCEINLINE_HINT bool operator==(const TStructTypeBitSet& Other) const
 	{
 		return Super::StructTypesBitArray == Other.StructTypesBitArray;
 	}
@@ -1070,7 +1070,7 @@ public:
 	 * @param Other - The other bitset to compare.
 	 * @return True if the bitsets are not equal; false otherwise.
 	 */
-	FORCEINLINE bool operator!=(const TStructTypeBitSet& Other) const
+	UE_FORCEINLINE_HINT bool operator!=(const TStructTypeBitSet& Other) const
 	{
 		return !(*this == Other);
 	}
diff --git a/Engine/Source/Runtime/CoreUObject/Public/StructUtils/UserDefinedStruct.h b/Engine/Source/Runtime/CoreUObject/Public/StructUtils/UserDefinedStruct.h
--- a/Engine/Source/Runtime/CoreUObject/Public/StructUtils/UserDefinedStruct.h
+++ b/Engine/Source/Runtime/CoreUObject/Public/StructUtils/UserDefinedStruct.h
@@ -141,7 +141,7 @@ public:
 	COREUOBJECT_API void UpdateStructFlags();
 
 #if WITH_EDITORONLY_DATA
-	FORCEINLINE FOnStructChanged& OnStructChanged() { return ChangedEvent; }
+	UE_FORCEINLINE_HINT FOnStructChanged& OnStructChanged() { return ChangedEvent; }
 public:
 	FOnStructChanged ChangedEvent;
 
diff --git a/Engine/Source/Runtime/CoreUObject/Public/Templates/Casts.h b/Engine/Source/Runtime/CoreUObject/Public/Templates/Casts.h
--- a/Engine/Source/Runtime/CoreUObject/Public/Templates/Casts.h
+++ b/Engine/Source/Runtime/CoreUObject/Public/Templates/Casts.h
@@ -171,7 +171,7 @@ inline TCopyQualifiersFromTo_T<From, To>* Cast(From* Src)
 }
 
 template <typename To, typename From>
-FORCEINLINE TCopyQualifiersFromTo_T<From, To>* ExactCast(From* Src)
+UE_FORCEINLINE_HINT TCopyQualifiersFromTo_T<From, To>* ExactCast(From* Src)
 {
 	return Src && (Src->GetClass() == To::StaticClass()) ? (To*)Src : nullptr;
 }
@@ -279,7 +279,7 @@ FORCEINLINE TCopyQualifiersFromTo_T<From, To>* ExactCast(From* Src)
 	}
 
 	template <typename To, typename From>
-	FORCEINLINE TCopyQualifiersFromTo_T<From, To>* CastChecked(From* Src, ECastCheckedType::Type CheckType)
+	UE_FORCEINLINE_HINT TCopyQualifiersFromTo_T<From, To>* CastChecked(From* Src, ECastCheckedType::Type CheckType)
 	{
 		return CastChecked<To>(Src);
 	}
@@ -288,22 +288,22 @@ FORCEINLINE TCopyQualifiersFromTo_T<From, To>* ExactCast(From* Src)
 
 // auto weak versions
 template <typename To, typename From>
-FORCEINLINE TCopyQualifiersFromTo_T<From, To>* Cast(const TWeakObjectPtr<From>& Src)
+UE_FORCEINLINE_HINT TCopyQualifiersFromTo_T<From, To>* Cast(const TWeakObjectPtr<From>& Src)
 {
 	return Cast<To>(Src.Get());
 }
 template <typename To, typename From>
-FORCEINLINE TCopyQualifiersFromTo_T<From, To>* ExactCast(const TWeakObjectPtr<From>& Src)
+UE_FORCEINLINE_HINT TCopyQualifiersFromTo_T<From, To>* ExactCast(const TWeakObjectPtr<From>& Src)
 {
 	return ExactCast  <To>(Src.Get());
 }
 template <typename To, typename From>
-FORCEINLINE TCopyQualifiersFromTo_T<From, To>* CastChecked(const TWeakObjectPtr<From>& Src)
+UE_FORCEINLINE_HINT TCopyQualifiersFromTo_T<From, To>* CastChecked(const TWeakObjectPtr<From>& Src)
 {
 	return CastChecked<To>(Src.Get());
 }
 template <typename To, typename From>
-FORCEINLINE TCopyQualifiersFromTo_T<From, To>* CastChecked(const TWeakObjectPtr<From>& Src, ECastCheckedType::Type CheckType)
+UE_FORCEINLINE_HINT TCopyQualifiersFromTo_T<From, To>* CastChecked(const TWeakObjectPtr<From>& Src, ECastCheckedType::Type CheckType)
 {
 	return CastChecked<To>(Src.Get(), CheckType);
 }
@@ -423,17 +423,17 @@ inline TCopyQualifiersFromTo_T<From, To>* CastChecked(const TObjectPtr<From>& Sr
 
 // TSubclassOf versions
 template <typename To, typename From>
-FORCEINLINE TCopyQualifiersFromTo_T<From, To>* Cast(const TSubclassOf<From>& Src)
+UE_FORCEINLINE_HINT TCopyQualifiersFromTo_T<From, To>* Cast(const TSubclassOf<From>& Src)
 {
 	return Cast<To>(*Src);
 }
 template <typename To, typename From>
-FORCEINLINE TCopyQualifiersFromTo_T<From, To>* CastChecked(const TSubclassOf<From>& Src)
+UE_FORCEINLINE_HINT TCopyQualifiersFromTo_T<From, To>* CastChecked(const TSubclassOf<From>& Src)
 {
 	return CastChecked<To>(*Src);
 }
 template <typename To, typename From>
-FORCEINLINE TCopyQualifiersFromTo_T<From, To>* CastChecked(const TSubclassOf<From>& Src, ECastCheckedType::Type CheckType)
+UE_FORCEINLINE_HINT TCopyQualifiersFromTo_T<From, To>* CastChecked(const TSubclassOf<From>& Src, ECastCheckedType::Type CheckType)
 {
 	return CastChecked<To>(*Src, CheckType);
 }
@@ -441,22 +441,22 @@ FORCEINLINE TCopyQualifiersFromTo_T<From, To>* CastChecked(const TSubclassOf<Fro
 // TNotNull versions of the casts
 #if UE_ENABLE_NOTNULL_WRAPPER
 	template <typename T, typename U>
-	FORCEINLINE auto Cast(TNotNull<U> Ptr) -> decltype(Cast<T>((U)Ptr))
+	UE_FORCEINLINE_HINT auto Cast(TNotNull<U> Ptr) -> decltype(Cast<T>((U)Ptr))
 	{
 		return Cast<T>((U)Ptr);
 	}
 	template <typename T, typename U>
-	FORCEINLINE auto ExactCast(TNotNull<U> Ptr) -> decltype(ExactCast<T>((U)Ptr))
+	UE_FORCEINLINE_HINT auto ExactCast(TNotNull<U> Ptr) -> decltype(ExactCast<T>((U)Ptr))
 	{
 		return ExactCast<T>((U)Ptr);
 	}
 	template <typename T, typename U>
-	FORCEINLINE auto CastChecked(TNotNull<U> Ptr) -> decltype(CastChecked<T>((U)Ptr))
+	UE_FORCEINLINE_HINT auto CastChecked(TNotNull<U> Ptr) -> decltype(CastChecked<T>((U)Ptr))
 	{
 		return CastChecked<T>((U)Ptr);
 	}
 	template <typename T, typename U>
-	FORCEINLINE auto CastChecked(TNotNull<U> Ptr, ECastCheckedType::Type CheckType) -> decltype(CastChecked<T>((U)Ptr, CheckType))
+	UE_FORCEINLINE_HINT auto CastChecked(TNotNull<U> Ptr, ECastCheckedType::Type CheckType) -> decltype(CastChecked<T>((U)Ptr, CheckType))
 	{
 		return CastChecked<T>((U)Ptr, CheckType);
 	}
diff --git a/Engine/Source/Runtime/CoreUObject/Public/Templates/NonNullSubclassOf.h b/Engine/Source/Runtime/CoreUObject/Public/Templates/NonNullSubclassOf.h
--- a/Engine/Source/Runtime/CoreUObject/Public/Templates/NonNullSubclassOf.h
+++ b/Engine/Source/Runtime/CoreUObject/Public/Templates/NonNullSubclassOf.h
@@ -21,7 +21,7 @@ class TNonNullSubclassOf : public TSubclassOf<T>
 
 public:
 	/** Default Constructor, defaults to null */
-	[[nodiscard]] FORCEINLINE TNonNullSubclassOf(EDefaultConstructNonNullSubclassOf)
+	[[nodiscard]] UE_FORCEINLINE_HINT TNonNullSubclassOf(EDefaultConstructNonNullSubclassOf)
 		: Super(nullptr)
 	{
 	}
@@ -41,7 +41,7 @@ public:
 		typename U
 		UE_REQUIRES(std::is_convertible_v<U*, T*>)
 	>
-	[[nodiscard]] FORCEINLINE TNonNullSubclassOf(const TSubclassOf<U>& From)
+	[[nodiscard]] UE_FORCEINLINE_HINT TNonNullSubclassOf(const TSubclassOf<U>& From)
 		: Super(From)
 	{
 	}
@@ -73,12 +73,12 @@ public:
 	///////////////////////////////////////////////////////////
 	constexpr static bool bHasIntrusiveUnsetOptionalState = true;
 	using IntrusiveUnsetOptionalStateType = TNonNullSubclassOf;
-	[[nodiscard]] FORCEINLINE explicit TNonNullSubclassOf(FIntrusiveUnsetOptionalState)
+	[[nodiscard]] UE_FORCEINLINE_HINT explicit TNonNullSubclassOf(FIntrusiveUnsetOptionalState)
 		: Super(nullptr)
 	{
 	}
 	
-	[[nodiscard]] FORCEINLINE bool operator==(FIntrusiveUnsetOptionalState) const
+	[[nodiscard]] UE_FORCEINLINE_HINT bool operator==(FIntrusiveUnsetOptionalState) const
 	{
 		return Super::Get() == nullptr;
 	}
diff --git a/Engine/Source/Runtime/CoreUObject/Public/Templates/SubclassOf.h b/Engine/Source/Runtime/CoreUObject/Public/Templates/SubclassOf.h
--- a/Engine/Source/Runtime/CoreUObject/Public/Templates/SubclassOf.h
+++ b/Engine/Source/Runtime/CoreUObject/Public/Templates/SubclassOf.h
@@ -43,7 +43,7 @@ public:
 	~TSubclassOf() = default;
 
 	/** Constructor that takes a UClass*. */
-	[[nodiscard]] FORCEINLINE TSubclassOf(UClass* From)
+	[[nodiscard]] UE_FORCEINLINE_HINT TSubclassOf(UClass* From)
 		: Class(From)
 	{
 	}
@@ -56,7 +56,7 @@ public:
 			std::is_convertible_v<U, UClass*>
 		)
 	>
-	[[nodiscard]] FORCEINLINE TSubclassOf(U&& From)
+	[[nodiscard]] UE_FORCEINLINE_HINT TSubclassOf(U&& From)
 		: Class(From)
 	{
 	}
@@ -116,19 +116,19 @@ public:
 	}
 
 	/** Dereference back into a UClass*, does runtime type checking. */
-	[[nodiscard]] FORCEINLINE UClass* Get() const
+	[[nodiscard]] UE_FORCEINLINE_HINT UClass* Get() const
 	{
 		return **this;
 	}
 
 	/** Dereference back into a UClass*, does runtime type checking. */
-	[[nodiscard]] FORCEINLINE UClass* operator->() const
+	[[nodiscard]] UE_FORCEINLINE_HINT UClass* operator->() const
 	{
 		return **this;
 	}
 
 	/** Implicit conversion to UClass*, does runtime type checking. */
-	[[nodiscard]] FORCEINLINE operator UClass*() const
+	[[nodiscard]] UE_FORCEINLINE_HINT operator UClass*() const
 	{
 		return **this;
 	}
@@ -149,12 +149,12 @@ public:
 		return (T*)Result;
 	}
 
-	FORCEINLINE void Serialize(FArchive& Ar)
+	UE_FORCEINLINE_HINT void Serialize(FArchive& Ar)
 	{
 		Ar << Class;
 	}
 
-	FORCEINLINE void Serialize(FStructuredArchiveSlot& Slot)
+	UE_FORCEINLINE_HINT void Serialize(FStructuredArchiveSlot& Slot)
 	{
 		Slot << Class;
 	}
diff --git a/Engine/Source/Runtime/CoreUObject/Public/Templates/SubclassOfField.h b/Engine/Source/Runtime/CoreUObject/Public/Templates/SubclassOfField.h
--- a/Engine/Source/Runtime/CoreUObject/Public/Templates/SubclassOfField.h
+++ b/Engine/Source/Runtime/CoreUObject/Public/Templates/SubclassOfField.h
@@ -40,7 +40,7 @@ public:
 	~TSubclassOfField() = default;
 
 	/** Constructor that takes a FFieldClass*. */
-	[[nodiscard]] FORCEINLINE TSubclassOfField(FFieldClass* From)
+	[[nodiscard]] UE_FORCEINLINE_HINT TSubclassOfField(FFieldClass* From)
 		: Class(From)
 	{
 	}
@@ -53,7 +53,7 @@ public:
 			std::is_convertible_v<U, FFieldClass*>
 		)
 	>
-	[[nodiscard]] FORCEINLINE TSubclassOfField(U&& From)
+	[[nodiscard]] UE_FORCEINLINE_HINT TSubclassOfField(U&& From)
 		: Class(From)
 	{
 	}
@@ -63,7 +63,7 @@ public:
 		typename OtherT
 		UE_REQUIRES(std::is_convertible_v<OtherT*, T*>)
 	>
-	[[nodiscard]] FORCEINLINE TSubclassOfField(const TSubclassOfField<OtherT>& Other)
+	[[nodiscard]] UE_FORCEINLINE_HINT TSubclassOfField(const TSubclassOfField<OtherT>& Other)
 		: Class(Other.Class)
 	{
 	}
@@ -111,19 +111,19 @@ public:
 	}
 
 	/** Dereference back into a FFieldClass*, does runtime type checking. */
-	[[nodiscard]] FORCEINLINE FFieldClass* Get() const
+	[[nodiscard]] UE_FORCEINLINE_HINT FFieldClass* Get() const
 	{
 		return **this;
 	}
 
 	/** Dereference back into a FFieldClass*, does runtime type checking. */
-	[[nodiscard]] FORCEINLINE FFieldClass* operator->() const
+	[[nodiscard]] UE_FORCEINLINE_HINT FFieldClass* operator->() const
 	{
 		return **this;
 	}
 
 	/** Implicit conversion to FFieldClass*, does runtime type checking. */
-	[[nodiscard]] FORCEINLINE operator FFieldClass*() const
+	[[nodiscard]] UE_FORCEINLINE_HINT operator FFieldClass*() const
 	{
 		return **this;
 	}
@@ -144,7 +144,7 @@ public:
 		return (T*)Result;
 	}
 
-	FORCEINLINE void Serialize(FArchive& Ar)
+	UE_FORCEINLINE_HINT void Serialize(FArchive& Ar)
 	{
 		Ar << Class;
 	}
diff --git a/Engine/Source/Runtime/CoreUObject/Public/UObject/Class.h b/Engine/Source/Runtime/CoreUObject/Public/UObject/Class.h
--- a/Engine/Source/Runtime/CoreUObject/Public/UObject/Class.h
+++ b/Engine/Source/Runtime/CoreUObject/Public/UObject/Class.h
@@ -659,19 +659,19 @@ public:
 	virtual const TCHAR* GetPrefixCPP() const { return TEXT("F"); }
 
 	/** Total size of all UProperties, the allocated structure may be larger due to alignment */
-	FORCEINLINE int32 GetPropertiesSize() const
+	UE_FORCEINLINE_HINT int32 GetPropertiesSize() const
 	{
 		return PropertiesSize;
 	}
 
 	/** Alignment of structure in memory, structure will be at least this large */
-	FORCEINLINE int32 GetMinAlignment() const
+	UE_FORCEINLINE_HINT int32 GetMinAlignment() const
 	{
 		return MinAlignment;
 	}
 
 	/** Returns actual allocated size of structure in memory */
-	FORCEINLINE int32 GetStructureSize() const
+	UE_FORCEINLINE_HINT int32 GetStructureSize() const
 	{
 		return Align(PropertiesSize,MinAlignment);
 	}
@@ -792,7 +792,7 @@ public:
 	UE_INTERNAL COREUOBJECT_API virtual void TrackDefaultInitializedProperties(void* DefaultData) const;
 
 	/* Returns true if this struct has Asset Registry searchable properties */
-	FORCEINLINE bool HasAssetRegistrySearchableProperties() const
+	UE_FORCEINLINE_HINT bool HasAssetRegistrySearchableProperties() const
 	{
 		return bHasAssetRegistrySearchableProperties;
 	}
@@ -1097,12 +1097,12 @@ public:
 		/** Call the C++ destructor **/
 		virtual void Destruct(void *Dest) = 0;
 		/** return the sizeof() of this structure **/
-		FORCEINLINE int32 GetSize() const
+		UE_FORCEINLINE_HINT int32 GetSize() const
 		{
 			return Size;
 		}
 		/** return the alignof() of this structure **/
-		FORCEINLINE int32 GetAlignment() const
+		UE_FORCEINLINE_HINT int32 GetAlignment() const
 		{
 			return Alignment;
 		}
@@ -1904,7 +1904,7 @@ public:
 	 * If it is native, it is assumed to have defaults because it has a constructor
 	 * @return true if this struct has defaults
 	 */
-	FORCEINLINE bool HasDefaults() const
+	UE_FORCEINLINE_HINT bool HasDefaults() const
 	{
 		return !!GetCppStructOps();
 	}
@@ -2117,7 +2117,7 @@ public:
 	 *
 	 * @return The native function pointer.
 	 */
-	FORCEINLINE FNativeFuncPtr GetNativeFunc() const
+	UE_FORCEINLINE_HINT FNativeFuncPtr GetNativeFunc() const
 	{
 		return Func;
 	}
@@ -2127,7 +2127,7 @@ public:
 	 *
 	 * @param InFunc - The new function pointer.
 	 */
-	FORCEINLINE void SetNativeFunc(FNativeFuncPtr InFunc)
+	UE_FORCEINLINE_HINT void SetNativeFunc(FNativeFuncPtr InFunc)
 	{
 		Func = InFunc;
 	}
@@ -2170,7 +2170,7 @@ public:
 	COREUOBJECT_API FProperty* GetReturnProperty() const;
 
 	/** Returns the owning UClass* without branching */
-	FORCEINLINE UClass* GetOuterUClassUnchecked() const
+	UE_FORCEINLINE_HINT UClass* GetOuterUClassUnchecked() const
 	{
 		// declaration order mandates reinterpret_cast:
 		return reinterpret_cast<UClass*>(GetOuter());
@@ -2184,7 +2184,7 @@ public:
 	 * @return	true if the passed in flag is set, false otherwise
 	 *			(including no flag passed in, unless the FlagsToCheck is CLASS_AllFlags)
 	 */
-	FORCEINLINE bool HasAnyFunctionFlags( EFunctionFlags FlagsToCheck ) const
+	UE_FORCEINLINE_HINT bool HasAnyFunctionFlags( EFunctionFlags FlagsToCheck ) const
 	{
 		return (FunctionFlags&FlagsToCheck) != 0 || FlagsToCheck == FUNC_AllFlags;
 	}
@@ -2195,7 +2195,7 @@ public:
 	 * @param FlagsToCheck	Function flags to check for
 	 * @return true if all of the passed in flags are set (including no flags passed in), false otherwise
 	 */
-	FORCEINLINE bool HasAllFunctionFlags( EFunctionFlags FlagsToCheck ) const
+	UE_FORCEINLINE_HINT bool HasAllFunctionFlags( EFunctionFlags FlagsToCheck ) const
 	{
 		return ((FunctionFlags & FlagsToCheck) == FlagsToCheck);
 	}
@@ -2622,13 +2622,13 @@ public:
 	}
 
 	template <typename T>
-	FORCEINLINE static FString GetValueAsString( const TCHAR* EnumPath, const TEnumAsByte<T> EnumeratorValue)
+	UE_FORCEINLINE_HINT static FString GetValueAsString( const TCHAR* EnumPath, const TEnumAsByte<T> EnumeratorValue)
 	{
 		return GetValueAsString_Internal(EnumPath, (int64)EnumeratorValue.GetValue());
 	}
 
 	template< class T >
-	FORCEINLINE static void GetValueAsString( const TCHAR* EnumPath, const T EnumeratorValue, FString& out_StringValue )
+	UE_FORCEINLINE_HINT static void GetValueAsString( const TCHAR* EnumPath, const T EnumeratorValue, FString& out_StringValue )
 	{
 		out_StringValue = GetValueAsString( EnumPath, EnumeratorValue );
 	}
@@ -2642,13 +2642,13 @@ public:
 	}
 
 	template <typename T>
-	FORCEINLINE static FString GetValueOrBitfieldAsString(const TCHAR* EnumPath, const TEnumAsByte<T> EnumeratorValue)
+	UE_FORCEINLINE_HINT static FString GetValueOrBitfieldAsString(const TCHAR* EnumPath, const TEnumAsByte<T> EnumeratorValue)
 	{
 		return GetValueOrBitfieldAsString_Internal(EnumPath, (int64)EnumeratorValue.GetValue());
 	}
 
 	template< class T >
-	FORCEINLINE static void GetValueOrBitfieldAsString(const TCHAR* EnumPath, const T EnumeratorValue, FString& out_StringValue)
+	UE_FORCEINLINE_HINT static void GetValueOrBitfieldAsString(const TCHAR* EnumPath, const T EnumeratorValue, FString& out_StringValue)
 	{
 		out_StringValue = GetValueOrBitfieldAsString( EnumPath, EnumeratorValue );
 	}
@@ -2667,13 +2667,13 @@ public:
 	}
 
 	template <typename T>
-	FORCEINLINE static FText GetDisplayValueAsText( const TCHAR* EnumPath, const TEnumAsByte<T> EnumeratorValue)
+	UE_FORCEINLINE_HINT static FText GetDisplayValueAsText( const TCHAR* EnumPath, const TEnumAsByte<T> EnumeratorValue)
 	{
 		return GetDisplayValueAsText_Internal(EnumPath, (int64)EnumeratorValue.GetValue());
 	}
 
 	template< class T >
-	FORCEINLINE static void GetDisplayValueAsText( const TCHAR* EnumPath, const T EnumeratorValue, FText& out_TextValue )
+	UE_FORCEINLINE_HINT static void GetDisplayValueAsText( const TCHAR* EnumPath, const T EnumeratorValue, FText& out_TextValue )
 	{
 		out_TextValue = GetDisplayValueAsText( EnumPath, EnumeratorValue);
 	}
@@ -2694,13 +2694,13 @@ public:
 	}
 
 	template<typename EnumType>
-	FORCEINLINE static FName GetValueAsName(const TEnumAsByte<EnumType> EnumeratorValue)
+	UE_FORCEINLINE_HINT static FName GetValueAsName(const TEnumAsByte<EnumType> EnumeratorValue)
 	{
 		return GetValueAsName(EnumeratorValue.GetValue());
 	}
 
 	template<typename EnumType>
-	FORCEINLINE static void GetValueAsName(const EnumType EnumeratorValue, FName& out_NameValue )
+	UE_FORCEINLINE_HINT static void GetValueAsName(const EnumType EnumeratorValue, FName& out_NameValue )
 	{
 		out_NameValue = GetValueAsName(EnumeratorValue);
 	}
@@ -2719,13 +2719,13 @@ public:
 	}
 
 	template<typename EnumType>
-	FORCEINLINE static FString GetValueAsString(const TEnumAsByte<EnumType> EnumeratorValue)
+	UE_FORCEINLINE_HINT static FString GetValueAsString(const TEnumAsByte<EnumType> EnumeratorValue)
 	{
 		return GetValueAsString(EnumeratorValue.GetValue());
 	}
 
 	template<typename EnumType>
-	FORCEINLINE static void GetValueAsString(const EnumType EnumeratorValue, FString& out_StringValue )
+	UE_FORCEINLINE_HINT static void GetValueAsString(const EnumType EnumeratorValue, FString& out_StringValue )
 	{
 		out_StringValue = GetValueAsString(EnumeratorValue );
 	}
@@ -2754,7 +2754,7 @@ public:
 		typename EnumType
 		UE_REQUIRES(!TIsEnumClass<EnumType>::Value)
 	>
-	FORCEINLINE static FString GetValueOrBitfieldAsString(const TEnumAsByte<EnumType> EnumeratorValue)
+	UE_FORCEINLINE_HINT static FString GetValueOrBitfieldAsString(const TEnumAsByte<EnumType> EnumeratorValue)
 	{
 		return GetValueOrBitfieldAsString(EnumeratorValue.GetValue());
 	}
@@ -2774,7 +2774,7 @@ public:
 	}
 
 	template<typename EnumType>
-	FORCEINLINE static void GetValueOrBitfieldAsString(const EnumType EnumeratorValue, FString& out_StringValue)
+	UE_FORCEINLINE_HINT static void GetValueOrBitfieldAsString(const EnumType EnumeratorValue, FString& out_StringValue)
 	{
 		out_StringValue = GetValueOrBitfieldAsString(EnumeratorValue);
 	}
@@ -2784,7 +2784,7 @@ public:
 		typename IntegralType
 		UE_REQUIRES(TIsEnum<EnumType>::Value && std::is_integral_v<IntegralType>)
 	>
-	FORCEINLINE static void GetValueOrBitfieldAsString(const IntegralType EnumeratorValue, FString& out_StringValue)
+	UE_FORCEINLINE_HINT static void GetValueOrBitfieldAsString(const IntegralType EnumeratorValue, FString& out_StringValue)
 	{
 		out_StringValue = GetValueOrBitfieldAsString<EnumType>(EnumeratorValue);
 	}
@@ -2805,13 +2805,13 @@ public:
 	}
 
 	template<typename EnumType>
-	FORCEINLINE static FText GetDisplayValueAsText(const TEnumAsByte<EnumType> EnumeratorValue)
+	UE_FORCEINLINE_HINT static FText GetDisplayValueAsText(const TEnumAsByte<EnumType> EnumeratorValue)
 	{
 		return GetDisplayValueAsText(EnumeratorValue.GetValue());
 	}
 
 	template<typename EnumType>
-	FORCEINLINE static void GetDisplayValueAsText(const EnumType EnumeratorValue, FText& out_TextValue )
+	UE_FORCEINLINE_HINT static void GetDisplayValueAsText(const EnumType EnumeratorValue, FText& out_TextValue )
 	{
 		out_TextValue = GetDisplayValueAsText(EnumeratorValue);
 	}
@@ -3409,7 +3409,7 @@ public:
 	/**
 	 * Returns class path name as a package + class FName pair
 	 */
-	FORCEINLINE FTopLevelAssetPath GetClassPathName() const
+	UE_FORCEINLINE_HINT FTopLevelAssetPath GetClassPathName() const
 	{
 		return GetStructPathName();
 	}
@@ -3709,7 +3709,7 @@ public:
 	 * @return	true if the passed in flag is set, false otherwise
 	 *			(including no flag passed in, unless the FlagsToCheck is CLASS_AllFlags)
 	 */
-	FORCEINLINE bool HasAnyClassFlags( EClassFlags FlagsToCheck ) const
+	UE_FORCEINLINE_HINT bool HasAnyClassFlags( EClassFlags FlagsToCheck ) const
 	{
 		return EnumHasAnyFlags(ClassFlags, FlagsToCheck) != 0;
 	}
@@ -3720,7 +3720,7 @@ public:
 	 * @param FlagsToCheck	Class flags to check for
 	 * @return true if all of the passed in flags are set (including no flags passed in), false otherwise
 	 */
-	FORCEINLINE bool HasAllClassFlags( EClassFlags FlagsToCheck ) const
+	UE_FORCEINLINE_HINT bool HasAllClassFlags( EClassFlags FlagsToCheck ) const
 	{
 		return EnumHasAllFlags(ClassFlags, FlagsToCheck);
 	}
@@ -3730,7 +3730,7 @@ public:
 	 *
 	 * @return	The class flags.
 	 */
-	FORCEINLINE EClassFlags GetClassFlags() const
+	UE_FORCEINLINE_HINT EClassFlags GetClassFlags() const
 	{
 		return ClassFlags;
 	}
@@ -3743,11 +3743,11 @@ public:
 	 * @return	true if the passed in flag is set, false otherwise
 	 *			(including no flag passed in)
 	 */
-	FORCEINLINE bool HasAnyCastFlag(EClassCastFlags FlagToCheck) const
+	UE_FORCEINLINE_HINT bool HasAnyCastFlag(EClassCastFlags FlagToCheck) const
 	{
 		return (ClassCastFlags&FlagToCheck) != 0;
 	}
-	FORCEINLINE bool HasAllCastFlags(EClassCastFlags FlagsToCheck) const
+	UE_FORCEINLINE_HINT bool HasAllCastFlags(EClassCastFlags FlagsToCheck) const
 	{
 		return (ClassCastFlags&FlagsToCheck) == FlagsToCheck;
 	}
diff --git a/Engine/Source/Runtime/CoreUObject/Public/UObject/ClassTree.h b/Engine/Source/Runtime/CoreUObject/Public/UObject/ClassTree.h
--- a/Engine/Source/Runtime/CoreUObject/Public/UObject/ClassTree.h
+++ b/Engine/Source/Runtime/CoreUObject/Public/UObject/ClassTree.h
@@ -597,12 +597,12 @@ public:
 	*
 	* @return	the number of child nodes of this node
 	*/
-	FORCEINLINE int32 NumChildren() const
+	UE_FORCEINLINE_HINT int32 NumChildren() const
 	{
 		return Children.Num();
 	}
 
-	FORCEINLINE const FClassTree* GetChild(int32 Index) const
+	UE_FORCEINLINE_HINT const FClassTree* GetChild(int32 Index) const
 	{
 		return Children[Index];
 	}
diff --git a/Engine/Source/Runtime/CoreUObject/Public/UObject/EnumProperty.h b/Engine/Source/Runtime/CoreUObject/Public/UObject/EnumProperty.h
--- a/Engine/Source/Runtime/CoreUObject/Public/UObject/EnumProperty.h
+++ b/Engine/Source/Runtime/CoreUObject/Public/UObject/EnumProperty.h
@@ -93,7 +93,7 @@ public:
 	 * Set the UEnum of this property.
 	 * @note Internal! Used for creation of enums that use impersonation.
 	 */
-	UE_INTERNAL FORCEINLINE void SetEnumForImpersonation(UEnum* InEnum)
+	UE_INTERNAL UE_FORCEINLINE_HINT void SetEnumForImpersonation(UEnum* InEnum)
 	{
 		Enum = InEnum;
 	}
@@ -101,7 +101,7 @@ public:
 	/**
 	 * Returns a pointer to the UEnum of this property.
 	 */
-	FORCEINLINE UEnum* GetEnum() const
+	UE_FORCEINLINE_HINT UEnum* GetEnum() const
 	{
 		return Enum;
 	}
@@ -109,7 +109,7 @@ public:
 	/**
 	 * Returns the numeric property which represents the integral type of the enum.
 	 */
-	FORCEINLINE FNumericProperty* GetUnderlyingProperty() const
+	UE_FORCEINLINE_HINT FNumericProperty* GetUnderlyingProperty() const
 	{
 		return UnderlyingProp;
 	}
diff --git a/Engine/Source/Runtime/CoreUObject/Public/UObject/Field.h b/Engine/Source/Runtime/CoreUObject/Public/UObject/Field.h
--- a/Engine/Source/Runtime/CoreUObject/Public/UObject/Field.h
+++ b/Engine/Source/Runtime/CoreUObject/Public/UObject/Field.h
@@ -403,12 +403,12 @@ public:
 		}
 	}
 	/** FOR INTERNAL USE ONLY: Function that returns the owner as FField without checking if it's actually an FField */
-	FORCEINLINE FField* ToFieldUnsafe() const
+	UE_FORCEINLINE_HINT FField* ToFieldUnsafe() const
 	{
 		return Container.Field;
 	}
 	/** FOR INTERNAL USE ONLY: Function that returns the owner as UObject without checking if it's actually a UObject */
-	FORCEINLINE UObject* ToUObjectUnsafe() const
+	UE_FORCEINLINE_HINT UObject* ToUObjectUnsafe() const
 	{
 		return (UObject*)((uintptr_t)Container.Object & ~UObjectMask);
 	}
@@ -646,7 +646,7 @@ public:
 	}
 
 	/** Internal function for quickly getting the owner of this object as UObject. FOR INTERNAL USE ONLY */
-	FORCEINLINE UObject* InternalGetOwnerAsUObjectUnsafe() const
+	UE_FORCEINLINE_HINT UObject* InternalGetOwnerAsUObjectUnsafe() const
 	{
 		return Owner.ToUObjectUnsafe();
 	}
@@ -966,25 +966,25 @@ public:
 // Support for casting between different FFIeld types
 
 template<typename FieldType>
-FORCEINLINE FieldType* CastField(FField* Src)
+UE_FORCEINLINE_HINT FieldType* CastField(FField* Src)
 {
 	return Src && Src->IsA<FieldType>() ? static_cast<FieldType*>(Src) : nullptr;
 }
 
 template<typename FieldType>
-FORCEINLINE const FieldType* CastField(const FField* Src)
+UE_FORCEINLINE_HINT const FieldType* CastField(const FField* Src)
 {
 	return Src && Src->IsA<FieldType>() ? static_cast<const FieldType*>(Src) : nullptr;
 }
 
 template<typename FieldType>
-FORCEINLINE FieldType* ExactCastField(FField* Src)
+UE_FORCEINLINE_HINT FieldType* ExactCastField(FField* Src)
 {
 	return (Src && (Src->GetClass() == FieldType::StaticClass())) ? static_cast<FieldType*>(Src) : nullptr;
 }
 
 template<typename FieldType>
-FORCEINLINE FieldType* ExactCastField(const FField* Src)
+UE_FORCEINLINE_HINT FieldType* ExactCastField(const FField* Src)
 {
 	return (Src && (Src->GetClass() == FieldType::StaticClass())) ? static_cast<const FieldType*>(Src) : nullptr;
 }
diff --git a/Engine/Source/Runtime/CoreUObject/Public/UObject/FieldIterator.h b/Engine/Source/Runtime/CoreUObject/Public/UObject/FieldIterator.h
--- a/Engine/Source/Runtime/CoreUObject/Public/UObject/FieldIterator.h
+++ b/Engine/Source/Runtime/CoreUObject/Public/UObject/FieldIterator.h
@@ -64,12 +64,12 @@ public:
 	}
 
 	/** conversion to "bool" returning true if the iterator is valid. */
-	FORCEINLINE explicit operator bool() const
+	UE_FORCEINLINE_HINT explicit operator bool() const
 	{
 		return (bool)FieldIterator || (bool)StructIterator;
 	}
 	/** inverse of the "bool" operator */
-	FORCEINLINE bool operator !() const
+	UE_FORCEINLINE_HINT bool operator !() const
 	{
 		return !(bool)*this;
 	}
diff --git a/Engine/Source/Runtime/CoreUObject/Public/UObject/FieldPath.h b/Engine/Source/Runtime/CoreUObject/Public/UObject/FieldPath.h
--- a/Engine/Source/Runtime/CoreUObject/Public/UObject/FieldPath.h
+++ b/Engine/Source/Runtime/CoreUObject/Public/UObject/FieldPath.h
@@ -106,7 +106,7 @@ private:
 #endif
 
 	/** FOR INTERNAL USE ONLY: gets the pointer to the resolved field without trying to resolve it */
-	FORCEINLINE FUObjectItem* GetResolvedOwnerItemInternal()
+	UE_FORCEINLINE_HINT FUObjectItem* GetResolvedOwnerItemInternal()
 	{
 		return ResolvedOwner.Internal_GetObjectItem();
 	}
@@ -252,12 +252,12 @@ public:
 		Path.Empty();
 	}
 
-	FORCEINLINE bool operator==(const FFieldPath& InOther) const
+	UE_FORCEINLINE_HINT bool operator==(const FFieldPath& InOther) const
 	{
 		return ResolvedOwner == InOther.ResolvedOwner && Path == InOther.Path;
 	}
 
-	FORCEINLINE bool operator!=(const FFieldPath& InOther) const
+	UE_FORCEINLINE_HINT bool operator!=(const FFieldPath& InOther) const
 	{
 		return ResolvedOwner != InOther.ResolvedOwner || Path != InOther.Path;
 	}
@@ -308,7 +308,7 @@ public:
 	/**
 	* Construct from a null pointer
 	**/
-	FORCEINLINE TFieldPath(TYPE_OF_NULLPTR)
+	UE_FORCEINLINE_HINT TFieldPath(TYPE_OF_NULLPTR)
 		: FFieldPath()
 	{
 	}
@@ -402,7 +402,7 @@ public:
 	 * @param InCurrentStruct Struct that's trying to resolve this field path
 	 * @return Field represented by this FFieldPath or null if it couldn't be resolved
 	 */
-	FORCEINLINE PropertyType* Get(UStruct* InCurrentStruct = nullptr) const
+	UE_FORCEINLINE_HINT PropertyType* Get(UStruct* InCurrentStruct = nullptr) const
 	{
 		return (PropertyType*)GetTyped(PropertyType::StaticClass(), InCurrentStruct);
 	}
@@ -417,7 +417,7 @@ public:
 	/**
 	* Dereference the weak pointer
 	**/
-	FORCEINLINE PropertyType* operator*() const
+	UE_FORCEINLINE_HINT PropertyType* operator*() const
 	{
 		return Get();
 	}
@@ -425,7 +425,7 @@ public:
 	/**
 	* Dereference the weak pointer
 	**/
-	FORCEINLINE PropertyType* operator->() const
+	UE_FORCEINLINE_HINT PropertyType* operator->() const
 	{
 		return Get();
 	}
@@ -437,7 +437,7 @@ public:
 * @param Rhs weak pointer to compare
 **/
 template <typename LhsType, typename RhsType>
-FORCEINLINE auto operator==(const TFieldPath<LhsType>& Lhs, const TFieldPath<LhsType>& Rhs)
+UE_FORCEINLINE_HINT auto operator==(const TFieldPath<LhsType>& Lhs, const TFieldPath<LhsType>& Rhs)
 	-> decltype((LhsType*)nullptr == (RhsType*)nullptr)
 {
 	return *(const FFieldPath*)&Lhs == *(const FFieldPath*)&Rhs;
@@ -449,7 +449,7 @@ FORCEINLINE auto operator==(const TFieldPath<LhsType>& Lhs, const TFieldPath<Lhs
 * @param Rhs pointer to compare
 **/
 template <typename LhsType, typename RhsType>
-FORCEINLINE auto operator==(const TFieldPath<LhsType>& Lhs, const RhsType* Rhs)
+UE_FORCEINLINE_HINT auto operator==(const TFieldPath<LhsType>& Lhs, const RhsType* Rhs)
 	-> decltype((LhsType*)nullptr == Rhs)
 {
 	return Lhs.Get() == Rhs;
@@ -461,7 +461,7 @@ FORCEINLINE auto operator==(const TFieldPath<LhsType>& Lhs, const RhsType* Rhs)
 * @param Rhs weak pointer to compare
 **/
 template <typename LhsType, typename RhsType>
-FORCEINLINE auto operator==(const LhsType* Lhs, const TFieldPath<RhsType>& Rhs)
+UE_FORCEINLINE_HINT auto operator==(const LhsType* Lhs, const TFieldPath<RhsType>& Rhs)
 	-> decltype(Lhs == (RhsType*)nullptr)
 {
 	return Lhs == Rhs.Get();
@@ -472,7 +472,7 @@ FORCEINLINE auto operator==(const LhsType* Lhs, const TFieldPath<RhsType>& Rhs)
 * @param Lhs pointer to test
 **/
 template <typename LhsType>
-FORCEINLINE bool operator==(TFieldPath<LhsType>& Lhs, TYPE_OF_NULLPTR)
+UE_FORCEINLINE_HINT bool operator==(TFieldPath<LhsType>& Lhs, TYPE_OF_NULLPTR)
 {
 	return !Lhs.Get();
 }
@@ -482,7 +482,7 @@ FORCEINLINE bool operator==(TFieldPath<LhsType>& Lhs, TYPE_OF_NULLPTR)
 * @param Rhs pointer to test
 **/
 template <typename RhsType>
-FORCEINLINE bool operator==(TYPE_OF_NULLPTR, TFieldPath<RhsType>& Rhs)
+UE_FORCEINLINE_HINT bool operator==(TYPE_OF_NULLPTR, TFieldPath<RhsType>& Rhs)
 {
 	return !Rhs.Get();
 }
@@ -494,7 +494,7 @@ FORCEINLINE bool operator==(TYPE_OF_NULLPTR, TFieldPath<RhsType>& Rhs)
 * @param Rhs weak pointer to compare
 **/
 template <typename LhsType, typename RhsType>
-FORCEINLINE auto operator!=(const TFieldPath<LhsType>& Lhs, const TFieldPath<LhsType>& Rhs)
+UE_FORCEINLINE_HINT auto operator!=(const TFieldPath<LhsType>& Lhs, const TFieldPath<LhsType>& Rhs)
 	-> decltype((LhsType*)nullptr != (RhsType*)nullptr)
 {
 	return !(Lhs == Rhs);
@@ -506,7 +506,7 @@ FORCEINLINE auto operator!=(const TFieldPath<LhsType>& Lhs, const TFieldPath<Lhs
 * @param Rhs pointer to compare
 **/
 template <typename LhsType, typename RhsType>
-FORCEINLINE auto operator!=(const TFieldPath<LhsType>& Lhs, const RhsType* Rhs)
+UE_FORCEINLINE_HINT auto operator!=(const TFieldPath<LhsType>& Lhs, const RhsType* Rhs)
 	-> decltype((LhsType*)nullptr != Rhs)
 {
 	return !(Lhs == Rhs);
@@ -518,7 +518,7 @@ FORCEINLINE auto operator!=(const TFieldPath<LhsType>& Lhs, const RhsType* Rhs)
 * @param Rhs weak pointer to compare
 **/
 template <typename LhsType, typename RhsType>
-FORCEINLINE auto operator!=(const LhsType* Lhs, const TFieldPath<RhsType>& Rhs)
+UE_FORCEINLINE_HINT auto operator!=(const LhsType* Lhs, const TFieldPath<RhsType>& Rhs)
 	-> decltype(Lhs != (RhsType*)nullptr)
 {
 	return !(Lhs == Rhs);
@@ -529,7 +529,7 @@ FORCEINLINE auto operator!=(const LhsType* Lhs, const TFieldPath<RhsType>& Rhs)
 * @param Lhs pointer to test
 **/
 template <typename LhsType>
-FORCEINLINE bool operator!=(TFieldPath<LhsType>& Lhs, TYPE_OF_NULLPTR)
+UE_FORCEINLINE_HINT bool operator!=(TFieldPath<LhsType>& Lhs, TYPE_OF_NULLPTR)
 {
 	return !(Lhs == nullptr);
 }
@@ -539,7 +539,7 @@ FORCEINLINE bool operator!=(TFieldPath<LhsType>& Lhs, TYPE_OF_NULLPTR)
 * @param Rhs pointer to test
 **/
 template <typename RhsType>
-FORCEINLINE bool operator!=(TYPE_OF_NULLPTR, TFieldPath<RhsType>& Rhs)
+UE_FORCEINLINE_HINT bool operator!=(TYPE_OF_NULLPTR, TFieldPath<RhsType>& Rhs)
 {
 	return !(nullptr == Rhs);
 }
@@ -547,7 +547,7 @@ FORCEINLINE bool operator!=(TYPE_OF_NULLPTR, TFieldPath<RhsType>& Rhs)
 
 // Helper function which deduces the type of the initializer
 template <typename PropertyType>
-FORCEINLINE TFieldPath<PropertyType> MakePropertyPath(PropertyType* Ptr)
+UE_FORCEINLINE_HINT TFieldPath<PropertyType> MakePropertyPath(PropertyType* Ptr)
 {
 	return TFieldPath<PropertyType>(Ptr);
 }
@@ -566,12 +566,12 @@ struct TPropertyPathMapKeyFuncs : public TDefaultMapKeyFuncs<KeyType, ValueType,
 {
 	typedef typename TDefaultMapKeyFuncs<KeyType, ValueType, bInAllowDuplicateKeys>::KeyInitType KeyInitType;
 
-	static FORCEINLINE bool Matches(KeyInitType A, KeyInitType B)
+	static UE_FORCEINLINE_HINT bool Matches(KeyInitType A, KeyInitType B)
 	{
 		return A == B;
 	}
 
-	static FORCEINLINE uint32 GetKeyHash(KeyInitType Key)
+	static UE_FORCEINLINE_HINT uint32 GetKeyHash(KeyInitType Key)
 	{
 		return GetTypeHash(Key);
 	}
diff --git a/Engine/Source/Runtime/CoreUObject/Public/UObject/GCObjectScopeGuard.h b/Engine/Source/Runtime/CoreUObject/Public/UObject/GCObjectScopeGuard.h
--- a/Engine/Source/Runtime/CoreUObject/Public/UObject/GCObjectScopeGuard.h
+++ b/Engine/Source/Runtime/CoreUObject/Public/UObject/GCObjectScopeGuard.h
@@ -38,7 +38,7 @@ public:
 		return TEXT("TGCObjectScopeGuard");
 	}
 
-	FORCEINLINE_DEBUGGABLE ObjectType* Get() const
+	UE_FORCEINLINE_HINT ObjectType* Get() const
 	{
 		return Object;
 	}
@@ -82,7 +82,7 @@ public:
 		return TEXT("TGCObjectsScopeGuard");
 	}
 
-	FORCEINLINE_DEBUGGABLE const TArray<ObjectType*>& Get() const
+	UE_FORCEINLINE_HINT const TArray<ObjectType*>& Get() const
 	{
 		return Objects;
 	}
diff --git a/Engine/Source/Runtime/CoreUObject/Public/UObject/GarbageCollection.h b/Engine/Source/Runtime/CoreUObject/Public/UObject/GarbageCollection.h
--- a/Engine/Source/Runtime/CoreUObject/Public/UObject/GarbageCollection.h
+++ b/Engine/Source/Runtime/CoreUObject/Public/UObject/GarbageCollection.h
@@ -196,7 +196,7 @@ COREUOBJECT_API double GetLastGCDuration();
 /**
 * Whether we are inside garbage collection
 */
-FORCEINLINE bool IsGarbageCollecting()
+UE_FORCEINLINE_HINT bool IsGarbageCollecting()
 {
 	return GIsGarbageCollecting;
 }
diff --git a/Engine/Source/Runtime/CoreUObject/Public/UObject/GarbageCollectionHistory.h b/Engine/Source/Runtime/CoreUObject/Public/UObject/GarbageCollectionHistory.h
--- a/Engine/Source/Runtime/CoreUObject/Public/UObject/GarbageCollectionHistory.h
+++ b/Engine/Source/Runtime/CoreUObject/Public/UObject/GarbageCollectionHistory.h
@@ -73,12 +73,12 @@ public:
 	 */
 	COREUOBJECT_API void SetHistorySize(int32 HistorySize);
 
-	FORCEINLINE int32 GetHistorySize() const
+	UE_FORCEINLINE_HINT int32 GetHistorySize() const
 	{
 		return Snapshots.Num();
 	}
 
-	FORCEINLINE bool IsActive() const
+	UE_FORCEINLINE_HINT bool IsActive() const
 	{
 		return GetHistorySize() > 0;
 	}
@@ -102,7 +102,7 @@ public:
 	COREUOBJECT_API int64 GetAllocatedSize() const;
 
 	/** Returns an object representing GC Barrier */
-	FORCEINLINE UObject* GetBarrierObject() const
+	UE_FORCEINLINE_HINT UObject* GetBarrierObject() const
 	{
 		return GCBarrier;
 	}
diff --git a/Engine/Source/Runtime/CoreUObject/Public/UObject/GarbageCollectionSchema.h b/Engine/Source/Runtime/CoreUObject/Public/UObject/GarbageCollectionSchema.h
--- a/Engine/Source/Runtime/CoreUObject/Public/UObject/GarbageCollectionSchema.h
+++ b/Engine/Source/Runtime/CoreUObject/Public/UObject/GarbageCollectionSchema.h
@@ -155,8 +155,8 @@ enum class EMemberlessId
 class FMemberId
 {
 public:
-	FORCEINLINE /* implicit */ FMemberId(EMemberlessId Memberless) : Index(0), MemberlessId((uint32)Memberless) {}
-	FORCEINLINE explicit FMemberId(uint32 Idx) : Index(Idx), MemberlessId(0) {}
+	UE_FORCEINLINE_HINT /* implicit */ FMemberId(EMemberlessId Memberless) : Index(0), MemberlessId((uint32)Memberless) {}
+	UE_FORCEINLINE_HINT explicit FMemberId(uint32 Idx) : Index(Idx), MemberlessId(0) {}
 	
 	bool IsMemberless() const { return MemberlessId != 0; }
 	uint32 GetIndex() const { check(!IsMemberless()); return Index; }
diff --git a/Engine/Source/Runtime/CoreUObject/Public/UObject/ICookInfo.h b/Engine/Source/Runtime/CoreUObject/Public/UObject/ICookInfo.h
--- a/Engine/Source/Runtime/CoreUObject/Public/UObject/ICookInfo.h
+++ b/Engine/Source/Runtime/CoreUObject/Public/UObject/ICookInfo.h
@@ -352,11 +352,11 @@ private:
 /** Defined in !WITH_EDITOR to remove the need for #ifdef boilerplate, but it is a noop. */
 struct FCookLoadScope
 {
-	FORCEINLINE explicit FCookLoadScope(ECookLoadType ScopeType)
+	UE_FORCEINLINE_HINT explicit FCookLoadScope(ECookLoadType ScopeType)
 	{
 	}
 
-	FORCEINLINE ~FCookLoadScope()
+	UE_FORCEINLINE_HINT ~FCookLoadScope()
 	{
 	}
 };
diff --git a/Engine/Source/Runtime/CoreUObject/Public/UObject/LazyObjectPtr.h b/Engine/Source/Runtime/CoreUObject/Public/UObject/LazyObjectPtr.h
--- a/Engine/Source/Runtime/CoreUObject/Public/UObject/LazyObjectPtr.h
+++ b/Engine/Source/Runtime/CoreUObject/Public/UObject/LazyObjectPtr.h
@@ -40,7 +40,7 @@ struct FUniqueObjectGuid
 	{}
 
 	/** Reset the guid pointer back to the invalid state */
-	FORCEINLINE void Reset()
+	UE_FORCEINLINE_HINT void Reset()
 	{
 		Guid.Invalidate();
 	}
@@ -65,34 +65,34 @@ struct FUniqueObjectGuid
 	COREUOBJECT_API UObject* ResolveObject() const;
 
 	/** Test if this can ever point to a live UObject */
-	FORCEINLINE bool IsValid() const
+	UE_FORCEINLINE_HINT bool IsValid() const
 	{
 		return Guid.IsValid();
 	}
 
-	FORCEINLINE bool operator==(const FUniqueObjectGuid& Other) const
+	UE_FORCEINLINE_HINT bool operator==(const FUniqueObjectGuid& Other) const
 	{
 		return Guid == Other.Guid;
 	}
 
-	FORCEINLINE bool operator!=(const FUniqueObjectGuid& Other) const
+	UE_FORCEINLINE_HINT bool operator!=(const FUniqueObjectGuid& Other) const
 	{
 		return Guid != Other.Guid;
 	}
 
 	/** Returns true is this is the default value */
-	FORCEINLINE bool IsDefault() const
+	UE_FORCEINLINE_HINT bool IsDefault() const
 	{
 		// A default GUID is 0,0,0,0 and this is "invalid"
 		return !IsValid(); 
 	}
 
-	FORCEINLINE friend uint32 GetTypeHash(const FUniqueObjectGuid& ObjectGuid)
+	UE_FORCEINLINE_HINT friend uint32 GetTypeHash(const FUniqueObjectGuid& ObjectGuid)
 	{
 		return GetTypeHash(ObjectGuid.Guid);
 	}
 	/** Returns wrapped Guid */
-	FORCEINLINE const FGuid& GetGuid() const
+	UE_FORCEINLINE_HINT const FGuid& GetGuid() const
 	{
 		return Guid;
 	}
@@ -120,12 +120,12 @@ struct FUniqueObjectGuid
 	}
 
 	static COREUOBJECT_API FUniqueObjectGuid GetOrCreateIDForObject(FObjectPtr Object);
-	static FORCEINLINE FUniqueObjectGuid GetOrCreateIDForObject(const UObject* Object)
+	static UE_FORCEINLINE_HINT FUniqueObjectGuid GetOrCreateIDForObject(const UObject* Object)
 	{
 		return GetOrCreateIDForObject(FObjectPtr(const_cast<UObject*>(Object)));
 	}
 	template <typename T>
-	static FORCEINLINE FUniqueObjectGuid GetOrCreateIDForObject(TObjectPtr<T> Object)
+	static UE_FORCEINLINE_HINT FUniqueObjectGuid GetOrCreateIDForObject(TObjectPtr<T> Object)
 	{
 		return GetOrCreateIDForObject(FObjectPtr(Object));
 	}
@@ -154,21 +154,21 @@ private:
 
 public:
 	/** Default constructor, sets to null */
-	FORCEINLINE FLazyObjectPtr()
+	UE_FORCEINLINE_HINT FLazyObjectPtr()
 	{
 	}
 
 	/** Construct from object already in memory */
-	explicit FORCEINLINE FLazyObjectPtr(FObjectPtr Object)
+	explicit UE_FORCEINLINE_HINT FLazyObjectPtr(FObjectPtr Object)
 	{
 		(*this)=Object;
 	}
-	explicit FORCEINLINE FLazyObjectPtr(const UObject* Object)
+	explicit UE_FORCEINLINE_HINT FLazyObjectPtr(const UObject* Object)
 		: FLazyObjectPtr(FObjectPtr(const_cast<UObject*>(Object)))
 	{
 	}
 	template <typename T>
-	explicit FORCEINLINE FLazyObjectPtr(TObjectPtr<T> Object)
+	explicit UE_FORCEINLINE_HINT FLazyObjectPtr(TObjectPtr<T> Object)
 		: FLazyObjectPtr(FObjectPtr(Object))
 	{
 		// This needs to be a template instead of TObjectPtr<const UObject> because C++ does derived-to-base
@@ -179,16 +179,16 @@ public:
 	}
 
 	/** Copy from an object already in memory */
-	FORCEINLINE void operator=(FObjectPtr Object)
+	UE_FORCEINLINE_HINT void operator=(FObjectPtr Object)
 	{
 		Super::operator=(Object);
 	}
-	FORCEINLINE void operator=(const UObject* Object)
+	UE_FORCEINLINE_HINT void operator=(const UObject* Object)
 	{
 		*this = FObjectPtr(const_cast<UObject*>(Object));
 	}
 	template <typename T>
-	FORCEINLINE void operator=(TObjectPtr<T> Object)
+	UE_FORCEINLINE_HINT void operator=(TObjectPtr<T> Object)
 	{
 		// This needs to be a template instead of TObjectPtr<const UObject> because C++ does derived-to-base
 		// pointer conversions ('standard conversion sequences') in more cases than TSmartPtr<Derived>-to-TSmartPtr<Base>
@@ -200,13 +200,13 @@ public:
 	}
 
 	/** Copy from a unique object identifier */
-	FORCEINLINE void operator=(const FUniqueObjectGuid& InObjectID)
+	UE_FORCEINLINE_HINT void operator=(const FUniqueObjectGuid& InObjectID)
 	{
 		Super::operator=(InObjectID);
 	}
 
 	/** Fixes up this FLazyObjectPtr to target the right UID as set in PIEGuidMap, this only works for directly serialized pointers */
-	FORCEINLINE void FixupForPIE(int32 PIEInstance)
+	UE_FORCEINLINE_HINT void FixupForPIE(int32 PIEInstance)
 	{
 		*this = GetUniqueID().FixupForPIE(PIEInstance);
 	}
@@ -243,7 +243,7 @@ public:
 		typename U
 		UE_REQUIRES(std::is_convertible_v<U*, T*>)
 	>
-	FORCEINLINE TLazyObjectPtr(const TLazyObjectPtr<U>& Other) :
+	UE_FORCEINLINE_HINT TLazyObjectPtr(const TLazyObjectPtr<U>& Other) :
 		FLazyObjectPtr((const FLazyObjectPtr&)Other)
 	{
 	}
@@ -253,26 +253,26 @@ public:
 		typename U
 		UE_REQUIRES(std::is_convertible_v<U*, T*>)
 	>
-	FORCEINLINE TLazyObjectPtr<T>& operator=(const TLazyObjectPtr<U>& Other)
+	UE_FORCEINLINE_HINT TLazyObjectPtr<T>& operator=(const TLazyObjectPtr<U>& Other)
 	{
 		FLazyObjectPtr::operator=((const FLazyObjectPtr&)Other);
 		return *this;
 	}
 
 	/** Construct from an object pointer */
-	FORCEINLINE TLazyObjectPtr(T* Object)
+	UE_FORCEINLINE_HINT TLazyObjectPtr(T* Object)
 	{
 		FLazyObjectPtr::operator=(Object);
 	}
 
 	/** Reset the lazy pointer back to the null state */
-	FORCEINLINE void Reset()
+	UE_FORCEINLINE_HINT void Reset()
 	{
 		FLazyObjectPtr::Reset();
 	}
 
 	/** Copy from an object pointer */
-	FORCEINLINE void operator=(T* Object)
+	UE_FORCEINLINE_HINT void operator=(T* Object)
 	{
 		FLazyObjectPtr::operator=(Object);
 	}
@@ -284,7 +284,7 @@ public:
 	 *
 	 * @param ObjectID Object identifier to create a lazy pointer to
 	 */
-	FORCEINLINE void operator=(const FUniqueObjectGuid& InObjectID)
+	UE_FORCEINLINE_HINT void operator=(const FUniqueObjectGuid& InObjectID)
 	{
 		FLazyObjectPtr::operator=(InObjectID);
 	}
@@ -294,7 +294,7 @@ public:
 	 *
 	 * @return Unique ID for this object, or an invalid FUniqueObjectID if this pointer isn't set to anything
 	 */
-	FORCEINLINE const FUniqueObjectGuid& GetUniqueID() const
+	UE_FORCEINLINE_HINT const FUniqueObjectGuid& GetUniqueID() const
 	{
 		return FLazyObjectPtr::GetUniqueID();
 	}
@@ -304,7 +304,7 @@ public:
 	 *
 	 * @return nullptr if this object is gone or the lazy pointer was null, otherwise a valid UObject pointer
 	 */
-	FORCEINLINE T* Get() const
+	UE_FORCEINLINE_HINT T* Get() const
 	{
 		// there are cases where a TLazyObjectPtr can get an object of the wrong type assigned to it which are difficult to avoid
 		// e.g. operator=(const FUniqueObjectGuid& ObjectID)
@@ -313,19 +313,19 @@ public:
 	}
 
 	/** Dereference the lazy pointer */
-	FORCEINLINE T& operator*() const
+	UE_FORCEINLINE_HINT T& operator*() const
 	{
 		return *Get();
 	}
 
 	/** Dereference the lazy pointer */
-	FORCEINLINE T* operator->() const
+	UE_FORCEINLINE_HINT T* operator->() const
 	{
 		return Get();
 	}
 
 	/** Test if this points to a live UObject */
-	FORCEINLINE bool IsValid() const
+	UE_FORCEINLINE_HINT bool IsValid() const
 	{
 		return FLazyObjectPtr::IsValid();
 	}
@@ -335,7 +335,7 @@ public:
 	 *
 	 * @return true if this used to point at a real object but no longer does.
 	 */
-	FORCEINLINE bool IsStale() const
+	UE_FORCEINLINE_HINT bool IsStale() const
 	{
 		return FLazyObjectPtr::IsStale();
 	}
@@ -345,7 +345,7 @@ public:
 	 *
 	 * @return true if this does not point to a real object, but could possibly
 	 */
-	FORCEINLINE bool IsPending() const
+	UE_FORCEINLINE_HINT bool IsPending() const
 	{
 		return FLazyObjectPtr::IsPending();
 	}
@@ -355,61 +355,61 @@ public:
 	 *
 	 * @return true if this is explicitly pointing to no object
 	 */
-	FORCEINLINE bool IsNull() const
+	UE_FORCEINLINE_HINT bool IsNull() const
 	{
 		return FLazyObjectPtr::IsNull();
 	}
 
 	/** Dereference lazy pointer to see if it points somewhere valid */
-	FORCEINLINE explicit operator bool() const
+	UE_FORCEINLINE_HINT explicit operator bool() const
 	{
 		return IsValid();
 	}
 
 	/** Hash function. */
-	FORCEINLINE uint32 GetLazyObjecPtrTypeHash() const
+	UE_FORCEINLINE_HINT uint32 GetLazyObjecPtrTypeHash() const
 	{
 		return GetTypeHash(static_cast<const FLazyObjectPtr&>(*this));
 	}
 
-	FORCEINLINE void SerializePtr(FArchive& Ar)
+	UE_FORCEINLINE_HINT void SerializePtr(FArchive& Ar)
 	{
 		Ar << static_cast<FLazyObjectPtr&>(*this);
 	}
 
 	/** Compare with another TLazyObjectPtr of related type */
 	template<typename U, typename = decltype((T*)nullptr == (U*)nullptr)>
-	FORCEINLINE bool operator==(const TLazyObjectPtr<U>& Rhs) const
+	UE_FORCEINLINE_HINT bool operator==(const TLazyObjectPtr<U>& Rhs) const
 	{
 		return (const FLazyObjectPtr&)*this == (const FLazyObjectPtr&)Rhs;
 	}
 	template<typename U, typename = decltype((T*)nullptr != (U*)nullptr)>
-	FORCEINLINE bool operator!=(const TLazyObjectPtr<U>& Rhs) const
+	UE_FORCEINLINE_HINT bool operator!=(const TLazyObjectPtr<U>& Rhs) const
 	{
 		return (const FLazyObjectPtr&)*this != (const FLazyObjectPtr&)Rhs;
 	}
 
 	/** Compare for equality with a raw pointer **/
 	template<typename U, typename = decltype((T*)nullptr == (U*)nullptr)>
-	FORCEINLINE bool operator==(const U* Rhs) const
+	UE_FORCEINLINE_HINT bool operator==(const U* Rhs) const
 	{
 		return Get() == Rhs;
 	}
 
 	/** Compare to null */
-	FORCEINLINE bool operator==(TYPE_OF_NULLPTR) const
+	UE_FORCEINLINE_HINT bool operator==(TYPE_OF_NULLPTR) const
 	{
 		return !IsValid();
 	}
 	/** Compare for inequality with a raw pointer	**/
 	template<typename U, typename = decltype((T*)nullptr != (U*)nullptr)>
-	FORCEINLINE bool operator!=(const U* Rhs) const
+	UE_FORCEINLINE_HINT bool operator!=(const U* Rhs) const
 	{
 		return Get() != Rhs;
 	}
 
 	/** Compare for inequality with null **/
-	FORCEINLINE bool operator!=(TYPE_OF_NULLPTR) const
+	UE_FORCEINLINE_HINT bool operator!=(TYPE_OF_NULLPTR) const
 	{
 		return IsValid();
 	}
@@ -417,7 +417,7 @@ public:
 
 /** Hash function. */
 template<typename T>
-FORCEINLINE uint32 GetTypeHash(const TLazyObjectPtr<T>& LazyObjectPtr)
+UE_FORCEINLINE_HINT uint32 GetTypeHash(const TLazyObjectPtr<T>& LazyObjectPtr)
 {
 	return LazyObjectPtr.GetLazyObjecPtrTypeHash();
 }
@@ -431,28 +431,28 @@ FArchive& operator<<(FArchive& Ar, TLazyObjectPtr<T>& LazyObjectPtr)
 
 /** Compare for equality with a raw pointer **/
 template<typename T, typename U, typename = decltype((T*)nullptr == (U*)nullptr)>
-FORCEINLINE bool operator==(const U* Lhs, const TLazyObjectPtr<T>& Rhs)
+UE_FORCEINLINE_HINT bool operator==(const U* Lhs, const TLazyObjectPtr<T>& Rhs)
 {
 	return Lhs == Rhs.Get();
 }
 
 /** Compare to null */
 template<typename T>
-FORCEINLINE bool operator==(TYPE_OF_NULLPTR, const TLazyObjectPtr<T>& Rhs)
+UE_FORCEINLINE_HINT bool operator==(TYPE_OF_NULLPTR, const TLazyObjectPtr<T>& Rhs)
 {
 	return !Rhs.IsValid();
 }
 
 /** Compare for inequality with a raw pointer	**/
 template<typename T, typename U, typename = decltype((T*)nullptr != (U*)nullptr)>
-FORCEINLINE bool operator!=(const U* Lhs, const TLazyObjectPtr<T>& Rhs)
+UE_FORCEINLINE_HINT bool operator!=(const U* Lhs, const TLazyObjectPtr<T>& Rhs)
 {
 	return Lhs != Rhs.Get();
 }
 
 /** Compare for inequality with null **/
 template<typename T>
-FORCEINLINE bool operator!=(TYPE_OF_NULLPTR, const TLazyObjectPtr<T>& Rhs)
+UE_FORCEINLINE_HINT bool operator!=(TYPE_OF_NULLPTR, const TLazyObjectPtr<T>& Rhs)
 {
 	return Rhs.IsValid();
 }
diff --git a/Engine/Source/Runtime/CoreUObject/Public/UObject/Linker.h b/Engine/Source/Runtime/CoreUObject/Public/UObject/Linker.h
--- a/Engine/Source/Runtime/CoreUObject/Public/UObject/Linker.h
+++ b/Engine/Source/Runtime/CoreUObject/Public/UObject/Linker.h
@@ -433,7 +433,7 @@ public:
 
 	virtual ~FLinker();
 
-	FORCEINLINE ELinkerType::Type GetType() const
+	UE_FORCEINLINE_HINT ELinkerType::Type GetType() const
 	{
 		return LinkerType;
 	}
@@ -597,7 +597,7 @@ public:
 
 
 template<typename T> 
-FORCEINLINE T* Cast(FLinker* Src)
+UE_FORCEINLINE_HINT T* Cast(FLinker* Src)
 {
 	return Src && T::StaticType() == Src->GetType() ? (T*)Src : nullptr;
 }
diff --git a/Engine/Source/Runtime/CoreUObject/Public/UObject/LinkerLoad.h b/Engine/Source/Runtime/CoreUObject/Public/UObject/LinkerLoad.h
--- a/Engine/Source/Runtime/CoreUObject/Public/UObject/LinkerLoad.h
+++ b/Engine/Source/Runtime/CoreUObject/Public/UObject/LinkerLoad.h
@@ -155,7 +155,7 @@ protected:
 	// Variables.
 public:
 
-	FORCEINLINE static ELinkerType::Type StaticType()
+	UE_FORCEINLINE_HINT static ELinkerType::Type StaticType()
 	{
 		return ELinkerType::Load;
 	}
@@ -176,7 +176,7 @@ public:
 
 	/** True if Loader is FAsyncArchive  */
 	bool					bIsAsyncLoader;
-	FORCEINLINE FAsyncArchive* GetAsyncLoader()
+	UE_FORCEINLINE_HINT FAsyncArchive* GetAsyncLoader()
 	{
 		return bIsAsyncLoader ? (FAsyncArchive*)Loader : nullptr;
 	}
@@ -192,7 +192,7 @@ public:
 		return PackagePath;
 	}
 
-	FORCEINLINE const FLinkerInstancingContext& GetInstancingContext() const
+	UE_FORCEINLINE_HINT const FLinkerInstancingContext& GetInstancingContext() const
 	{
 		return InstancingContext;
 	}
@@ -283,7 +283,7 @@ public:
 	 */
 	COREUOBJECT_API void DetachLoader();
 
-	FORCEINLINE bool IsDestroyingLoader() const
+	UE_FORCEINLINE_HINT bool IsDestroyingLoader() const
 	{
 		return bIsDestroyingLoader;
 	}
@@ -304,7 +304,7 @@ public:
 	static constexpr int32 ExportHashCount = 256;
 	TUniquePtr<int32[]> ExportHash;
 
-	FORCEINLINE static int32 GetHashBucket(FName Object)
+	UE_FORCEINLINE_HINT static int32 GetHashBucket(FName Object)
 	{
 		return GetTypeHash(Object.GetComparisonIndex()) & (ExportHashCount - 1);
 	}
@@ -537,7 +537,7 @@ private:
 	FScopedSlowTask* LoadProgressScope;
 
 	/** Test whether we should report progress or not */
-	FORCEINLINE bool ShouldReportProgress() const
+	UE_FORCEINLINE_HINT bool ShouldReportProgress() const
 	{
 		return !IsAsyncLoading() && (LoadFlags & (LOAD_Quiet | LOAD_Async)) == 0;
 	}
@@ -610,13 +610,13 @@ public:
 	 *
 	 * @return true if initialized, false if pending.
 	 */
-	FORCEINLINE bool HasFinishedInitialization() const
+	UE_FORCEINLINE_HINT bool HasFinishedInitialization() const
 	{
         return bHasFinishedInitialization;
 	}
 
 	/** Returns ID of the thread that created this linker */
-	FORCEINLINE int32 GetOwnerThreadId() const
+	UE_FORCEINLINE_HINT int32 GetOwnerThreadId() const
 	{
 		return OwnerThread;
 	}
diff --git a/Engine/Source/Runtime/CoreUObject/Public/UObject/LinkerSave.h b/Engine/Source/Runtime/CoreUObject/Public/UObject/LinkerSave.h
--- a/Engine/Source/Runtime/CoreUObject/Public/UObject/LinkerSave.h
+++ b/Engine/Source/Runtime/CoreUObject/Public/UObject/LinkerSave.h
@@ -47,7 +47,7 @@ class FLinkerSave : public FLinker, public FArchiveUObject
 {
 public:
 
-	FORCEINLINE static ELinkerType::Type StaticType()
+	UE_FORCEINLINE_HINT static ELinkerType::Type StaticType()
 	{
 		return ELinkerType::Save;
 	}
diff --git a/Engine/Source/Runtime/CoreUObject/Public/UObject/Object.h b/Engine/Source/Runtime/CoreUObject/Public/UObject/Object.h
--- a/Engine/Source/Runtime/CoreUObject/Public/UObject/Object.h
+++ b/Engine/Source/Runtime/CoreUObject/Public/UObject/Object.h
@@ -309,8 +309,8 @@ public:
 	COREUOBJECT_API bool CanModify() const;
 
 #else
-	FORCEINLINE bool Modify(bool bAlwaysMarkDirty = true) { return false; }
-	FORCEINLINE bool IsCapturingAsRootObjectForTransaction() const { return false; };
+	UE_FORCEINLINE_HINT bool Modify(bool bAlwaysMarkDirty = true) { return false; }
+	UE_FORCEINLINE_HINT bool IsCapturingAsRootObjectForTransaction() const { return false; };
 #endif
 
 #if WITH_EDITOR
@@ -1844,7 +1844,7 @@ private:
 * @param	Test			The object to test
 * @return	Return true if the object is usable: non-null and not pending kill or garbage
 */
-FORCEINLINE bool IsValid(const UObject *Test)
+UE_FORCEINLINE_HINT bool IsValid(const UObject *Test)
 {
 	return Test && FInternalUObjectBaseUtilityIsValidFlagsChecker::CheckObjectValidBasedOnItsFlags(Test);
 }
diff --git a/Engine/Source/Runtime/CoreUObject/Public/UObject/ObjectHandle.h b/Engine/Source/Runtime/CoreUObject/Public/UObject/ObjectHandle.h
--- a/Engine/Source/Runtime/CoreUObject/Public/UObject/ObjectHandle.h
+++ b/Engine/Source/Runtime/CoreUObject/Public/UObject/ObjectHandle.h
@@ -62,7 +62,7 @@ namespace UE::CoreUObject::Private
 		//Stores either FPackedObjectRef or a UObject*
 		TSAN_ATOMIC(UPTRINT) PointerOrRef;
 
-		FORCEINLINE UPTRINT GetPointerOrRef() const
+		UE_FORCEINLINE_HINT UPTRINT GetPointerOrRef() const
 		{
 			return PointerOrRef;
 		}
@@ -95,7 +95,7 @@ namespace UE::CoreUObject::Private
 			PointerOrHandle = UPTRINT(RemoteInfo) | 1;
 		}
 
-		FORCEINLINE const UE::RemoteObject::Handle::FRemoteObjectStub* ToStub() const
+		UE_FORCEINLINE_HINT const UE::RemoteObject::Handle::FRemoteObjectStub* ToStub() const
 		{
 			return reinterpret_cast<UE::RemoteObject::Handle::FRemoteObjectStub*>(PointerOrHandle & ~UPTRINT(1));
 		}
diff --git a/Engine/Source/Runtime/CoreUObject/Public/UObject/ObjectKey.h b/Engine/Source/Runtime/CoreUObject/Public/UObject/ObjectKey.h
--- a/Engine/Source/Runtime/CoreUObject/Public/UObject/ObjectKey.h
+++ b/Engine/Source/Runtime/CoreUObject/Public/UObject/ObjectKey.h
@@ -44,7 +44,7 @@ public:
 		typename U
 		UE_REQUIRES(std::is_convertible_v<U, const UObject*>)
 	>
-	FORCEINLINE FObjectKey(U Object)
+	UE_FORCEINLINE_HINT FObjectKey(U Object)
 		: FObjectKey(ImplicitConv<const UObject*>(Object))
 	{
 	}
@@ -230,50 +230,50 @@ public:
 	typedef InElementType ElementType;
 
 	/** Default constructor */
-	FORCEINLINE TObjectKey() = default;
+	UE_FORCEINLINE_HINT TObjectKey() = default;
 
 	/** Construct from an object pointer */
 	template <
 		typename U
 		UE_REQUIRES(std::is_convertible_v<U, const InElementType*>)
 	>
-	FORCEINLINE TObjectKey(U Object)
+	UE_FORCEINLINE_HINT TObjectKey(U Object)
 		: ObjectKey(ImplicitConv<const InElementType*>(Object))
 	{
 	}
 
 	/** Compare this key with another */
-	FORCEINLINE bool operator==(const TObjectKey& Other) const
+	UE_FORCEINLINE_HINT bool operator==(const TObjectKey& Other) const
 	{
 		return ObjectKey == Other.ObjectKey;
 	}
 
 	/** Compare this key with another */
-	FORCEINLINE bool operator!=(const TObjectKey& Other) const
+	UE_FORCEINLINE_HINT bool operator!=(const TObjectKey& Other) const
 	{
 		return ObjectKey != Other.ObjectKey;
 	}
 
 	/** Compare this key with another */
-	FORCEINLINE bool operator<(const TObjectKey& Other) const
+	UE_FORCEINLINE_HINT bool operator<(const TObjectKey& Other) const
 	{
 		return ObjectKey < Other.ObjectKey;
 	}
 
 	/** Compare this key with another */
-	FORCEINLINE bool operator<=(const TObjectKey& Other) const
+	UE_FORCEINLINE_HINT bool operator<=(const TObjectKey& Other) const
 	{
 		return ObjectKey <= Other.ObjectKey;
 	}
 
 	/** Compare this key with another */
-	FORCEINLINE bool operator>(const TObjectKey& Other) const
+	UE_FORCEINLINE_HINT bool operator>(const TObjectKey& Other) const
 	{
 		return ObjectKey > Other.ObjectKey;
 	}
 
 	/** Compare this key with another */
-	FORCEINLINE bool operator>=(const TObjectKey& Other) const
+	UE_FORCEINLINE_HINT bool operator>=(const TObjectKey& Other) const
 	{
 		return ObjectKey >= Other.ObjectKey;
 	}
diff --git a/Engine/Source/Runtime/CoreUObject/Public/UObject/ObjectMemoryAnalyzer.h b/Engine/Source/Runtime/CoreUObject/Public/UObject/ObjectMemoryAnalyzer.h
--- a/Engine/Source/Runtime/CoreUObject/Public/UObject/ObjectMemoryAnalyzer.h
+++ b/Engine/Source/Runtime/CoreUObject/Public/UObject/ObjectMemoryAnalyzer.h
@@ -54,15 +54,15 @@ struct FObjectMemoryUsage
 	 * Determine if this annotation
 	 * @return true is this is a default pair. We only check the linker because CheckInvariants rules out bogus combinations
 	 */
-	FORCEINLINE bool IsDefault()
+	UE_FORCEINLINE_HINT bool IsDefault()
 	{
 		return false;
 	}
 
-	FORCEINLINE bool IsRoot() const { return !!(Flags&EObjFlags::IsRoot); }
-	FORCEINLINE bool IsReferencedByRoot() const { return !!(Flags&EObjFlags::IsReferencedByRoot); }
-	FORCEINLINE bool IsProcessed() const { return !!(Flags&EObjFlags::IsProcessed); }
-	FORCEINLINE bool IsReferencedByNonRoot() const { return !!(Flags&EObjFlags::IsReferencedByNonRoot); }
+	UE_FORCEINLINE_HINT bool IsRoot() const { return !!(Flags&EObjFlags::IsRoot); }
+	UE_FORCEINLINE_HINT bool IsReferencedByRoot() const { return !!(Flags&EObjFlags::IsReferencedByRoot); }
+	UE_FORCEINLINE_HINT bool IsProcessed() const { return !!(Flags&EObjFlags::IsProcessed); }
+	UE_FORCEINLINE_HINT bool IsReferencedByNonRoot() const { return !!(Flags&EObjFlags::IsReferencedByNonRoot); }
 	
 	SIZE_T InclusiveMemoryUsage;
 	SIZE_T ExclusiveMemoryUsage;
diff --git a/Engine/Source/Runtime/CoreUObject/Public/UObject/ObjectPtrLinkedListBuilder.h b/Engine/Source/Runtime/CoreUObject/Public/UObject/ObjectPtrLinkedListBuilder.h
--- a/Engine/Source/Runtime/CoreUObject/Public/UObject/ObjectPtrLinkedListBuilder.h
+++ b/Engine/Source/Runtime/CoreUObject/Public/UObject/ObjectPtrLinkedListBuilder.h
@@ -10,7 +10,7 @@ struct TObjectPtrLinkedListBuilderNextLink
 {
 	using ElementType = InElementType;
 
-	[[nodiscard]] FORCEINLINE static auto GetNextPtr(ElementType& Element)
+	[[nodiscard]] UE_FORCEINLINE_HINT static auto GetNextPtr(ElementType& Element)
 	{
 		return &(Element.Next);
 	}
@@ -31,7 +31,7 @@ public:
 	UE_NONCOPYABLE(TObjectPtrLinkedListBuilder);
 
 private:
-	FORCEINLINE void WriteEndPtr(InElementType* NewValue)
+	UE_FORCEINLINE_HINT void WriteEndPtr(InElementType* NewValue)
 	{
 		// Do not overwrite the same value to avoid dirtying the cache and
 		// also prevent TSAN from thinking we are messing around with existing data.
@@ -60,26 +60,26 @@ public:
 
 	// Move builder back to start and prepare for overwriting
 	// It only changes state of builder, use NullTerminate() to mark list as empty!
-	FORCEINLINE void Restart()
+	UE_FORCEINLINE_HINT void Restart()
 	{
 		EndPtr = StartPtr;
 	}
 
 	// Append element, don't touch next link
-	FORCEINLINE void AppendNoTerminate(ElementType& Element)
+	UE_FORCEINLINE_HINT void AppendNoTerminate(ElementType& Element)
 	{
 		WriteEndPtr(&Element);
 		EndPtr = LinkAccessor::GetNextPtr(Element);
 	}
 
 	// Append element and mark it as last
-	FORCEINLINE void AppendTerminated(ElementType& Element)
+	UE_FORCEINLINE_HINT void AppendTerminated(ElementType& Element)
 	{
 		AppendNoTerminate(Element);
 		NullTerminate();
 	}
 
-	FORCEINLINE void Remove(ElementType& Element)
+	UE_FORCEINLINE_HINT void Remove(ElementType& Element)
 	{
 		TObjectPtr<ElementType>* PrevIt = StartPtr;
 		for (TObjectPtr<ElementType> It = *StartPtr; It; It = GetNext(*It))
@@ -97,12 +97,12 @@ public:
 	}
 
 	// Mark end of the list
-	FORCEINLINE void NullTerminate()
+	UE_FORCEINLINE_HINT void NullTerminate()
 	{
 		WriteEndPtr(nullptr);
 	}
 
-	FORCEINLINE void MoveToEnd()
+	UE_FORCEINLINE_HINT void MoveToEnd()
 	{
 		for (TObjectPtr<ElementType> It = *StartPtr; It; It = GetNext(*It))
 		{
@@ -110,7 +110,7 @@ public:
 		}
 	}
 
-	FORCEINLINE bool MoveToNext()
+	UE_FORCEINLINE_HINT bool MoveToNext()
 	{
 		if (*EndPtr)
 		{
@@ -121,23 +121,23 @@ public:
 		return false;
 	}
 
-	[[nodiscard]] FORCEINLINE ElementType* GetNext(ElementType& Element) const
+	[[nodiscard]] UE_FORCEINLINE_HINT ElementType* GetNext(ElementType& Element) const
 	{
 		return GetNextRef(Element);
 	}
 
-	[[nodiscard]] FORCEINLINE ElementType* GetListStart() const
+	[[nodiscard]] UE_FORCEINLINE_HINT ElementType* GetListStart() const
 	{
 		return *StartPtr;
 	}
 
-	[[nodiscard]] FORCEINLINE ElementType* GetListEnd() const
+	[[nodiscard]] UE_FORCEINLINE_HINT ElementType* GetListEnd() const
 	{
 		return *EndPtr;
 	}
 
 private:
-	[[nodiscard]] FORCEINLINE TObjectPtr<ElementType>& GetNextRef(ElementType& Element) const
+	[[nodiscard]] UE_FORCEINLINE_HINT TObjectPtr<ElementType>& GetNextRef(ElementType& Element) const
 	{
 		return *LinkAccessor::GetNextPtr(Element);
 	}
diff --git a/Engine/Source/Runtime/CoreUObject/Public/UObject/ObjectResource.h b/Engine/Source/Runtime/CoreUObject/Public/UObject/ObjectResource.h
--- a/Engine/Source/Runtime/CoreUObject/Public/UObject/ObjectResource.h
+++ b/Engine/Source/Runtime/CoreUObject/Public/UObject/ObjectResource.h
@@ -52,30 +52,30 @@ class FPackageIndex
 	int32 Index;
 
 	/** Internal constructor, sets the index directly **/
-	FORCEINLINE explicit FPackageIndex(int32 InIndex)
+	UE_FORCEINLINE_HINT explicit FPackageIndex(int32 InIndex)
 		: Index(InIndex)
 	{
 
 	}
 public:
 	/** Constructor, sets the value to null **/
-	FORCEINLINE FPackageIndex()
+	UE_FORCEINLINE_HINT FPackageIndex()
 		: Index(0)
 	{
 
 	}
 	/** return true if this is an index into the import map **/
-	FORCEINLINE bool IsImport() const
+	UE_FORCEINLINE_HINT bool IsImport() const
 	{
 		return Index < 0;
 	}
 	/** return true if this is an index into the export map **/
-	FORCEINLINE bool IsExport() const
+	UE_FORCEINLINE_HINT bool IsExport() const
 	{
 		return Index > 0;
 	}
 	/** return true if this null (i.e. neither an import nor an export) **/
-	FORCEINLINE bool IsNull() const
+	UE_FORCEINLINE_HINT bool IsNull() const
 	{
 		return Index == 0;
 	}
@@ -92,7 +92,7 @@ public:
 		return Index - 1;
 	}
 	/** Return the raw value, for debugging purposes**/
-	FORCEINLINE int32 ForDebugging() const
+	UE_FORCEINLINE_HINT int32 ForDebugging() const
 	{
 		return Index;
 	}
@@ -111,30 +111,30 @@ public:
 	}
 
 	/** Compare package indecies for equality **/
-	FORCEINLINE bool operator==(const FPackageIndex& Other) const
+	UE_FORCEINLINE_HINT bool operator==(const FPackageIndex& Other) const
 	{
 		return Index == Other.Index;
 	}
 	/** Compare package indecies for inequality **/
-	FORCEINLINE bool operator!=(const FPackageIndex& Other) const
+	UE_FORCEINLINE_HINT bool operator!=(const FPackageIndex& Other) const
 	{
 		return Index != Other.Index;
 	}
 
 	/** Compare package indecies **/
-	FORCEINLINE bool operator<(const FPackageIndex& Other) const
+	UE_FORCEINLINE_HINT bool operator<(const FPackageIndex& Other) const
 	{
 		return Index < Other.Index;
 	}
-	FORCEINLINE bool operator>(const FPackageIndex& Other) const
+	UE_FORCEINLINE_HINT bool operator>(const FPackageIndex& Other) const
 	{
 		return Index > Other.Index;
 	}
-	FORCEINLINE bool operator<=(const FPackageIndex& Other) const
+	UE_FORCEINLINE_HINT bool operator<=(const FPackageIndex& Other) const
 	{
 		return Index <= Other.Index;
 	}
-	FORCEINLINE bool operator>=(const FPackageIndex& Other) const
+	UE_FORCEINLINE_HINT bool operator>=(const FPackageIndex& Other) const
 	{
 		return Index >= Other.Index;
 	}
@@ -156,12 +156,12 @@ public:
 	 * @param Slot - The structured archive slot to serialize from or to.
 	 * @param Value - The value to serialize.
 	 */
-	FORCEINLINE friend void operator<<(FStructuredArchive::FSlot Slot, FPackageIndex& Value)
+	UE_FORCEINLINE_HINT friend void operator<<(FStructuredArchive::FSlot Slot, FPackageIndex& Value)
 	{
 		Slot << Value.Index;
 	}
 
-	[[nodiscard]] FORCEINLINE friend uint32 GetTypeHash(const FPackageIndex& In)
+	[[nodiscard]] UE_FORCEINLINE_HINT friend uint32 GetTypeHash(const FPackageIndex& In)
 	{
 		return uint32(In.Index);
 	}
diff --git a/Engine/Source/Runtime/CoreUObject/Public/UObject/OverridableManager.h b/Engine/Source/Runtime/CoreUObject/Public/UObject/OverridableManager.h
--- a/Engine/Source/Runtime/CoreUObject/Public/UObject/OverridableManager.h
+++ b/Engine/Source/Runtime/CoreUObject/Public/UObject/OverridableManager.h
@@ -294,7 +294,7 @@ public:
 	 * @param PropertyChain leading to the property the caller is interested in
 	 * @param bOutInheritedOperation optional parameter to know if the operation returned was inherited from a parent property
 	 * @return the current type of override operation on the property */
-	FORCEINLINE EOverriddenPropertyOperation GetOverriddenPropertyOperation(TNotNull<UObject*> Object, const FPropertyChangedEvent& PropertyEvent, const FEditPropertyChain& PropertyChain, bool* bOutInheritedOperation = nullptr)
+	UE_FORCEINLINE_HINT EOverriddenPropertyOperation GetOverriddenPropertyOperation(TNotNull<UObject*> Object, const FPropertyChangedEvent& PropertyEvent, const FEditPropertyChain& PropertyChain, bool* bOutInheritedOperation = nullptr)
 	{
 		return GetOverriddenPropertyOperation(Object, FPropertyVisitorPath(PropertyEvent, PropertyChain), bOutInheritedOperation);
 	}
diff --git a/Engine/Source/Runtime/CoreUObject/Public/UObject/Package.h b/Engine/Source/Runtime/CoreUObject/Public/UObject/Package.h
--- a/Engine/Source/Runtime/CoreUObject/Public/UObject/Package.h
+++ b/Engine/Source/Runtime/CoreUObject/Public/UObject/Package.h
@@ -877,7 +877,7 @@ public:
 	* @return	true if the passed in flag is set, false otherwise
 	*			(including no flag passed in, unless the FlagsToCheck is CLASS_AllFlags)
 	*/
-	FORCEINLINE bool HasAnyPackageFlags( uint32 FlagsToCheck ) const
+	UE_FORCEINLINE_HINT bool HasAnyPackageFlags( uint32 FlagsToCheck ) const
 	{
 		return (GetPackageFlags() & FlagsToCheck) != 0;
 	}
@@ -888,7 +888,7 @@ public:
 	* @param FlagsToCheck	Package flags to check for
 	* @return true if all of the passed in flags are set (including no flags passed in), false otherwise
 	*/
-	FORCEINLINE bool HasAllPackagesFlags( uint32 FlagsToCheck ) const
+	UE_FORCEINLINE_HINT bool HasAllPackagesFlags( uint32 FlagsToCheck ) const
 	{
 		return ((GetPackageFlags() & FlagsToCheck) == FlagsToCheck);
 	}
@@ -913,7 +913,7 @@ public:
 	/**
 	* @return true if the package is marked as ExternallyReferenceable by all plugins and mount points
 	*/
-	FORCEINLINE bool IsExternallyReferenceable() const
+	UE_FORCEINLINE_HINT bool IsExternallyReferenceable() const
 	{
 		return GetAssetAccessSpecifier() == EAssetAccessSpecifier::Public;
 	}
@@ -923,7 +923,7 @@ public:
 	* 
 	* @param bValue Sets the package to be ExternallyReferenceable if true
 	*/
-	FORCEINLINE void SetIsExternallyReferenceable(bool bValue)
+	UE_FORCEINLINE_HINT void SetIsExternallyReferenceable(bool bValue)
 	{
 		SetAssetAccessSpecifier(bValue ? EAssetAccessSpecifier::Public : EAssetAccessSpecifier::Private);
 	}
@@ -1002,12 +1002,12 @@ public:
 	}
 
 	/** returns our persistent Guid */
-	FORCEINLINE FGuid GetPersistentGuid() const
+	UE_FORCEINLINE_HINT FGuid GetPersistentGuid() const
 	{
 		return PersistentGuid;
 	}
 	/** sets a specific persistent Guid */
-	FORCEINLINE void SetPersistentGuid(FGuid NewPersistentGuid)
+	UE_FORCEINLINE_HINT void SetPersistentGuid(FGuid NewPersistentGuid)
 	{
 		PersistentGuid = NewPersistentGuid;
 	}
@@ -1051,7 +1051,7 @@ public:
 	}
 
 	/** returns our FileSize */
-	FORCEINLINE int64 GetFileSize() const
+	UE_FORCEINLINE_HINT int64 GetFileSize() const
 	{
 		return AdditionalInfo.IsValid() ? AdditionalInfo->FileSize : 0;
 	}
@@ -1068,7 +1068,7 @@ public:
 	}
 
 	/** sets our ChunkIDs */
-	FORCEINLINE void SetChunkIDs(const TArray<int32>& InChunkIDs)
+	UE_FORCEINLINE_HINT void SetChunkIDs(const TArray<int32>& InChunkIDs)
 	{
 #if WITH_EDITORONLY_DATA
 		ChunkIDs = InChunkIDs;
@@ -1076,19 +1076,19 @@ public:
 	}
 
 	/** returns the unique package id */
-	FORCEINLINE FPackageId GetPackageId() const
+	UE_FORCEINLINE_HINT FPackageId GetPackageId() const
 	{
 		return PackageId;
 	}
 
 	/** sets the unique package id */
-	FORCEINLINE void SetPackageId(FPackageId InPackageId)
+	UE_FORCEINLINE_HINT void SetPackageId(FPackageId InPackageId)
 	{
 		PackageId = InPackageId;
 	}
 
 	/** returns the unique package id to load */
-	FORCEINLINE FPackageId GetPackageIdToLoad() const
+	UE_FORCEINLINE_HINT FPackageId GetPackageIdToLoad() const
 	{
 		return FPackageId::FromName(LoadedPath.GetPackageFName());
 	}
diff --git a/Engine/Source/Runtime/CoreUObject/Public/UObject/PackageReload.h b/Engine/Source/Runtime/CoreUObject/Public/UObject/PackageReload.h
--- a/Engine/Source/Runtime/CoreUObject/Public/UObject/PackageReload.h
+++ b/Engine/Source/Runtime/CoreUObject/Public/UObject/PackageReload.h
@@ -77,7 +77,7 @@ public:
 	/**
 	 * Get the old package pointer (the original package that was replaced).
 	 */
-	FORCEINLINE const UPackage* GetOldPackage() const
+	UE_FORCEINLINE_HINT const UPackage* GetOldPackage() const
 	{
 		return OldPackage;
 	}
@@ -85,7 +85,7 @@ public:
 	/**
 	 * Get the new package pointer (the package that replaced the original one). 
 	 */
-	FORCEINLINE const UPackage* GetNewPackage() const
+	UE_FORCEINLINE_HINT const UPackage* GetNewPackage() const
 	{
 		return NewPackage;
 	}
@@ -93,7 +93,7 @@ public:
 	/**
 	 * Get the raw map of repointed objects.
 	 */
-	FORCEINLINE const TMap<UObject*, UObject*>& GetRepointedObjects() const
+	UE_FORCEINLINE_HINT const TMap<UObject*, UObject*>& GetRepointedObjects() const
 	{
 		return ObjectPtrDecay(RepointedObjects);
 	}
@@ -124,7 +124,7 @@ public:
 	 * Get the set of objects that were referencing any of the objects that were replaced by the package reload.
 	 * @note Beware that object referencers may become null as the result of other operations that happen during reloading (such as recompiling Blueprints).
 	 */
-	FORCEINLINE const TArray<TWeakObjectPtr<UObject>>& GetObjectReferencers() const
+	UE_FORCEINLINE_HINT const TArray<TWeakObjectPtr<UObject>>& GetObjectReferencers() const
 	{
 		return ObjectReferencers;
 	}
@@ -132,7 +132,7 @@ public:
 	/**
 	 * Add a referencing object for any of the objects that were replaced by the package reload.
 	 */
-	FORCEINLINE void AddObjectReferencer(UObject* InObject)
+	UE_FORCEINLINE_HINT void AddObjectReferencer(UObject* InObject)
 	{
 		ObjectReferencers.AddUnique(InObject);
 	}
diff --git a/Engine/Source/Runtime/CoreUObject/Public/UObject/PersistentObjectPtr.h b/Engine/Source/Runtime/CoreUObject/Public/UObject/PersistentObjectPtr.h
--- a/Engine/Source/Runtime/CoreUObject/Public/UObject/PersistentObjectPtr.h
+++ b/Engine/Source/Runtime/CoreUObject/Public/UObject/PersistentObjectPtr.h
@@ -20,7 +20,7 @@ struct TPersistentObjectPtr
 	using ElementType = TObjectID;
 
 	/** Default constructor, will be null */
-	FORCEINLINE TPersistentObjectPtr()
+	UE_FORCEINLINE_HINT TPersistentObjectPtr()
 	{
 		Reset();
 	}
@@ -33,7 +33,7 @@ struct TPersistentObjectPtr
 	}
 
 	/** Resets the weak ptr only, call this when ObjectId may change */
-	FORCEINLINE void ResetWeakPtr()
+	UE_FORCEINLINE_HINT void ResetWeakPtr()
 	{
 		WeakPtr.Reset();
 	}
@@ -72,12 +72,12 @@ struct TPersistentObjectPtr
 			Reset();
 		}
 	}
-	FORCEINLINE void operator=(const UObject* Object)
+	UE_FORCEINLINE_HINT void operator=(const UObject* Object)
 	{
 		*this = FObjectPtr(const_cast<UObject*>(Object));
 	}
 	template <typename T>
-	FORCEINLINE void operator=(TObjectPtr<T> Object)
+	UE_FORCEINLINE_HINT void operator=(TObjectPtr<T> Object)
 	{
 		// This needs to be a template instead of TObjectPtr<const UObject> because C++ does derived-to-base
 		// pointer conversions ('standard conversion sequences') in more cases than TSmartPtr<Derived>-to-TSmartPtr<Base>
@@ -101,13 +101,13 @@ struct TPersistentObjectPtr
 	 *
 	 * @return Unique ID for this object, or an invalid FUniqueObjectGuid if this pointer isn't set to anything
 	 */
-	FORCEINLINE const TObjectID& GetUniqueID() const
+	UE_FORCEINLINE_HINT const TObjectID& GetUniqueID() const
 	{
 		return ObjectID;
 	}
 
 	/** Non-const version of the above */
-	FORCEINLINE TObjectID& GetUniqueID()
+	UE_FORCEINLINE_HINT TObjectID& GetUniqueID()
 	{
 		return ObjectID;
 	}
@@ -165,46 +165,46 @@ struct TPersistentObjectPtr
 	}
 
 	/** Dereference the pointer */
-	FORCEINLINE UObject& operator*() const
+	UE_FORCEINLINE_HINT UObject& operator*() const
 	{
 		return *Get();
 	}
 
 	/** Dereference the pointer */
-	FORCEINLINE UObject* operator->() const
+	UE_FORCEINLINE_HINT UObject* operator->() const
 	{
 		return Get();
 	}
 
 	/** Compare pointers for equality. Only Serial Number matters for the base implementation */
-	FORCEINLINE bool operator==(const TPersistentObjectPtr& Rhs) const
+	UE_FORCEINLINE_HINT bool operator==(const TPersistentObjectPtr& Rhs) const
 	{
 		return ObjectID == Rhs.ObjectID;
 	}
 
-	FORCEINLINE bool operator==(TYPE_OF_NULLPTR) const
+	UE_FORCEINLINE_HINT bool operator==(TYPE_OF_NULLPTR) const
 	{
 		return !IsValid();
 	}
 
 	/** Compare pointers for inequality. Only Serial Number matters for the base implementation */
-	FORCEINLINE bool operator!=(const TPersistentObjectPtr& Rhs) const
+	UE_FORCEINLINE_HINT bool operator!=(const TPersistentObjectPtr& Rhs) const
 	{
 		return ObjectID != Rhs.ObjectID;
 	}
 
-	FORCEINLINE bool operator!=(TYPE_OF_NULLPTR) const
+	UE_FORCEINLINE_HINT bool operator!=(TYPE_OF_NULLPTR) const
 	{
 		return IsValid();
 	}
 
 #if !PLATFORM_COMPILER_HAS_GENERATED_COMPARISON_OPERATORS
-	FORCEINLINE friend bool operator==(TYPE_OF_NULLPTR, const TPersistentObjectPtr& Rhs)
+	UE_FORCEINLINE_HINT friend bool operator==(TYPE_OF_NULLPTR, const TPersistentObjectPtr& Rhs)
 	{
 		return !Rhs.IsValid();
 	}
 
-	FORCEINLINE friend bool operator!=(TYPE_OF_NULLPTR, const TPersistentObjectPtr& Rhs)
+	UE_FORCEINLINE_HINT friend bool operator!=(TYPE_OF_NULLPTR, const TPersistentObjectPtr& Rhs)
 	{
 		return Rhs.IsValid();
 	}
@@ -215,7 +215,7 @@ struct TPersistentObjectPtr
 	 * 
 	 * @return true if this does not point to a real object, but could possibly
 	 */
-	FORCEINLINE bool IsPending() const
+	UE_FORCEINLINE_HINT bool IsPending() const
 	{
 		return Get() == nullptr && ObjectID.IsValid();
 	}
@@ -225,7 +225,7 @@ struct TPersistentObjectPtr
 	 *
 	 * @return true if Get() would return a valid non-null pointer
 	 */
-	FORCEINLINE bool IsValid() const
+	UE_FORCEINLINE_HINT bool IsValid() const
 	{
 		return !!Get();
 	}
@@ -235,7 +235,7 @@ struct TPersistentObjectPtr
 	 *
 	 * @return true if this used to point at a real object but no longer does.
 	 */
-	FORCEINLINE bool IsStale() const
+	UE_FORCEINLINE_HINT bool IsStale() const
 	{
 		return WeakPtr.IsStale();
 	}
@@ -244,13 +244,13 @@ struct TPersistentObjectPtr
 	 *
 	 * @return true if this is explicitly pointing to no object
 	 */
-	FORCEINLINE bool IsNull() const
+	UE_FORCEINLINE_HINT bool IsNull() const
 	{
 		return !ObjectID.IsValid();
 	}
 
 	/** Hash function */
-	FORCEINLINE friend uint32 GetTypeHash(const TPersistentObjectPtr& Ptr)
+	UE_FORCEINLINE_HINT friend uint32 GetTypeHash(const TPersistentObjectPtr& Ptr)
 	{
 		return GetTypeHash(Ptr.ObjectID);
 	}
diff --git a/Engine/Source/Runtime/CoreUObject/Public/UObject/PropertyIterator.h b/Engine/Source/Runtime/CoreUObject/Public/UObject/PropertyIterator.h
--- a/Engine/Source/Runtime/CoreUObject/Public/UObject/PropertyIterator.h
+++ b/Engine/Source/Runtime/CoreUObject/Public/UObject/PropertyIterator.h
@@ -39,12 +39,12 @@ public:
 	}
 
 	/** conversion to "bool" returning true if the iterator is valid. */
-	FORCEINLINE explicit operator bool() const
+	UE_FORCEINLINE_HINT explicit operator bool() const
 	{
 		return Field != nullptr;
 	}
 	/** inverse of the "bool" operator */
-	FORCEINLINE bool operator !() const
+	UE_FORCEINLINE_HINT bool operator !() const
 	{
 		return !(bool)*this;
 	}
diff --git a/Engine/Source/Runtime/CoreUObject/Public/UObject/PropertyOptional.h b/Engine/Source/Runtime/CoreUObject/Public/UObject/PropertyOptional.h
--- a/Engine/Source/Runtime/CoreUObject/Public/UObject/PropertyOptional.h
+++ b/Engine/Source/Runtime/CoreUObject/Public/UObject/PropertyOptional.h
@@ -167,11 +167,11 @@ protected:
 		return ValueProperty->GetSize();
 	}
 
-	FORCEINLINE bool* GetIsSetPointer(void* Data) const
+	UE_FORCEINLINE_HINT bool* GetIsSetPointer(void* Data) const
 	{
 		return reinterpret_cast<bool*>(reinterpret_cast<uint8*>(Data) + CalcIsSetOffset());
 	}
-	FORCEINLINE const bool* GetIsSetPointer(const void* Data) const
+	UE_FORCEINLINE_HINT const bool* GetIsSetPointer(const void* Data) const
 	{
 		return reinterpret_cast<const bool*>(reinterpret_cast<const uint8*>(Data) + CalcIsSetOffset());
 	}
diff --git a/Engine/Source/Runtime/CoreUObject/Public/UObject/PropertyTag.h b/Engine/Source/Runtime/CoreUObject/Public/UObject/PropertyTag.h
--- a/Engine/Source/Runtime/CoreUObject/Public/UObject/PropertyTag.h
+++ b/Engine/Source/Runtime/CoreUObject/Public/UObject/PropertyTag.h
@@ -117,7 +117,7 @@ struct UE_INTERNAL FPropertyTagScope
 		CurrentPropertyTag = PropertyTagToRestore;
 	}
 
-	static FORCEINLINE const FPropertyTag* GetCurrentPropertyTag()
+	static UE_FORCEINLINE_HINT const FPropertyTag* GetCurrentPropertyTag()
 	{
 		return CurrentPropertyTag;
 	}
diff --git a/Engine/Source/Runtime/CoreUObject/Public/UObject/PropertyTempVal.h b/Engine/Source/Runtime/CoreUObject/Public/UObject/PropertyTempVal.h
--- a/Engine/Source/Runtime/CoreUObject/Public/UObject/PropertyTempVal.h
+++ b/Engine/Source/Runtime/CoreUObject/Public/UObject/PropertyTempVal.h
@@ -29,13 +29,13 @@ struct FPropertyTempVal
 	COREUOBJECT_API void ExportText(FString& ValueStr, const void* Defaults = nullptr, UObject* Parent = nullptr, int32 PortFlags = 0, UObject* ExportRootScope = nullptr);
 
 	// Returns a pointer to the internal instance
-	FORCEINLINE void* Get()
+	UE_FORCEINLINE_HINT void* Get()
 	{
 		return Value;
 	}
 
 	// Returns a pointer to the internal instance
-	FORCEINLINE const void* Get() const
+	UE_FORCEINLINE_HINT const void* Get() const
 	{
 		return Value;
 	}
diff --git a/Engine/Source/Runtime/CoreUObject/Public/UObject/RemoteObject.h b/Engine/Source/Runtime/CoreUObject/Public/UObject/RemoteObject.h
--- a/Engine/Source/Runtime/CoreUObject/Public/UObject/RemoteObject.h
+++ b/Engine/Source/Runtime/CoreUObject/Public/UObject/RemoteObject.h
@@ -22,7 +22,7 @@ namespace UE::RemoteObject
 	* Returns a unique id associated with this (server) process
 	*/
 	UE_DEPRECATED(5.6, "Use FRemoteServerId::GetLocalServerId() instead.")
-	FORCEINLINE FRemoteServerId GetGlobalServerId()
+	UE_FORCEINLINE_HINT FRemoteServerId GetGlobalServerId()
 	{
 		return FRemoteServerId::GetLocalServerId();
 	}
@@ -82,7 +82,7 @@ namespace UE::RemoteObject::Handle
 		}
 	};
 
-	FORCEINLINE uint32 GetTypeHash(const FPhysicsIslandId& Key)
+	UE_FORCEINLINE_HINT uint32 GetTypeHash(const FPhysicsIslandId& Key)
 	{
 		// Combine hashes of the members
 		return HashCombine(::GetTypeHash(Key.PhysicsServerId.GetIdNumber()), ::GetTypeHash(Key.PhysicsLocalIslandId));
diff --git a/Engine/Source/Runtime/CoreUObject/Public/UObject/RemoteObjectPathName.h b/Engine/Source/Runtime/CoreUObject/Public/UObject/RemoteObjectPathName.h
--- a/Engine/Source/Runtime/CoreUObject/Public/UObject/RemoteObjectPathName.h
+++ b/Engine/Source/Runtime/CoreUObject/Public/UObject/RemoteObjectPathName.h
@@ -12,7 +12,7 @@ struct FRemoteObjectTables
 	TArray<FName> Names;
 	TArray<FRemoteObjectId> RemoteIds;
 
-	FORCEINLINE int32 Num() const
+	UE_FORCEINLINE_HINT int32 Num() const
 	{
 		return Names.Num();
 	}
@@ -38,17 +38,17 @@ struct FPackedRemoteObjectPathName
 	TArray<FNameIndexType> RemoteIds;
 	TArray<FNameIndexType> Names;
 
-	FORCEINLINE int32 Num() const
+	UE_FORCEINLINE_HINT int32 Num() const
 	{
 		return Names.Num();
 	}
 
-	FORCEINLINE FName GetSegmentName(int32 InSegmentIndex, const FRemoteObjectTables& InTables) const
+	UE_FORCEINLINE_HINT FName GetSegmentName(int32 InSegmentIndex, const FRemoteObjectTables& InTables) const
 	{
 		return InTables.Names[Names[InSegmentIndex]];
 	}
 
-	FORCEINLINE FRemoteObjectId GetSegmentId(int32 InSegmentIndex, const FRemoteObjectTables& InTables) const
+	UE_FORCEINLINE_HINT FRemoteObjectId GetSegmentId(int32 InSegmentIndex, const FRemoteObjectTables& InTables) const
 	{
 		return InTables.RemoteIds[RemoteIds[InSegmentIndex]];
 	}
@@ -80,12 +80,12 @@ struct FRemoteObjectPathName : public FRemoteObjectTables
 	FRemoteObjectPathName(const FRemoteObjectPathName&) = default;
 	FRemoteObjectPathName& operator=(const FRemoteObjectPathName&) = default;
 
-	FORCEINLINE FName GetSegmentName(int32 InSegmentIndex, const FRemoteObjectTables&) const
+	UE_FORCEINLINE_HINT FName GetSegmentName(int32 InSegmentIndex, const FRemoteObjectTables&) const
 	{
 		return Names[InSegmentIndex];
 	}
 
-	FORCEINLINE FRemoteObjectId GetSegmentId(int32 InSegmentIndex, const FRemoteObjectTables&) const
+	UE_FORCEINLINE_HINT FRemoteObjectId GetSegmentId(int32 InSegmentIndex, const FRemoteObjectTables&) const
 	{
 		return RemoteIds[InSegmentIndex];
 	}
diff --git a/Engine/Source/Runtime/CoreUObject/Public/UObject/RemoteObjectTransfer.h b/Engine/Source/Runtime/CoreUObject/Public/UObject/RemoteObjectTransfer.h
--- a/Engine/Source/Runtime/CoreUObject/Public/UObject/RemoteObjectTransfer.h
+++ b/Engine/Source/Runtime/CoreUObject/Public/UObject/RemoteObjectTransfer.h
@@ -103,7 +103,7 @@ namespace UE::RemoteObject::Transfer
 	* Utility that converts FRemoteObjectId to TObjectPtr of a given type. Will not resolve the object if it's not local.
 	*/
 	UE_DEPRECATED(5.6, "RemoteObjectIdToObjectPtr is deprecated use FObjectPtr(FRemoteObjectId) instead.")
-	FORCEINLINE FObjectPtr RemoteObjectIdToObjectPtr(FRemoteObjectId RemoteId)
+	UE_FORCEINLINE_HINT FObjectPtr RemoteObjectIdToObjectPtr(FRemoteObjectId RemoteId)
 	{
 #if UE_WITH_REMOTE_OBJECT_HANDLE
 		return FObjectPtr(UE::CoreUObject::Private::FRemoteObjectHandlePrivate::FromIdNoResolve(RemoteId));
@@ -199,7 +199,7 @@ public:
 	{
 	}
 
-	FORCEINLINE bool operator==(const FRemoteObjectReference& Other) const
+	UE_FORCEINLINE_HINT bool operator==(const FRemoteObjectReference& Other) const
 	{
 		return ObjectId == Other.ObjectId;
 	}
@@ -208,15 +208,15 @@ public:
 	COREUOBJECT_API FWeakObjectPtr ToWeakPtr() const;
 	COREUOBJECT_API UObject* Resolve() const;
 
-	FORCEINLINE FRemoteObjectId GetRemoteId() const
+	UE_FORCEINLINE_HINT FRemoteObjectId GetRemoteId() const
 	{
 		return ObjectId;
 	}
-	FORCEINLINE FRemoteServerId GetSharingServerId() const
+	UE_FORCEINLINE_HINT FRemoteServerId GetSharingServerId() const
 	{
 		return ServerId;
 	}
-	FORCEINLINE bool IsRemote() const
+	UE_FORCEINLINE_HINT bool IsRemote() const
 	{
 		return UE::RemoteObject::Handle::IsRemote(ObjectId);
 	}
diff --git a/Engine/Source/Runtime/CoreUObject/Public/UObject/RemoteObjectTypes.h b/Engine/Source/Runtime/CoreUObject/Public/UObject/RemoteObjectTypes.h
--- a/Engine/Source/Runtime/CoreUObject/Public/UObject/RemoteObjectTypes.h
+++ b/Engine/Source/Runtime/CoreUObject/Public/UObject/RemoteObjectTypes.h
@@ -116,7 +116,7 @@ struct FRemoteServerId
 		return (Id == (uint32)ERemoteServerIdConstants::Local) || (*this == GlobalServerId);
 	}
 
-	static FORCEINLINE FRemoteServerId GetLocalServerId()
+	static UE_FORCEINLINE_HINT FRemoteServerId GetLocalServerId()
 	{
 		return FRemoteServerId(ERemoteServerIdConstants::Local);
 	}
@@ -132,7 +132,7 @@ struct FRemoteServerId
 	COREUOBJECT_API bool Serialize(FArchive& Ar);
 	COREUOBJECT_API bool NetSerialize(FArchive& Ar, class UPackageMap* Map, bool& bOutSuccess);
 
-	FORCEINLINE uint32 GetTypeHash() const
+	UE_FORCEINLINE_HINT uint32 GetTypeHash() const
 	{
 		return ::GetTypeHash(GetLocalized().Id);
 	}
@@ -254,58 +254,58 @@ public:
 
 	explicit COREUOBJECT_API FRemoteObjectId(const UObjectBase* InObject);
 
-	FORCEINLINE uint32 GetTypeHash() const
+	UE_FORCEINLINE_HINT uint32 GetTypeHash() const
 	{
 		return ::GetTypeHash(GetLocalized().Id);
 	}
 
-	FORCEINLINE bool operator==(const FRemoteObjectId& Other) const
+	UE_FORCEINLINE_HINT bool operator==(const FRemoteObjectId& Other) const
 	{
 		// Don't compare against the Reserved bit
 		return Compare(Other) == 0;
 	}
-	FORCEINLINE bool operator!=(const FRemoteObjectId& Other) const
+	UE_FORCEINLINE_HINT bool operator!=(const FRemoteObjectId& Other) const
 	{
 		// Don't compare against the Reserved bit
 		return Compare(Other) != 0;
 	}
-	FORCEINLINE bool operator<(const FRemoteObjectId& Other) const
+	UE_FORCEINLINE_HINT bool operator<(const FRemoteObjectId& Other) const
 	{
 		return Compare(Other) < 0;
 	}
-	FORCEINLINE bool operator<=(const FRemoteObjectId& Other) const
+	UE_FORCEINLINE_HINT bool operator<=(const FRemoteObjectId& Other) const
 	{
 		return Compare(Other) <= 0;
 	}
-	FORCEINLINE bool operator>(const FRemoteObjectId& Other) const
+	UE_FORCEINLINE_HINT bool operator>(const FRemoteObjectId& Other) const
 	{
 		return Compare(Other) > 0;
 	}
-	FORCEINLINE bool operator>=(const FRemoteObjectId& Other) const
+	UE_FORCEINLINE_HINT bool operator>=(const FRemoteObjectId& Other) const
 	{
 		return Compare(Other) >= 0;
 	}
-	FORCEINLINE bool IsValid() const
+	UE_FORCEINLINE_HINT bool IsValid() const
 	{
 		return *this != FRemoteObjectId();
 	}
 
-	FORCEINLINE uint64 GetIdNumber() const
+	UE_FORCEINLINE_HINT uint64 GetIdNumber() const
 	{
 		return GetGlobalized().Id;
 	}
 	
-	FORCEINLINE FRemoteServerId GetServerId() const
+	UE_FORCEINLINE_HINT FRemoteServerId GetServerId() const
 	{
 		return FRemoteServerId(ServerId, ERemoteServerIdConstants::Invalid);
 	}
 
-	FORCEINLINE bool IsAsset() const
+	UE_FORCEINLINE_HINT bool IsAsset() const
 	{
 		return GetServerId().IsAsset();
 	}
 
-	FORCEINLINE bool IsLocal() const
+	UE_FORCEINLINE_HINT bool IsLocal() const
 	{
 		return GetServerId().IsLocal();
 	}
@@ -324,12 +324,12 @@ public:
 	friend class FNetGUIDCache;	
 };
 
-FORCEINLINE uint32 GetTypeHash(const FRemoteObjectId& ObjectId)
+UE_FORCEINLINE_HINT uint32 GetTypeHash(const FRemoteObjectId& ObjectId)
 {
 	return ObjectId.GetTypeHash();
 }
 
-FORCEINLINE uint32 GetTypeHash(const FRemoteServerId& ServerId)
+UE_FORCEINLINE_HINT uint32 GetTypeHash(const FRemoteServerId& ServerId)
 {
 	return ServerId.GetTypeHash();
 }
