Replace some FORCEINLINE usage in Core/CoreUObject with UE_FORCEINLINE_HINT.

diff --git a/Engine/Source/Runtime/CoreUObject/Public/UObject/Script.h b/Engine/Source/Runtime/CoreUObject/Public/UObject/Script.h
--- a/Engine/Source/Runtime/CoreUObject/Public/UObject/Script.h
+++ b/Engine/Source/Runtime/CoreUObject/Public/UObject/Script.h
@@ -169,7 +169,7 @@ enum EFunctionFlags : uint32
 	FUNC_AllFlags		= 0xFFFFFFFF,
 };
 
-FORCEINLINE FArchive& operator<<(FArchive& Ar, EFunctionFlags& Flags)
+UE_FORCEINLINE_HINT FArchive& operator<<(FArchive& Ar, EFunctionFlags& Flags)
 {
 	Ar << (uint32&)Flags;
 	return Ar;
@@ -475,7 +475,7 @@ struct COREUOBJECT_API FBlueprintContextTracker : TThreadSingleton<FBlueprintCon
 	void ResetRunaway();
 
 	/** Increments and returns the Runaway counter */
-	FORCEINLINE int32 AddRunaway()
+	UE_FORCEINLINE_HINT int32 AddRunaway()
 	{
 		return ++Runaway;
 	}
@@ -493,25 +493,25 @@ struct COREUOBJECT_API FBlueprintContextTracker : TThreadSingleton<FBlueprintCon
 	bool RecordAccessViolation(const UObject* Object);
 
 	/** Returns how many function executions deep we are, may be higher than ScriptStack size */
-	FORCEINLINE int32 GetScriptEntryTag() const
+	UE_FORCEINLINE_HINT int32 GetScriptEntryTag() const
 	{
 		return ScriptEntryTag;
 	}
 
 	/** Returns current script stack frame */
 	UE_DEPRECATED(5.1, "GetScriptStack() inefficiently copies the array to return and is now deprecated. Use GetCurrentScriptStack() which returns a TArrayView instead")
-	FORCEINLINE TArray<const FFrame*> GetScriptStack() const
+	UE_FORCEINLINE_HINT TArray<const FFrame*> GetScriptStack() const
 	{
 		return TArray<const FFrame*>(GetCurrentScriptStack());
 	}
 
 	/** Returns current script stack frame */
-	FORCEINLINE TArrayView<const FFrame* const> GetCurrentScriptStack() const
+	UE_FORCEINLINE_HINT TArrayView<const FFrame* const> GetCurrentScriptStack() const
 	{
 		return MakeArrayView<const FFrame* const>(ScriptStack.GetData(), ScriptStack.Num());
 	}
 
-	FORCEINLINE TArrayView<FFrame* const> GetCurrentScriptStackWritable() const
+	UE_FORCEINLINE_HINT TArrayView<FFrame* const> GetCurrentScriptStackWritable() const
 	{
 		return MakeArrayView<FFrame* const>(ScriptStack.GetData(), ScriptStack.Num());
 	}
diff --git a/Engine/Source/Runtime/CoreUObject/Public/UObject/ScriptInterface.h b/Engine/Source/Runtime/CoreUObject/Public/UObject/ScriptInterface.h
--- a/Engine/Source/Runtime/CoreUObject/Public/UObject/ScriptInterface.h
+++ b/Engine/Source/Runtime/CoreUObject/Public/UObject/ScriptInterface.h
@@ -59,7 +59,7 @@ public:
 	/**
 	 * Returns the ObjectPointer contained by this FScriptInterface
 	 */
-	FORCEINLINE UObject* GetObject() const
+	UE_FORCEINLINE_HINT UObject* GetObject() const
 	{
 		return ObjectPointer;
 	}
@@ -67,7 +67,7 @@ public:
 	/**
 	 * Returns the ObjectPointer contained by this FScriptInterface
 	 */
-	FORCEINLINE TObjectPtr<UObject>& GetObjectRef()
+	UE_FORCEINLINE_HINT TObjectPtr<UObject>& GetObjectRef()
 	{
 		return ObjectPointer;
 	}
@@ -75,7 +75,7 @@ public:
 	/**
 	 * Returns the pointer to the native interface if it is valid
 	 */
-	FORCEINLINE void* GetInterface() const
+	UE_FORCEINLINE_HINT void* GetInterface() const
 	{
 		// Only access the InterfacePointer if we have a valid ObjectPointer. This is necessary because garbage collection may only clear the ObjectPointer.
 		// This will also return null for objects that only implement the interface in a blueprint class because there is no native representation.
@@ -97,7 +97,7 @@ public:
 	/**
 	 * Sets the value of the InterfacePointer for this FScriptInterface
 	 */
-	FORCEINLINE void SetInterface( void* InInterfacePointer )
+	UE_FORCEINLINE_HINT void SetInterface( void* InInterfacePointer )
 	{
 		InterfacePointer = InInterfacePointer;
 	}
@@ -105,11 +105,11 @@ public:
 	/**
 	 * Comparison operator, taking a reference to another FScriptInterface
 	 */
-	FORCEINLINE bool operator==( const FScriptInterface& Other ) const
+	UE_FORCEINLINE_HINT bool operator==( const FScriptInterface& Other ) const
 	{
 		return GetInterface() == Other.GetInterface() && ObjectPointer == Other.GetObject();
 	}
-	FORCEINLINE bool operator!=( const FScriptInterface& Other ) const
+	UE_FORCEINLINE_HINT bool operator!=( const FScriptInterface& Other ) const
 	{
 		return GetInterface() != Other.GetInterface() || ObjectPointer != Other.GetObject();
 	}
@@ -274,7 +274,7 @@ public:
 		typename OtherInterface
 		UE_REQUIRES(std::is_convertible_v<OtherInterface*, InInterfaceType*>)
 	>
-	FORCEINLINE bool operator==( const OtherInterface* Other ) const
+	UE_FORCEINLINE_HINT bool operator==( const OtherInterface* Other ) const
 	{
 		return GetInterface() == Other;
 	}
@@ -282,7 +282,7 @@ public:
 		typename OtherInterface
 		UE_REQUIRES(std::is_convertible_v<OtherInterface*, InInterfaceType*>)
 	>
-	FORCEINLINE bool operator!=( const OtherInterface* Other ) const
+	UE_FORCEINLINE_HINT bool operator!=( const OtherInterface* Other ) const
 	{
 		return GetInterface() != Other;
 	}
@@ -290,11 +290,11 @@ public:
 	/**
 	 * Comparison operator, taking a reference to another TScriptInterface
 	 */
-	FORCEINLINE bool operator==( const TScriptInterface& Other ) const
+	UE_FORCEINLINE_HINT bool operator==( const TScriptInterface& Other ) const
 	{
 		return GetInterface() == Other.GetInterface() && GetObject() == Other.GetObject();
 	}
-	FORCEINLINE bool operator!=( const TScriptInterface& Other ) const
+	UE_FORCEINLINE_HINT bool operator!=( const TScriptInterface& Other ) const
 	{
 		return GetInterface() != Other.GetInterface() || GetObject() != Other.GetObject();
 	}
@@ -302,11 +302,11 @@ public:
 	/**
 	 * Comparison operator, taking a nullptr
 	 */
-	FORCEINLINE bool operator==(TYPE_OF_NULLPTR) const
+	UE_FORCEINLINE_HINT bool operator==(TYPE_OF_NULLPTR) const
 	{
 		return GetInterface() == nullptr;
 	}
-	FORCEINLINE bool operator!=(TYPE_OF_NULLPTR) const
+	UE_FORCEINLINE_HINT bool operator!=(TYPE_OF_NULLPTR) const
 	{
 		return GetInterface() != nullptr;
 	}
@@ -314,7 +314,7 @@ public:
 	/**
 	 * Member access operator.  Provides transparent access to the native interface pointer contained by this TScriptInterface
 	 */
-	FORCEINLINE InInterfaceType* operator->() const
+	UE_FORCEINLINE_HINT InInterfaceType* operator->() const
 	{
 		return GetInterface();
 	}
@@ -324,7 +324,7 @@ public:
 	 *
 	 * @return	a reference (of type InterfaceType) to the object pointed to by InterfacePointer
 	 */
-	FORCEINLINE InInterfaceType& operator*() const
+	UE_FORCEINLINE_HINT InInterfaceType& operator*() const
 	{
 		return *GetInterface();
 	}
@@ -332,7 +332,7 @@ public:
 	/**
 	 * Returns the pointer to the interface
 	 */
-	FORCEINLINE InInterfaceType* GetInterface() const
+	UE_FORCEINLINE_HINT InInterfaceType* GetInterface() const
 	{
 		return (InInterfaceType*)FScriptInterface::GetInterface();
 	}
@@ -340,7 +340,7 @@ public:
 	/**
 	 * Sets the value of the InterfacePointer for this TScriptInterface
 	 */
-	FORCEINLINE void SetInterface(InInterfaceType* InInterfacePointer)
+	UE_FORCEINLINE_HINT void SetInterface(InInterfaceType* InInterfacePointer)
 	{
 		FScriptInterface::SetInterface((void*)InInterfacePointer);
 	}
@@ -348,7 +348,7 @@ public:
 	/**
 	 * Returns the ObjectPointer contained by this TScriptInterface
 	 */
-	FORCEINLINE UObjectType* GetObject() const
+	UE_FORCEINLINE_HINT UObjectType* GetObject() const
 	{
 		return FScriptInterface::GetObject();
 	}
@@ -356,7 +356,7 @@ public:
 	/**
 	 * Returns the ObjectPointer contained by this TScriptInterface
 	 */
-	FORCEINLINE TObjectPtr<UObjectType>& GetObjectRef()
+	UE_FORCEINLINE_HINT TObjectPtr<UObjectType>& GetObjectRef()
 	{
 		return *(TObjectPtr<UObjectType>*)&FScriptInterface::GetObjectRef();
 	}
@@ -364,7 +364,7 @@ public:
 	/**
 	 * Sets the value of the ObjectPointer for this TScriptInterface
 	 */
-	FORCEINLINE void SetObject( UObjectType* InObjectPointer )
+	UE_FORCEINLINE_HINT void SetObject( UObjectType* InObjectPointer )
 	{
 		FScriptInterface::SetObject(const_cast<UObject*>(InObjectPointer));
 	}
@@ -373,7 +373,7 @@ public:
 	 * Boolean operator, returns true if this object natively implements InterfaceType.
 	 * This will return false for objects that only implement the interface in blueprint classes.
 	 */
-	FORCEINLINE explicit operator bool() const
+	UE_FORCEINLINE_HINT explicit operator bool() const
 	{
 		return GetInterface() != nullptr;
 	}
diff --git a/Engine/Source/Runtime/CoreUObject/Public/UObject/SoftObjectPath.h b/Engine/Source/Runtime/CoreUObject/Public/UObject/SoftObjectPath.h
--- a/Engine/Source/Runtime/CoreUObject/Public/UObject/SoftObjectPath.h
+++ b/Engine/Source/Runtime/CoreUObject/Public/UObject/SoftObjectPath.h
@@ -237,21 +237,21 @@ public:
 
 	/** Returns the sub path, which is often empty */
 	UE_SOFTOBJECTPATH_FSTRING_DEPRECATED(5.6, "FSoftObjectPath::GetSubPathString has been deprecated - please use GetSubPathUtf8String() instead.")
-	[[nodiscard]] FORCEINLINE FString GetSubPathString() const
+	[[nodiscard]] UE_FORCEINLINE_HINT FString GetSubPathString() const
 	{
 		return FString(SubPathString);
 	}
-	[[nodiscard]] FORCEINLINE const FUtf8String& GetSubPathUtf8String() const
+	[[nodiscard]] UE_FORCEINLINE_HINT const FUtf8String& GetSubPathUtf8String() const
 	{
 		return SubPathString;
 	}
 
 	UE_SOFTOBJECTPATH_FSTRING_DEPRECATED(5.6, "FSoftObjectPath::ConstructFromAssetPathAndSubpath with a wide string subpath has been deprecated - please use a UTF8 subpath instead.")
-	FORCEINLINE void SetSubPathString(TStringOverload<FWideString> InSubPathString)
+	UE_FORCEINLINE_HINT void SetSubPathString(TStringOverload<FWideString> InSubPathString)
 	{
 		SubPathString = FUtf8String(InSubPathString.MoveTemp());
 	}
-	FORCEINLINE void SetSubPathString(TStringOverload<FUtf8String> InSubPathString)
+	UE_FORCEINLINE_HINT void SetSubPathString(TStringOverload<FUtf8String> InSubPathString)
 	{
 		SubPathString = InSubPathString.MoveTemp();
 	}
@@ -340,25 +340,25 @@ public:
 	}
 	
 	/** Check if this could possibly refer to a real object, or was initialized to null */
-	[[nodiscard]] FORCEINLINE bool IsValid() const
+	[[nodiscard]] UE_FORCEINLINE_HINT bool IsValid() const
 	{
 		return AssetPath.IsValid();
 	}
 
 	/** Checks to see if this is initialized to null */
-	[[nodiscard]] FORCEINLINE bool IsNull() const
+	[[nodiscard]] UE_FORCEINLINE_HINT bool IsNull() const
 	{
 		return AssetPath.IsNull();
 	}
 
 	/** Check if this represents an asset, meaning it is not null but does not have a sub path */
-	[[nodiscard]] FORCEINLINE bool IsAsset() const
+	[[nodiscard]] UE_FORCEINLINE_HINT bool IsAsset() const
 	{
 		return !AssetPath.IsNull() && SubPathString.IsEmpty();
 	}
 
 	/** Check if this represents a sub object, meaning it has a sub path */
-	[[nodiscard]] FORCEINLINE bool IsSubobject() const
+	[[nodiscard]] UE_FORCEINLINE_HINT bool IsSubobject() const
 	{
 		return !AssetPath.IsNull() && !SubPathString.IsEmpty();
 	}
@@ -438,12 +438,12 @@ public:
 	}
 
 	[[nodiscard]] static COREUOBJECT_API FSoftObjectPath GetOrCreateIDForObject(FObjectPtr Object);
-	[[nodiscard]] static FORCEINLINE FSoftObjectPath GetOrCreateIDForObject(const UObject* Object)
+	[[nodiscard]] static UE_FORCEINLINE_HINT FSoftObjectPath GetOrCreateIDForObject(const UObject* Object)
 	{
 		return GetOrCreateIDForObject(FObjectPtr(const_cast<UObject*>(Object)));
 	}
 	template <typename T>
-	[[nodiscard]] static FORCEINLINE FSoftObjectPath GetOrCreateIDForObject(TObjectPtr<T> Object)
+	[[nodiscard]] static UE_FORCEINLINE_HINT FSoftObjectPath GetOrCreateIDForObject(TObjectPtr<T> Object)
 	{
 		// This needs to be a template instead of TObjectPtr<const UObject> because C++ does derived-to-base
 		// pointer conversions ('standard conversion sequences') in more cases than TSmartPtr<Derived>-to-TSmartPtr<Base>
diff --git a/Engine/Source/Runtime/CoreUObject/Public/UObject/SoftObjectPtr.h b/Engine/Source/Runtime/CoreUObject/Public/UObject/SoftObjectPtr.h
--- a/Engine/Source/Runtime/CoreUObject/Public/UObject/SoftObjectPtr.h
+++ b/Engine/Source/Runtime/CoreUObject/Public/UObject/SoftObjectPtr.h
@@ -47,28 +47,28 @@ private:
 	using Super = TPersistentObjectPtr<FSoftObjectPath>;
 
 public:	
-	[[nodiscard]] FORCEINLINE FSoftObjectPtr() = default;
-	[[nodiscard]] FORCEINLINE FSoftObjectPtr(const FSoftObjectPtr& Other) = default;
-	[[nodiscard]] FORCEINLINE FSoftObjectPtr(FSoftObjectPtr&& Other) = default;
-	FORCEINLINE ~FSoftObjectPtr() = default;
-	FORCEINLINE FSoftObjectPtr& operator=(const FSoftObjectPtr& Other) = default;
-	FORCEINLINE FSoftObjectPtr& operator=(FSoftObjectPtr&& Other) = default;
-
-	[[nodiscard]] explicit FORCEINLINE FSoftObjectPtr(const FSoftObjectPath& ObjectPath)
+	[[nodiscard]] UE_FORCEINLINE_HINT FSoftObjectPtr() = default;
+	[[nodiscard]] UE_FORCEINLINE_HINT FSoftObjectPtr(const FSoftObjectPtr& Other) = default;
+	[[nodiscard]] UE_FORCEINLINE_HINT FSoftObjectPtr(FSoftObjectPtr&& Other) = default;
+	UE_FORCEINLINE_HINT ~FSoftObjectPtr() = default;
+	UE_FORCEINLINE_HINT FSoftObjectPtr& operator=(const FSoftObjectPtr& Other) = default;
+	UE_FORCEINLINE_HINT FSoftObjectPtr& operator=(FSoftObjectPtr&& Other) = default;
+
+	[[nodiscard]] explicit UE_FORCEINLINE_HINT FSoftObjectPtr(const FSoftObjectPath& ObjectPath)
 		: Super(ObjectPath)
 	{
 	}
 
-	[[nodiscard]] explicit FORCEINLINE FSoftObjectPtr(FObjectPtr Object)
+	[[nodiscard]] explicit UE_FORCEINLINE_HINT FSoftObjectPtr(FObjectPtr Object)
 	{
 		(*this)=Object;
 	}
-	[[nodiscard]] explicit FORCEINLINE FSoftObjectPtr(const UObject* Object)
+	[[nodiscard]] explicit UE_FORCEINLINE_HINT FSoftObjectPtr(const UObject* Object)
 		: FSoftObjectPtr(FObjectPtr(const_cast<UObject*>(Object)))
 	{
 	}
 	template <typename T>
-	[[nodiscard]] explicit FORCEINLINE FSoftObjectPtr(TObjectPtr<T> Object)
+	[[nodiscard]] explicit UE_FORCEINLINE_HINT FSoftObjectPtr(TObjectPtr<T> Object)
 		: FSoftObjectPtr(FObjectPtr(Object))
 	{
 		// This needs to be a template instead of TObjectPtr<const UObject> because C++ does derived-to-base
@@ -93,31 +93,31 @@ public:
 	}
 
 	/** Returns the StringObjectPath that is wrapped by this SoftObjectPtr */
-	[[nodiscard]] FORCEINLINE const FSoftObjectPath& ToSoftObjectPath() const
+	[[nodiscard]] UE_FORCEINLINE_HINT const FSoftObjectPath& ToSoftObjectPath() const
 	{
 		return GetUniqueID();
 	}
 
 	/** Returns string representation of reference, in form /package/path.assetname */
-	[[nodiscard]] FORCEINLINE FString ToString() const
+	[[nodiscard]] UE_FORCEINLINE_HINT FString ToString() const
 	{
 		return ToSoftObjectPath().ToString();
 	}
 
 	/** Returns /package/path string, leaving off the asset name */
-	[[nodiscard]] FORCEINLINE FString GetLongPackageName() const
+	[[nodiscard]] UE_FORCEINLINE_HINT FString GetLongPackageName() const
 	{
 		return ToSoftObjectPath().GetLongPackageName();
 	}
 	
 	/** Returns /package/path name, leaving off the asset name */
-	[[nodiscard]] FORCEINLINE FName GetLongPackageFName() const
+	[[nodiscard]] UE_FORCEINLINE_HINT FName GetLongPackageFName() const
 	{
 		return ToSoftObjectPath().GetLongPackageFName();
 	}
 
 	/** Returns assetname string, leaving off the /package/path. part */
-	[[nodiscard]] FORCEINLINE FString GetAssetName() const
+	[[nodiscard]] UE_FORCEINLINE_HINT FString GetAssetName() const
 	{
 		return ToSoftObjectPath().GetAssetName();
 	}
@@ -145,12 +145,12 @@ public:
 		(*SuperThis) = Ptr;
 		return *this;
 	}
-	FORCEINLINE FSoftObjectPtr& operator=(const UObject* Ptr)
+	UE_FORCEINLINE_HINT FSoftObjectPtr& operator=(const UObject* Ptr)
 	{
 		return *this = FObjectPtr(const_cast<UObject*>(Ptr));
 	}
 	template <typename T>
-	FORCEINLINE FSoftObjectPtr& operator=(TObjectPtr<T> Ptr)
+	UE_FORCEINLINE_HINT FSoftObjectPtr& operator=(TObjectPtr<T> Ptr)
 	{
 		// This needs to be a template instead of TObjectPtr<const UObject> because C++ does derived-to-base
 		// pointer conversions ('standard conversion sequences') in situations that TSmartPtr<Derived>-to-TSmartPtr<Base>
@@ -178,19 +178,19 @@ struct TSoftObjectPtr
 public:
 	using ElementType = T;
 	
-	[[nodiscard]] FORCEINLINE TSoftObjectPtr() = default;
-	[[nodiscard]] FORCEINLINE TSoftObjectPtr(const TSoftObjectPtr& Other) = default;
-	[[nodiscard]] FORCEINLINE TSoftObjectPtr(TSoftObjectPtr&& Other) = default;
-	FORCEINLINE ~TSoftObjectPtr() = default;
-	FORCEINLINE TSoftObjectPtr& operator=(const TSoftObjectPtr& Other) = default;
-	FORCEINLINE TSoftObjectPtr& operator=(TSoftObjectPtr&& Other) = default;
+	[[nodiscard]] UE_FORCEINLINE_HINT TSoftObjectPtr() = default;
+	[[nodiscard]] UE_FORCEINLINE_HINT TSoftObjectPtr(const TSoftObjectPtr& Other) = default;
+	[[nodiscard]] UE_FORCEINLINE_HINT TSoftObjectPtr(TSoftObjectPtr&& Other) = default;
+	UE_FORCEINLINE_HINT ~TSoftObjectPtr() = default;
+	UE_FORCEINLINE_HINT TSoftObjectPtr& operator=(const TSoftObjectPtr& Other) = default;
+	UE_FORCEINLINE_HINT TSoftObjectPtr& operator=(TSoftObjectPtr&& Other) = default;
 	
 	/** Construct from another soft pointer */
 	template <
 		class U
 		UE_REQUIRES(std::is_convertible_v<U*, T*>)
 	>
-	[[nodiscard]] FORCEINLINE TSoftObjectPtr(const TSoftObjectPtr<U>& Other)
+	[[nodiscard]] UE_FORCEINLINE_HINT TSoftObjectPtr(const TSoftObjectPtr<U>& Other)
 		: SoftObjectPtr(Other.SoftObjectPtr)
 	{
 	}
@@ -199,7 +199,7 @@ public:
 		UE_REQUIRES(!std::is_convertible_v<U*, T*>)
 	>
 	UE_SOFTOBJECTPTR_CONVERSION_DEPRECATED(5.5, "Constructing TSoftObjectPtr from an incompatible pointer type has been deprecated.")
-	[[nodiscard]] FORCEINLINE TSoftObjectPtr(const TSoftObjectPtr<U>& Other)
+	[[nodiscard]] UE_FORCEINLINE_HINT TSoftObjectPtr(const TSoftObjectPtr<U>& Other)
 		: SoftObjectPtr(Other.SoftObjectPtr)
 	{
 	}
@@ -209,7 +209,7 @@ public:
 		class U
 		UE_REQUIRES(std::is_convertible_v<U*, T*>)
 	>
-	[[nodiscard]] FORCEINLINE TSoftObjectPtr(TSoftObjectPtr<U>&& Other)
+	[[nodiscard]] UE_FORCEINLINE_HINT TSoftObjectPtr(TSoftObjectPtr<U>&& Other)
 		: SoftObjectPtr(MoveTemp(Other.SoftObjectPtr))
 	{
 	}
@@ -218,7 +218,7 @@ public:
 		UE_REQUIRES(!std::is_convertible_v<U*, T*>)
 	>
 	UE_SOFTOBJECTPTR_CONVERSION_DEPRECATED(5.5, "Constructing TSoftObjectPtr from an incompatible pointer type has been deprecated.")
-	[[nodiscard]] FORCEINLINE TSoftObjectPtr(TSoftObjectPtr<U>&& Other)
+	[[nodiscard]] UE_FORCEINLINE_HINT TSoftObjectPtr(TSoftObjectPtr<U>&& Other)
 		: SoftObjectPtr(MoveTemp(Other.SoftObjectPtr))
 	{
 	}
@@ -228,7 +228,7 @@ public:
 		typename U
 		UE_REQUIRES(std::is_convertible_v<U*, T*>)
 	>
-	[[nodiscard]] FORCEINLINE TSoftObjectPtr(U* Object)
+	[[nodiscard]] UE_FORCEINLINE_HINT TSoftObjectPtr(U* Object)
 		: SoftObjectPtr(Object)
 	{
 	}
@@ -237,7 +237,7 @@ public:
 		UE_REQUIRES(!std::is_convertible_v<U*, T*>)
 	>
 	UE_SOFTOBJECTPTR_CONVERSION_DEPRECATED(5.5, "Constructing TSoftObjectPtr from an incompatible pointer type has been deprecated.")
-	[[nodiscard]] FORCEINLINE TSoftObjectPtr(U* Object)
+	[[nodiscard]] UE_FORCEINLINE_HINT TSoftObjectPtr(U* Object)
 		: SoftObjectPtr(Object)
 	{
 	}
@@ -248,7 +248,7 @@ public:
 		typename U
 		UE_REQUIRES(std::is_convertible_v<U*, T*>)
 	>
-	[[nodiscard]] FORCEINLINE TSoftObjectPtr(TNotNull<U*> Object)
+	[[nodiscard]] UE_FORCEINLINE_HINT TSoftObjectPtr(TNotNull<U*> Object)
 		: SoftObjectPtr(Object)
 	{
 	}
@@ -257,7 +257,7 @@ public:
 		typename U
 		UE_REQUIRES(std::is_convertible_v<U*, T*>)
 	>
-	[[nodiscard]] FORCEINLINE TSoftObjectPtr(TNotNull<TObjectPtr<U>> Object)
+	[[nodiscard]] UE_FORCEINLINE_HINT TSoftObjectPtr(TNotNull<TObjectPtr<U>> Object)
 		: SoftObjectPtr(static_cast<TObjectPtr<U>>(Object).Get())
 	{
 	}
@@ -268,7 +268,7 @@ public:
 		typename U
 		UE_REQUIRES(std::is_convertible_v<U*, T*>)
 	>
-	[[nodiscard]] FORCEINLINE TSoftObjectPtr(const TObjectPtr<U> Object)
+	[[nodiscard]] UE_FORCEINLINE_HINT TSoftObjectPtr(const TObjectPtr<U> Object)
 		: SoftObjectPtr(Object.Get())
 	{
 	}
@@ -277,13 +277,13 @@ public:
 		UE_REQUIRES(!std::is_convertible_v<U*, T*>)
 	>
 	UE_SOFTOBJECTPTR_CONVERSION_DEPRECATED(5.5, "Constructing TSoftObjectPtr from an incompatible pointer type has been deprecated.")
-	[[nodiscard]] FORCEINLINE TSoftObjectPtr(const TObjectPtr<U> Object)
+	[[nodiscard]] UE_FORCEINLINE_HINT TSoftObjectPtr(const TObjectPtr<U> Object)
 		: SoftObjectPtr(Object.Get())
 	{
 	}
 
 	/** Construct from a nullptr */
-	[[nodiscard]] FORCEINLINE TSoftObjectPtr(TYPE_OF_NULLPTR)
+	[[nodiscard]] UE_FORCEINLINE_HINT TSoftObjectPtr(TYPE_OF_NULLPTR)
 		: SoftObjectPtr(nullptr)
 	{
 	}
@@ -293,7 +293,7 @@ public:
 		typename SoftObjectPathType
 		UE_REQUIRES(std::is_same_v<SoftObjectPathType, FSoftObjectPath>)
 	>
-	[[nodiscard]] explicit FORCEINLINE TSoftObjectPtr(SoftObjectPathType ObjectPath)
+	[[nodiscard]] explicit UE_FORCEINLINE_HINT TSoftObjectPtr(SoftObjectPathType ObjectPath)
 		: SoftObjectPtr(MoveTemp(ObjectPath))
 	{
 		// The reason for the strange templated/constrained signature is to prevent this type hole:
@@ -310,13 +310,13 @@ public:
 	}
 
 	/** Reset the soft pointer back to the null state */
-	FORCEINLINE void Reset()
+	UE_FORCEINLINE_HINT void Reset()
 	{
 		SoftObjectPtr.Reset();
 	}
 
 	/** Resets the weak ptr only, call this when ObjectId may change */
-	FORCEINLINE void ResetWeakPtr()
+	UE_FORCEINLINE_HINT void ResetWeakPtr()
 	{
 		SoftObjectPtr.ResetWeakPtr();
 	}
@@ -448,7 +448,7 @@ public:
 	 *
 	 * @param Other soft pointer to compare to
 	 */
-	[[nodiscard]] FORCEINLINE bool operator==(const TSoftObjectPtr& Rhs) const
+	[[nodiscard]] UE_FORCEINLINE_HINT bool operator==(const TSoftObjectPtr& Rhs) const
 	{
 		return SoftObjectPtr == Rhs.SoftObjectPtr;
 	}
@@ -456,7 +456,7 @@ public:
 		typename U
 		UE_REQUIRES(TModels_V<CEqualityComparableWith, T*, U*>)
 	>
-	[[nodiscard]] FORCEINLINE bool operator==(const TSoftObjectPtr<U>& Rhs) const
+	[[nodiscard]] UE_FORCEINLINE_HINT bool operator==(const TSoftObjectPtr<U>& Rhs) const
 	{
 		return SoftObjectPtr == Rhs.SoftObjectPtr;
 	}
@@ -465,7 +465,7 @@ public:
 		UE_REQUIRES(!TModels_V<CEqualityComparableWith, T*, U*>)
 	>
 	UE_SOFTOBJECTPTR_CONVERSION_DEPRECATED(5.5, "Comparing TSoftObjectPtrs of incompatible pointer types has been deprecated.")
-	[[nodiscard]] FORCEINLINE bool operator==(const TSoftObjectPtr<U>& Rhs) const
+	[[nodiscard]] UE_FORCEINLINE_HINT bool operator==(const TSoftObjectPtr<U>& Rhs) const
 	{
 		return SoftObjectPtr == Rhs.SoftObjectPtr;
 	}
@@ -473,7 +473,7 @@ public:
 		typename U
 		UE_REQUIRES(TModels_V<CEqualityComparableWith, T*, U*>)
 	>
-	[[nodiscard]] FORCEINLINE bool operator==(U* Rhs) const
+	[[nodiscard]] UE_FORCEINLINE_HINT bool operator==(U* Rhs) const
 	{
 		return SoftObjectPtr == TSoftObjectPtr<U>(Rhs).SoftObjectPtr;
 	}
@@ -482,11 +482,11 @@ public:
 		UE_REQUIRES(!TModels_V<CEqualityComparableWith, T*, U*>)
 	>
 	UE_SOFTOBJECTPTR_CONVERSION_DEPRECATED(5.5, "Comparing TSoftObjectPtrs of incompatible pointer types has been deprecated.")
-	[[nodiscard]] FORCEINLINE bool operator==(U* Rhs) const
+	[[nodiscard]] UE_FORCEINLINE_HINT bool operator==(U* Rhs) const
 	{
 		return SoftObjectPtr == TSoftObjectPtr<U>(Rhs).SoftObjectPtr;
 	}
-	[[nodiscard]] FORCEINLINE bool operator==(TYPE_OF_NULLPTR) const
+	[[nodiscard]] UE_FORCEINLINE_HINT bool operator==(TYPE_OF_NULLPTR) const
 	{
 		return SoftObjectPtr == nullptr;
 	}
@@ -499,13 +499,13 @@ public:
 	[[nodiscard]] T* Get() const;
 
 	/** Dereference the soft pointer */
-	[[nodiscard]] FORCEINLINE T& operator*() const
+	[[nodiscard]] UE_FORCEINLINE_HINT T& operator*() const
 	{
 		return *Get();
 	}
 
 	/** Dereference the soft pointer */
-	[[nodiscard]] FORCEINLINE T* operator->() const
+	[[nodiscard]] UE_FORCEINLINE_HINT T* operator->() const
 	{
 		return Get();
 	}
@@ -535,7 +535,7 @@ public:
 	 *
 	 * @return true if Get() would return a valid non-null pointer
 	 */
-	[[nodiscard]] FORCEINLINE bool IsValid() const
+	[[nodiscard]] UE_FORCEINLINE_HINT bool IsValid() const
 	{
 		// This does the runtime type check
 		return Get() != nullptr;
@@ -546,7 +546,7 @@ public:
 	 *
 	 * @return true if this does not point to a real object, but could possibly
 	 */
-	[[nodiscard]] FORCEINLINE bool IsPending() const
+	[[nodiscard]] UE_FORCEINLINE_HINT bool IsPending() const
 	{
 		return SoftObjectPtr.IsPending();
 	}
@@ -556,60 +556,60 @@ public:
 	 *
 	 * @return true if this is explicitly pointing to no object
 	 */
-	[[nodiscard]] FORCEINLINE bool IsNull() const
+	[[nodiscard]] UE_FORCEINLINE_HINT bool IsNull() const
 	{
 		return SoftObjectPtr.IsNull();
 	}
 
 	/** Returns the StringObjectPath that is wrapped by this SoftObjectPtr */
-	[[nodiscard]] FORCEINLINE const FSoftObjectPath& GetUniqueID() const
+	[[nodiscard]] UE_FORCEINLINE_HINT const FSoftObjectPath& GetUniqueID() const
 	{
 		return SoftObjectPtr.GetUniqueID();
 	}
 
 	/** Returns the StringObjectPath that is wrapped by this SoftObjectPtr */
-	[[nodiscard]] FORCEINLINE const FSoftObjectPath& ToSoftObjectPath() const
+	[[nodiscard]] UE_FORCEINLINE_HINT const FSoftObjectPath& ToSoftObjectPath() const
 	{
 		return SoftObjectPtr.GetUniqueID();
 	}
 
 	/** Returns string representation of reference, in form /package/path.assetname */
-	[[nodiscard]] FORCEINLINE FString ToString() const
+	[[nodiscard]] UE_FORCEINLINE_HINT FString ToString() const
 	{
 		return ToSoftObjectPath().ToString();
 	}
 
 	/** Returns /package/path string, leaving off the asset name */
-	[[nodiscard]] FORCEINLINE FString GetLongPackageName() const
+	[[nodiscard]] UE_FORCEINLINE_HINT FString GetLongPackageName() const
 	{
 		return ToSoftObjectPath().GetLongPackageName();
 	}
 	
 	/** Returns /package/path name, leaving off the asset name */
-	[[nodiscard]] FORCEINLINE FName GetLongPackageFName() const
+	[[nodiscard]] UE_FORCEINLINE_HINT FName GetLongPackageFName() const
 	{
 		return ToSoftObjectPath().GetLongPackageFName();
 	}
 
 	/** Returns assetname string, leaving off the /package/path part */
-	[[nodiscard]] FORCEINLINE FString GetAssetName() const
+	[[nodiscard]] UE_FORCEINLINE_HINT FString GetAssetName() const
 	{
 		return ToSoftObjectPath().GetAssetName();
 	}
 
 	/** Dereference soft pointer to see if it points somewhere valid */
-	[[nodiscard]] FORCEINLINE explicit operator bool() const
+	[[nodiscard]] UE_FORCEINLINE_HINT explicit operator bool() const
 	{
 		return IsValid();
 	}
 
 	/** Hash function */
-	[[nodiscard]] FORCEINLINE uint32 GetPtrTypeHash() const
+	[[nodiscard]] UE_FORCEINLINE_HINT uint32 GetPtrTypeHash() const
 	{
 		return GetTypeHash(static_cast<const TPersistentObjectPtr<FSoftObjectPath>&>(SoftObjectPtr));
 	}
 
-	FORCEINLINE void Serialize(FArchive& Ar)
+	UE_FORCEINLINE_HINT void Serialize(FArchive& Ar)
 	{
 		Ar << SoftObjectPtr;
 	}
@@ -627,7 +627,7 @@ template <
 	typename RhsType
 	UE_REQUIRES(TModels_V<CEqualityComparableWith, LhsType*, RhsType*>)
 >
-[[nodiscard]] FORCEINLINE bool operator==(LhsType* Lhs, const TSoftObjectPtr<RhsType>& Rhs)
+[[nodiscard]] UE_FORCEINLINE_HINT bool operator==(LhsType* Lhs, const TSoftObjectPtr<RhsType>& Rhs)
 {
 	return Rhs == Lhs;
 }
@@ -637,12 +637,12 @@ template <
 	UE_REQUIRES(!TModels_V<CEqualityComparableWith, LhsType*, RhsType*>)
 >
 UE_SOFTOBJECTPTR_CONVERSION_DEPRECATED(5.5, "Comparing TSoftObjectPtrs of incompatible pointer types has been deprecated.")
-[[nodiscard]] FORCEINLINE bool operator==(LhsType* Lhs, const TSoftObjectPtr<RhsType>& Rhs)
+[[nodiscard]] UE_FORCEINLINE_HINT bool operator==(LhsType* Lhs, const TSoftObjectPtr<RhsType>& Rhs)
 {
 	return Rhs == Lhs;
 }
 template <typename RhsType>
-[[nodiscard]] FORCEINLINE bool operator==(TYPE_OF_NULLPTR Lhs, const TSoftObjectPtr<RhsType>& Rhs)
+[[nodiscard]] UE_FORCEINLINE_HINT bool operator==(TYPE_OF_NULLPTR Lhs, const TSoftObjectPtr<RhsType>& Rhs)
 {
 	return Rhs == Lhs;
 }
@@ -651,7 +651,7 @@ template <
 	typename RhsType
 	UE_REQUIRES(TModels_V<CEqualityComparableWith, LhsType*, RhsType*>)
 >
-[[nodiscard]] FORCEINLINE bool operator!=(const TSoftObjectPtr<LhsType>& Lhs, const TSoftObjectPtr<RhsType>& Rhs)
+[[nodiscard]] UE_FORCEINLINE_HINT bool operator!=(const TSoftObjectPtr<LhsType>& Lhs, const TSoftObjectPtr<RhsType>& Rhs)
 {
 	return !(Lhs == Rhs);
 }
@@ -661,7 +661,7 @@ template <
 	UE_REQUIRES(!TModels_V<CEqualityComparableWith, LhsType*, RhsType*>)
 >
 UE_SOFTOBJECTPTR_CONVERSION_DEPRECATED(5.5, "Comparing TSoftObjectPtrs of incompatible pointer types has been deprecated.")
-[[nodiscard]] FORCEINLINE bool operator!=(const TSoftObjectPtr<LhsType>& Lhs, const TSoftObjectPtr<RhsType>& Rhs)
+[[nodiscard]] UE_FORCEINLINE_HINT bool operator!=(const TSoftObjectPtr<LhsType>& Lhs, const TSoftObjectPtr<RhsType>& Rhs)
 {
 	return !(Lhs == Rhs);
 }
@@ -670,7 +670,7 @@ template <
 	typename RhsType
 	UE_REQUIRES(TModels_V<CEqualityComparableWith, LhsType*, RhsType*>)
 >
-[[nodiscard]] FORCEINLINE bool operator!=(const TSoftObjectPtr<LhsType>& Lhs, RhsType* Rhs)
+[[nodiscard]] UE_FORCEINLINE_HINT bool operator!=(const TSoftObjectPtr<LhsType>& Lhs, RhsType* Rhs)
 {
 	return !(Lhs == Rhs);
 }
@@ -680,12 +680,12 @@ template <
 	UE_REQUIRES(!TModels_V<CEqualityComparableWith, LhsType*, RhsType*>)
 >
 UE_SOFTOBJECTPTR_CONVERSION_DEPRECATED(5.5, "Comparing TSoftObjectPtrs of incompatible pointer types has been deprecated.")
-[[nodiscard]] FORCEINLINE bool operator!=(const TSoftObjectPtr<LhsType>& Lhs, RhsType* Rhs)
+[[nodiscard]] UE_FORCEINLINE_HINT bool operator!=(const TSoftObjectPtr<LhsType>& Lhs, RhsType* Rhs)
 {
 	return !(Lhs == Rhs);
 }
 template <typename LhsType>
-[[nodiscard]] FORCEINLINE bool operator!=(const TSoftObjectPtr<LhsType>& Lhs, TYPE_OF_NULLPTR Rhs)
+[[nodiscard]] UE_FORCEINLINE_HINT bool operator!=(const TSoftObjectPtr<LhsType>& Lhs, TYPE_OF_NULLPTR Rhs)
 {
 	return !(Lhs == Rhs);
 }
@@ -694,7 +694,7 @@ template <
 	typename RhsType
 	UE_REQUIRES(TModels_V<CEqualityComparableWith, LhsType*, RhsType*>)
 >
-[[nodiscard]] FORCEINLINE bool operator!=(LhsType* Lhs, const TSoftObjectPtr<RhsType>& Rhs)
+[[nodiscard]] UE_FORCEINLINE_HINT bool operator!=(LhsType* Lhs, const TSoftObjectPtr<RhsType>& Rhs)
 {
 	return !(Lhs == Rhs);
 }
@@ -704,12 +704,12 @@ template <
 	UE_REQUIRES(!TModels_V<CEqualityComparableWith, LhsType*, RhsType*>)
 >
 UE_SOFTOBJECTPTR_CONVERSION_DEPRECATED(5.5, "Comparing TSoftObjectPtrs of incompatible pointer types has been deprecated.")
-[[nodiscard]] FORCEINLINE bool operator!=(LhsType* Lhs, const TSoftObjectPtr<RhsType>& Rhs)
+[[nodiscard]] UE_FORCEINLINE_HINT bool operator!=(LhsType* Lhs, const TSoftObjectPtr<RhsType>& Rhs)
 {
 	return !(Lhs == Rhs);
 }
 template <typename RhsType>
-[[nodiscard]] FORCEINLINE bool operator!=(TYPE_OF_NULLPTR Lhs, const TSoftObjectPtr<RhsType>& Rhs)
+[[nodiscard]] UE_FORCEINLINE_HINT bool operator!=(TYPE_OF_NULLPTR Lhs, const TSoftObjectPtr<RhsType>& Rhs)
 {
 	return !(Lhs == Rhs);
 }
@@ -717,7 +717,7 @@ template <typename RhsType>
 
 /** Hash function */
 template<class T>
-[[nodiscard]] FORCEINLINE uint32 GetTypeHash(const TSoftObjectPtr<T>& Ptr)
+[[nodiscard]] UE_FORCEINLINE_HINT uint32 GetTypeHash(const TSoftObjectPtr<T>& Ptr)
 {
 	return Ptr.GetPtrTypeHash();
 }
@@ -767,55 +767,55 @@ class TSoftClassPtr
 public:
 	using ElementType = TClass;
 	
-	[[nodiscard]] FORCEINLINE TSoftClassPtr() = default;
-	[[nodiscard]] FORCEINLINE TSoftClassPtr(const TSoftClassPtr& Other) = default;
-	[[nodiscard]] FORCEINLINE TSoftClassPtr(TSoftClassPtr&& Other) = default;
-	FORCEINLINE ~TSoftClassPtr() = default;
-	FORCEINLINE TSoftClassPtr& operator=(const TSoftClassPtr& Other) = default;
-	FORCEINLINE TSoftClassPtr& operator=(TSoftClassPtr&& Other) = default;
+	[[nodiscard]] UE_FORCEINLINE_HINT TSoftClassPtr() = default;
+	[[nodiscard]] UE_FORCEINLINE_HINT TSoftClassPtr(const TSoftClassPtr& Other) = default;
+	[[nodiscard]] UE_FORCEINLINE_HINT TSoftClassPtr(TSoftClassPtr&& Other) = default;
+	UE_FORCEINLINE_HINT ~TSoftClassPtr() = default;
+	UE_FORCEINLINE_HINT TSoftClassPtr& operator=(const TSoftClassPtr& Other) = default;
+	UE_FORCEINLINE_HINT TSoftClassPtr& operator=(TSoftClassPtr&& Other) = default;
 		
 	/** Construct from another soft pointer */
 	template <
 		class TClassA
 		UE_REQUIRES(std::is_convertible_v<TClassA*, TClass*>)
 	>
-	[[nodiscard]] FORCEINLINE TSoftClassPtr(const TSoftClassPtr<TClassA>& Other)
+	[[nodiscard]] UE_FORCEINLINE_HINT TSoftClassPtr(const TSoftClassPtr<TClassA>& Other)
 		: SoftObjectPtr(Other.SoftObjectPtr)
 	{
 	}
 
 	/** Construct from a class already in memory */
-	[[nodiscard]] FORCEINLINE TSoftClassPtr(const UClass* From)
+	[[nodiscard]] UE_FORCEINLINE_HINT TSoftClassPtr(const UClass* From)
 		: SoftObjectPtr(From)
 	{
 	}
 
 	/** Construct from a soft object path */
-	[[nodiscard]] explicit FORCEINLINE TSoftClassPtr(const FSoftObjectPath& ObjectPath)
+	[[nodiscard]] explicit UE_FORCEINLINE_HINT TSoftClassPtr(const FSoftObjectPath& ObjectPath)
 		: SoftObjectPtr(ObjectPath)
 	{
 	}
 
 	/** Reset the soft pointer back to the null state */
-	FORCEINLINE void Reset()
+	UE_FORCEINLINE_HINT void Reset()
 	{
 		SoftObjectPtr.Reset();
 	}
 
 	/** Resets the weak ptr only, call this when ObjectId may change */
-	FORCEINLINE void ResetWeakPtr()
+	UE_FORCEINLINE_HINT void ResetWeakPtr()
 	{
 		SoftObjectPtr.ResetWeakPtr();
 	}
 
 	/** Copy from a class already in memory */
-	FORCEINLINE void operator=(const UClass* From)
+	UE_FORCEINLINE_HINT void operator=(const UClass* From)
 	{
 		SoftObjectPtr = From;
 	}
 
 	/** Copy from a soft object path */
-	FORCEINLINE void operator=(const FSoftObjectPath& ObjectPath)
+	UE_FORCEINLINE_HINT void operator=(const FSoftObjectPath& ObjectPath)
 	{
 		SoftObjectPtr = ObjectPath;
 	}
@@ -848,7 +848,7 @@ public:
 	 *
 	 * @param Other soft pointer to compare to 
 	 */
-	[[nodiscard]] FORCEINLINE bool operator==(const TSoftClassPtr& Rhs) const
+	[[nodiscard]] UE_FORCEINLINE_HINT bool operator==(const TSoftClassPtr& Rhs) const
 	{
 		return SoftObjectPtr == Rhs.SoftObjectPtr;
 	}
@@ -859,7 +859,7 @@ public:
 	 *
 	 * @param Other soft pointer to compare to
 	 */
-	[[nodiscard]] FORCEINLINE bool operator!=(const TSoftClassPtr& Rhs) const
+	[[nodiscard]] UE_FORCEINLINE_HINT bool operator!=(const TSoftClassPtr& Rhs) const
 	{
 		return SoftObjectPtr != Rhs.SoftObjectPtr;
 	}
@@ -881,13 +881,13 @@ public:
 	}
 
 	/** Dereference the soft pointer */
-	[[nodiscard]] FORCEINLINE UClass& operator*() const
+	[[nodiscard]] UE_FORCEINLINE_HINT UClass& operator*() const
 	{
 		return *Get();
 	}
 
 	/** Dereference the soft pointer */
-	[[nodiscard]] FORCEINLINE UClass* operator->() const
+	[[nodiscard]] UE_FORCEINLINE_HINT UClass* operator->() const
 	{
 		return Get();
 	}
@@ -897,7 +897,7 @@ public:
 	 *
 	 * @return true if Get() would return a valid non-null pointer
 	 */
-	[[nodiscard]] FORCEINLINE bool IsValid() const
+	[[nodiscard]] UE_FORCEINLINE_HINT bool IsValid() const
 	{
 		// This also does the UClass type check
 		return Get() != nullptr;
@@ -908,7 +908,7 @@ public:
 	 *
 	 * @return true if this does not point to a real object, but could possibly
 	 */
-	[[nodiscard]] FORCEINLINE bool IsPending() const
+	[[nodiscard]] UE_FORCEINLINE_HINT bool IsPending() const
 	{
 		return SoftObjectPtr.IsPending();
 	}
@@ -918,55 +918,55 @@ public:
 	 *
 	 * @return true if this is explicitly pointing to no object
 	 */
-	[[nodiscard]] FORCEINLINE bool IsNull() const
+	[[nodiscard]] UE_FORCEINLINE_HINT bool IsNull() const
 	{
 		return SoftObjectPtr.IsNull();
 	}
 
 	/** Returns the StringObjectPath that is wrapped by this SoftObjectPtr */
-	[[nodiscard]] FORCEINLINE const FSoftObjectPath& GetUniqueID() const
+	[[nodiscard]] UE_FORCEINLINE_HINT const FSoftObjectPath& GetUniqueID() const
 	{
 		return SoftObjectPtr.GetUniqueID();
 	}
 
 	/** Returns the StringObjectPath that is wrapped by this SoftObjectPtr */
-	[[nodiscard]] FORCEINLINE const FSoftObjectPath& ToSoftObjectPath() const
+	[[nodiscard]] UE_FORCEINLINE_HINT const FSoftObjectPath& ToSoftObjectPath() const
 	{
 		return SoftObjectPtr.GetUniqueID();
 	}
 
 	/** Returns string representation of reference, in form /package/path.assetname  */
-	[[nodiscard]] FORCEINLINE FString ToString() const
+	[[nodiscard]] UE_FORCEINLINE_HINT FString ToString() const
 	{
 		return ToSoftObjectPath().ToString();
 	}
 
 	/** Returns /package/path string, leaving off the asset name */
-	[[nodiscard]] FORCEINLINE FString GetLongPackageName() const
+	[[nodiscard]] UE_FORCEINLINE_HINT FString GetLongPackageName() const
 	{
 		return ToSoftObjectPath().GetLongPackageName();
 	}
 	
 	/** Returns /package/path name, leaving off the asset name */
-	[[nodiscard]] FORCEINLINE FName GetLongPackageFName() const
+	[[nodiscard]] UE_FORCEINLINE_HINT FName GetLongPackageFName() const
 	{
 		return ToSoftObjectPath().GetLongPackageFName();
 	}
 
 	/** Returns assetname string, leaving off the /package/path part */
-	[[nodiscard]] FORCEINLINE FString GetAssetName() const
+	[[nodiscard]] UE_FORCEINLINE_HINT FString GetAssetName() const
 	{
 		return ToSoftObjectPath().GetAssetName();
 	}
 
 	/** Dereference soft pointer to see if it points somewhere valid */
-	[[nodiscard]] FORCEINLINE explicit operator bool() const
+	[[nodiscard]] UE_FORCEINLINE_HINT explicit operator bool() const
 	{
 		return IsValid();
 	}
 
 	/** Hash function */
-	[[nodiscard]] FORCEINLINE uint32 GetPtrTypeHash() const
+	[[nodiscard]] UE_FORCEINLINE_HINT uint32 GetPtrTypeHash() const
 	{
 		return GetTypeHash(static_cast<const TPersistentObjectPtr<FSoftObjectPath>&>(SoftObjectPtr));
 	}
@@ -1058,13 +1058,13 @@ T* TSoftObjectPtr<T>::Get() const
 
 #if !PLATFORM_COMPILER_HAS_GENERATED_COMPARISON_OPERATORS
 template<class TClass>
-[[nodiscard]] FORCEINLINE bool operator==(const UClass* Lhs, const TSoftClassPtr<TClass>& Rhs)
+[[nodiscard]] UE_FORCEINLINE_HINT bool operator==(const UClass* Lhs, const TSoftClassPtr<TClass>& Rhs)
 {
 	return Rhs == Lhs;
 }
 
 template<class TClass>
-[[nodiscard]] FORCEINLINE bool operator!=(const UClass* Lhs, const TSoftClassPtr<TClass>& Rhs)
+[[nodiscard]] UE_FORCEINLINE_HINT bool operator!=(const UClass* Lhs, const TSoftClassPtr<TClass>& Rhs)
 {
 	return Rhs != Lhs;
 }
@@ -1072,7 +1072,7 @@ template<class TClass>
 
 /** Hash function */
 template<class TClass>
-[[nodiscard]] FORCEINLINE uint32 GetTypeHash(const TSoftClassPtr<TClass>& Ptr)
+[[nodiscard]] UE_FORCEINLINE_HINT uint32 GetTypeHash(const TSoftClassPtr<TClass>& Ptr)
 {
 	return Ptr.GetPtrTypeHash();
 }
diff --git a/Engine/Source/Runtime/CoreUObject/Public/UObject/TopLevelAssetPath.h b/Engine/Source/Runtime/CoreUObject/Public/UObject/TopLevelAssetPath.h
--- a/Engine/Source/Runtime/CoreUObject/Public/UObject/TopLevelAssetPath.h
+++ b/Engine/Source/Runtime/CoreUObject/Public/UObject/TopLevelAssetPath.h
@@ -216,7 +216,7 @@ inline FUtf8StringBuilderBase& operator<<(FUtf8StringBuilderBase& Builder, const
 /** Fast non-alphabetical order that is only stable during this process' lifetime */
 struct FTopLevelAssetPathFastLess
 {
-	FORCEINLINE bool operator()(const FTopLevelAssetPath& A, const FTopLevelAssetPath& B) const
+	UE_FORCEINLINE_HINT bool operator()(const FTopLevelAssetPath& A, const FTopLevelAssetPath& B) const
 	{
 		return A.CompareFast(B) < 0;
 	}
@@ -225,7 +225,7 @@ struct FTopLevelAssetPathFastLess
 /** Slow alphabetical order that is stable / deterministic over process runs */
 struct FTopLevelAssetPathLexicalLess
 {
-	FORCEINLINE bool operator()(const FTopLevelAssetPath& A, const FTopLevelAssetPath& B) const
+	UE_FORCEINLINE_HINT bool operator()(const FTopLevelAssetPath& A, const FTopLevelAssetPath& B) const
 	{
 		return A.Compare(B) < 0;
 	}
diff --git a/Engine/Source/Runtime/CoreUObject/Public/UObject/UObjectAnnotation.h b/Engine/Source/Runtime/CoreUObject/Public/UObject/UObjectAnnotation.h
--- a/Engine/Source/Runtime/CoreUObject/Public/UObject/UObjectAnnotation.h
+++ b/Engine/Source/Runtime/CoreUObject/Public/UObject/UObjectAnnotation.h
@@ -540,7 +540,7 @@ struct FBoolAnnotation
 	 * Determine if this annotation
 	 * @return true is this is a default pair. We only check the linker because CheckInvariants rules out bogus combinations
 	 */
-	FORCEINLINE bool IsDefault()
+	UE_FORCEINLINE_HINT bool IsDefault()
 	{
 		return !Mark;
 	}
@@ -569,7 +569,7 @@ public:
 	 *
 	 * @param Object		Object to annotate.
 	 */
-	FORCEINLINE void Set(const UObjectBase *Object)
+	UE_FORCEINLINE_HINT void Set(const UObjectBase *Object)
 	{
 		this->AddAnnotation(Object,FBoolAnnotation(true));
 	}
@@ -578,7 +578,7 @@ public:
 	 *
 	 * @param Object		Object to de-annotate.
 	 */
-	FORCEINLINE void Clear(const UObjectBase *Object)
+	UE_FORCEINLINE_HINT void Clear(const UObjectBase *Object)
 	{
 		this->RemoveAnnotation(Object);
 	}
@@ -587,7 +587,7 @@ public:
 	 * Removes all bool annotations 
 	 *
 	 */
-	FORCEINLINE void ClearAll()
+	UE_FORCEINLINE_HINT void ClearAll()
 	{
 		this->RemoveAllAnnotations();
 	}
@@ -597,7 +597,7 @@ public:
 	 *
 	 * @param Object		Object to return the bool annotation for
 	 */
-	FORCEINLINE bool Get(const UObjectBase *Object)
+	UE_FORCEINLINE_HINT bool Get(const UObjectBase *Object)
 	{
 		return this->GetAnnotation(Object).Mark;
 	}
@@ -605,12 +605,12 @@ public:
 	/** 
 	 * Reserves memory for the annotation map for the specified number of elements, used to avoid reallocations. 
 	 */
-	FORCEINLINE void Reserve(int32 ExpectedNumElements)
+	UE_FORCEINLINE_HINT void Reserve(int32 ExpectedNumElements)
 	{
 		FUObjectAnnotationSparse<FBoolAnnotation,true>::Reserve(ExpectedNumElements);
 	}
 
-	FORCEINLINE int32 Num() const
+	UE_FORCEINLINE_HINT int32 Num() const
 	{
 		return this->GetAnnotationMap().Num();
 	}
@@ -957,7 +957,7 @@ public:
 	* Thread safe, but you know, someone might have added more elements before this even returns
 	* @return	the number of elements in the array
 	**/
-	FORCEINLINE int32 GetAnnotationCount() const
+	UE_FORCEINLINE_HINT int32 GetAnnotationCount() const
 	{
 		return NumAnnotations;
 	}
@@ -967,7 +967,7 @@ public:
 	* Thread safe, but you know, someone might have added more elements before this even returns
 	* @return	the maximum number of elements in the array
 	**/
-	FORCEINLINE int32 GetMaxAnnotations() const TSAN_SAFE
+	UE_FORCEINLINE_HINT int32 GetMaxAnnotations() const TSAN_SAFE
 	{
 		return MaxAnnotations;
 	}
@@ -978,7 +978,7 @@ public:
 	* @return	the maximum number of elements in the array
 	**/
 	UE_DEPRECATED(5.3, "Use GetMaxAnnotations instead")
-	FORCEINLINE int32 GetMaxAnnottations() const TSAN_SAFE
+	UE_FORCEINLINE_HINT int32 GetMaxAnnottations() const TSAN_SAFE
 	{
 		return MaxAnnotations;
 	}
@@ -989,7 +989,7 @@ public:
 	* @param	Index	Index to test
 	* @return	true, if this is a valid
 	**/
-	FORCEINLINE bool IsValidIndex(int32 Index) const
+	UE_FORCEINLINE_HINT bool IsValidIndex(int32 Index) const
 	{
 		return Index >= 0 && Index < MaxAnnotations;
 	}
@@ -1411,7 +1411,7 @@ public:
 	 * Removes all bool annotations 
 	 *
 	 */
-	FORCEINLINE void ClearAll()
+	UE_FORCEINLINE_HINT void ClearAll()
 	{
 		RemoveAllAnnotations();
 	}
diff --git a/Engine/Source/Runtime/CoreUObject/Public/UObject/UObjectArray.h b/Engine/Source/Runtime/CoreUObject/Public/UObject/UObjectArray.h
--- a/Engine/Source/Runtime/CoreUObject/Public/UObject/UObjectArray.h
+++ b/Engine/Source/Runtime/CoreUObject/Public/UObject/UObjectArray.h
@@ -169,18 +169,18 @@ public:
 #endif
 	}
 
-	FORCEINLINE void SetOwnerIndex(int32 OwnerIndex)
+	UE_FORCEINLINE_HINT void SetOwnerIndex(int32 OwnerIndex)
 	{
 		ClusterRootIndex = OwnerIndex;
 	}
 
-	FORCEINLINE int32 GetOwnerIndex() const
+	UE_FORCEINLINE_HINT int32 GetOwnerIndex() const
 	{
 		return ClusterRootIndex;
 	}
 
 	/** Encodes the cluster index in the ClusterRootIndex variable */
-	FORCEINLINE void SetClusterIndex(int32 ClusterIndex)
+	UE_FORCEINLINE_HINT void SetClusterIndex(int32 ClusterIndex)
 	{
 		ClusterRootIndex = -ClusterIndex - 1;
 	}
@@ -205,17 +205,17 @@ public:
 	}
 
 #if UE_WITH_REMOTE_OBJECT_HANDLE
-	FORCEINLINE void SetRemoteId(FRemoteObjectId Id)
+	UE_FORCEINLINE_HINT void SetRemoteId(FRemoteObjectId Id)
 	{
 		RemoteId = Id.GetLocalized();
 	}
-	FORCEINLINE FRemoteObjectId GetRemoteId() const
+	UE_FORCEINLINE_HINT FRemoteObjectId GetRemoteId() const
 	{
 		return RemoteId;
 	}
 #endif // UE_WITH_REMOTE_OBJECT_HANDLE
 
-	FORCEINLINE void SetFlags(EInternalObjectFlags FlagsToSet)
+	UE_FORCEINLINE_HINT void SetFlags(EInternalObjectFlags FlagsToSet)
 	{
 		ThisThreadAtomicallySetFlag(FlagsToSet);
 	}
@@ -229,7 +229,7 @@ public:
 #endif
 	}
 
-	FORCEINLINE void ClearFlags(EInternalObjectFlags FlagsToClear)
+	UE_FORCEINLINE_HINT void ClearFlags(EInternalObjectFlags FlagsToClear)
 	{
 		ThisThreadAtomicallyClearedFlag(FlagsToClear);
 	}
@@ -283,48 +283,48 @@ public:
 		return Result;
 	}
 
-	FORCEINLINE bool HasAnyFlags(EInternalObjectFlags InFlags) const
+	UE_FORCEINLINE_HINT bool HasAnyFlags(EInternalObjectFlags InFlags) const
 	{
 		return !!(GetFlagsInternal() & int32(InFlags));
 	}
 
-	FORCEINLINE bool HasAllFlags(EInternalObjectFlags InFlags) const
+	UE_FORCEINLINE_HINT bool HasAllFlags(EInternalObjectFlags InFlags) const
 	{
 		return (GetFlagsInternal() & int32(InFlags)) == int32(InFlags);
 	}
 
-	FORCEINLINE bool IsUnreachable() const
+	UE_FORCEINLINE_HINT bool IsUnreachable() const
 	{
 		return !!(GetFlagsInternal() & int32(EInternalObjectFlags::Unreachable));
 	}
 
-	FORCEINLINE void SetGarbage()
+	UE_FORCEINLINE_HINT void SetGarbage()
 	{
 		AtomicallySetFlag_ForGC(EInternalObjectFlags::Garbage);
 	}
-	FORCEINLINE void ClearGarbage()
+	UE_FORCEINLINE_HINT void ClearGarbage()
 	{
 		AtomicallyClearFlag_ForGC(EInternalObjectFlags::Garbage);
 	}
-	FORCEINLINE bool IsGarbage() const
+	UE_FORCEINLINE_HINT bool IsGarbage() const
 	{
 		return !!(GetFlagsInternal() & int32(EInternalObjectFlags::Garbage));
 	}
 
-	FORCEINLINE void SetRootSet()
+	UE_FORCEINLINE_HINT void SetRootSet()
 	{
 		ThisThreadAtomicallySetFlag(EInternalObjectFlags::RootSet);
 	}
-	FORCEINLINE void ClearRootSet()
+	UE_FORCEINLINE_HINT void ClearRootSet()
 	{
 		ThisThreadAtomicallyClearedFlag(EInternalObjectFlags::RootSet);
 	}
-	FORCEINLINE bool IsRootSet() const
+	UE_FORCEINLINE_HINT bool IsRootSet() const
 	{
 		return !!(GetFlagsInternal() & int32(EInternalObjectFlags::RootSet));
 	}
 
-	FORCEINLINE int32 GetRefCount() const
+	UE_FORCEINLINE_HINT int32 GetRefCount() const
 	{
 		return (int32)(FlagsAndRefCount & EInternalObjectFlags_RefCountMask);
 	}
@@ -389,22 +389,22 @@ public:
 #endif
 
 private:
-	FORCEINLINE static int32 GetFlagsInternal(int64 FlagsAndRefCount)
+	UE_FORCEINLINE_HINT static int32 GetFlagsInternal(int64 FlagsAndRefCount)
 	{
 		return (int32)(FlagsAndRefCount >> 32);
 	}
 
-	FORCEINLINE int32 GetFlagsInternal() const
+	UE_FORCEINLINE_HINT int32 GetFlagsInternal() const
 	{
 		return GetFlagsInternal(FPlatformAtomics::AtomicRead_Relaxed((int64*)&FlagsAndRefCount));
 	}
 	
-	FORCEINLINE static int32 GetRefCountInternal(int64 FlagsAndRefCount)
+	UE_FORCEINLINE_HINT static int32 GetRefCountInternal(int64 FlagsAndRefCount)
 	{
 		return (int32)(FlagsAndRefCount & EInternalObjectFlags_RefCountMask);
 	}
 
-	FORCEINLINE int32 GetRefCountInternal() const
+	UE_FORCEINLINE_HINT int32 GetRefCountInternal() const
 	{
 		return GetRefCountInternal(FPlatformAtomics::AtomicRead_Relaxed((int64*)&FlagsAndRefCount));
 	}
@@ -628,7 +628,7 @@ public:
 	* Thread safe, but you know, someone might have added more elements before this even returns
 	* @return	the number of elements in the array
 	**/
-	FORCEINLINE int32 Num() const
+	UE_FORCEINLINE_HINT int32 Num() const
 	{
 		return NumElements;
 	}
@@ -638,7 +638,7 @@ public:
 	* Thread safe, but you know, someone might have added more elements before this even returns
 	* @return	the maximum number of elements in the array
 	**/
-	FORCEINLINE int32 Capacity() const
+	UE_FORCEINLINE_HINT int32 Capacity() const
 	{
 		return MaxElements;
 	}
@@ -649,7 +649,7 @@ public:
 	* @param	Index	Index to test
 	* @return	true, if this is a valid
 	**/
-	FORCEINLINE bool IsValidIndex(int32 Index) const
+	UE_FORCEINLINE_HINT bool IsValidIndex(int32 Index) const
 	{
 		return Index < Num() && Index >= 0;
 	}
@@ -795,7 +795,7 @@ public:
 	* Thread safe, but you know, someone might have added more elements before this even returns
 	* @return	the number of elements in the array
 	**/
-	FORCEINLINE int32 Num() const
+	UE_FORCEINLINE_HINT int32 Num() const
 	{
 		return NumElements;
 	}
@@ -805,7 +805,7 @@ public:
 	* Thread safe, but you know, someone might have added more elements before this even returns
 	* @return	the maximum number of elements in the array
 	**/
-	FORCEINLINE int32 Capacity() const
+	UE_FORCEINLINE_HINT int32 Capacity() const
 	{
 		return MaxElements;
 	}
@@ -816,7 +816,7 @@ public:
 	* @param	Index	Index to test
 	* @return	true, if this is a valid
 	**/
-	FORCEINLINE bool IsValidIndex(int32 Index) const
+	UE_FORCEINLINE_HINT bool IsValidIndex(int32 Index) const
 	{
 		return Index < Num() && Index >= 0;
 	}
@@ -1052,7 +1052,7 @@ public:
 	 * @param Object object to get the index of
 	 * @return index of this object
 	 */
-	FORCEINLINE int32 ObjectToIndex(const class UObjectBase* Object) const
+	UE_FORCEINLINE_HINT int32 ObjectToIndex(const class UObjectBase* Object) const
 	{
 		return Object->InternalIndex;
 	}
@@ -1073,7 +1073,7 @@ public:
 		return nullptr;
 	}
 
-	FORCEINLINE FUObjectItem* IndexToObjectUnsafeForGC(int32 Index)
+	UE_FORCEINLINE_HINT FUObjectItem* IndexToObjectUnsafeForGC(int32 Index)
 	{
 		return const_cast<FUObjectItem*>(&ObjObjects[Index]);
 	}
@@ -1119,7 +1119,7 @@ public:
 		return IsValid(ObjectItem, bEvenIfGarbage);
 	}
 
-	FORCEINLINE bool IsStale(FUObjectItem* ObjectItem, bool bIncludingGarbage)
+	UE_FORCEINLINE_HINT bool IsStale(FUObjectItem* ObjectItem, bool bIncludingGarbage)
 	{
 		// This method assumes ObjectItem is valid.
 		return bIncludingGarbage ? (ObjectItem->HasAnyFlags(EInternalObjectFlags::Unreachable | EInternalObjectFlags::Garbage)) : (ObjectItem->IsUnreachable());
@@ -1137,7 +1137,7 @@ public:
 	}
 
 	/** Returns the index of the first object outside of the disregard for GC pool */
-	FORCEINLINE int32 GetFirstGCIndex() const
+	UE_FORCEINLINE_HINT int32 GetFirstGCIndex() const
 	{
 		return ObjFirstGCIndex;
 	}
@@ -1186,7 +1186,7 @@ public:
 	COREUOBJECT_API bool IsValid(const UObjectBase* Object) const;
 
 	/** Checks if the object index is valid. */
-	FORCEINLINE bool IsValidIndex(const UObjectBase* Object) const 
+	UE_FORCEINLINE_HINT bool IsValidIndex(const UObjectBase* Object) const 
 	{ 
 		return ObjObjects.IsValidIndex(Object->InternalIndex);
 	}
@@ -1197,7 +1197,7 @@ public:
 	 * @param Object object to get for disregard for GC
 	 * @return true if this object si disregard for GC
 	 */
-	FORCEINLINE bool IsDisregardForGC(const class UObjectBase* Object)
+	UE_FORCEINLINE_HINT bool IsDisregardForGC(const class UObjectBase* Object)
 	{
 		return Object->InternalIndex <= ObjLastNonGCIndex;
 	}
@@ -1206,7 +1206,7 @@ public:
 	 *
 	 * @return	the number of UObjects in the global array
 	 */
-	FORCEINLINE int32 GetObjectArrayNum() const 
+	UE_FORCEINLINE_HINT int32 GetObjectArrayNum() const 
 	{ 
 		return ObjObjects.Num();
 	}
@@ -1216,7 +1216,7 @@ public:
 	 *
 	 * @return	the number of UObjects in the global array
 	 */
-	FORCEINLINE int32 GetObjectArrayNumMinusPermanent() const 
+	UE_FORCEINLINE_HINT int32 GetObjectArrayNumMinusPermanent() const 
 	{ 
 		return ObjObjects.Num() - (ObjLastNonGCIndex + 1);
 	}
@@ -1226,7 +1226,7 @@ public:
 	 *
 	 * @return	the number of permanent objects
 	 */
-	FORCEINLINE int32 GetObjectArrayNumPermanent() const 
+	UE_FORCEINLINE_HINT int32 GetObjectArrayNumPermanent() const 
 	{ 
 		return ObjLastNonGCIndex + 1;
 	}
@@ -1359,7 +1359,7 @@ public:
 		/**
 		 * Iterator advance
 		 */
-		FORCEINLINE void operator++()
+		UE_FORCEINLINE_HINT void operator++()
 		{
 			Advance();
 		}
@@ -1368,17 +1368,17 @@ public:
 		bool operator!=(const TIterator& Rhs) const { return Index != Rhs.Index; }
 
 		/** Conversion to "bool" returning true if the iterator is valid. */
-		FORCEINLINE explicit operator bool() const
+		UE_FORCEINLINE_HINT explicit operator bool() const
 		{ 
 			return !!CurrentObject;
 		}
 		/** inverse of the "bool" operator */
-		FORCEINLINE bool operator !() const 
+		UE_FORCEINLINE_HINT bool operator !() const 
 		{
 			return !(bool)*this;
 		}
 
-		FORCEINLINE int32 GetIndex() const
+		UE_FORCEINLINE_HINT int32 GetIndex() const
 		{
 			return Index;
 		}
@@ -1390,7 +1390,7 @@ public:
 		 *
 		 * @return	the UObject at the iterator
 		 */
-		FORCEINLINE FUObjectItem* GetObject() const
+		UE_FORCEINLINE_HINT FUObjectItem* GetObject() const
 		{ 
 			return CurrentObject;
 		}
diff --git a/Engine/Source/Runtime/CoreUObject/Public/UObject/UObjectBaseUtility.h b/Engine/Source/Runtime/CoreUObject/Public/UObject/UObjectBaseUtility.h
--- a/Engine/Source/Runtime/CoreUObject/Public/UObject/UObjectBaseUtility.h
+++ b/Engine/Source/Runtime/CoreUObject/Public/UObject/UObjectBaseUtility.h
@@ -108,7 +108,7 @@ public:
 	 * @param Mask	Mask to mask object flags with
 	 * @param Objects flags that are set in both the object and the mask
 	 */
-	FORCEINLINE EObjectFlags GetMaskedFlags( EObjectFlags Mask = RF_AllFlags ) const
+	UE_FORCEINLINE_HINT EObjectFlags GetMaskedFlags( EObjectFlags Mask = RF_AllFlags ) const
 	{
 		return EObjectFlags(GetFlags() & Mask);
 	}
@@ -122,7 +122,7 @@ public:
 	 *
 	 * @param	Marks	Logical OR of OBJECTMARK_'s to apply 
 	 */
-	FORCEINLINE void Mark(EObjectMark Marks) const
+	UE_FORCEINLINE_HINT void Mark(EObjectMark Marks) const
 	{
 		MarkObject(this,Marks);
 	}
@@ -132,7 +132,7 @@ public:
 	 *
 	 * @param	Marks	Logical OR of OBJECTMARK_'s to remove 
 	 */
-	FORCEINLINE void UnMark(EObjectMark Marks) const
+	UE_FORCEINLINE_HINT void UnMark(EObjectMark Marks) const
 	{
 		UnMarkObject(this,Marks);
 	}
@@ -143,7 +143,7 @@ public:
 	 * @param	Marks	Logical OR of OBJECTMARK_'s to test
 	 * @return	true if the object has any of the given marks.
 	 */
-	FORCEINLINE bool HasAnyMarks(EObjectMark Marks) const
+	UE_FORCEINLINE_HINT bool HasAnyMarks(EObjectMark Marks) const
 	{
 		return ObjectHasAnyMarks(this,Marks);
 	}
@@ -154,7 +154,7 @@ public:
 	 * @param	Marks	Logical OR of OBJECTMARK_'s to test
 	 * @return	true if the object has any of the given marks.
 	 */
-	FORCEINLINE bool HasAllMarks(EObjectMark Marks) const
+	UE_FORCEINLINE_HINT bool HasAllMarks(EObjectMark Marks) const
 	{
 		return ObjectHasAllMarks(this,Marks);
 	}
@@ -165,7 +165,7 @@ public:
 	 * @param	Object	Object to get marks for
 	 * @return	all Marks for an object
 	 */
-	FORCEINLINE EObjectMark GetAllMarks() const
+	UE_FORCEINLINE_HINT EObjectMark GetAllMarks() const
 	{
 		return ObjectGetAllMarks(this);
 	}
@@ -203,13 +203,13 @@ public:
 	 * Add an object to the root set. This prevents the object and all
 	 * its descendants from being deleted during garbage collection.
 	 */
-	FORCEINLINE void AddToRoot()
+	UE_FORCEINLINE_HINT void AddToRoot()
 	{
 		GUObjectArray.IndexToObject(InternalIndex)->SetRootSet();
 	}
 
 	/** Remove an object from the root set. */
-	FORCEINLINE void RemoveFromRoot()
+	UE_FORCEINLINE_HINT void RemoveFromRoot()
 	{
 		GUObjectArray.IndexToObject(InternalIndex)->ClearRootSet();
 	}
@@ -219,7 +219,7 @@ public:
 	 *
 	 * @return true if the object was explicitly added as part of the root set.
 	 */
-	FORCEINLINE bool IsRooted() const
+	UE_FORCEINLINE_HINT bool IsRooted() const
 	{
 		return GUObjectArray.IndexToObject(InternalIndex)->IsRootSet();
 	}
@@ -229,7 +229,7 @@ public:
 	 *
 	 * @return the refcount of the object.
 	 */
-	FORCEINLINE int32 GetRefCount() const
+	UE_FORCEINLINE_HINT int32 GetRefCount() const
 	{
 		return GUObjectArray.IndexToObject(InternalIndex)->GetRefCount();
 	}
@@ -243,13 +243,13 @@ public:
 	bool ThisThreadAtomicallyClearedRFUnreachable();
 
 	/** Checks if the object is unreachable. */
-	FORCEINLINE bool IsUnreachable() const
+	UE_FORCEINLINE_HINT bool IsUnreachable() const
 	{
 		return GUObjectArray.IndexToObject(InternalIndex)->IsUnreachable();
 	}
 
 	/** Checks if the object is native. */
-	FORCEINLINE bool IsNative() const
+	UE_FORCEINLINE_HINT bool IsNative() const
 	{
 		return GUObjectArray.IndexToObject(InternalIndex)->HasAnyFlags(EInternalObjectFlags::Native);
 	}
@@ -275,7 +275,7 @@ public:
 	 * @param FlagsToClear	Object flags to clear.
 	 * @return				true if any of the passed in flags are set, false otherwise  (including no flags passed in).
 	 */
-	FORCEINLINE EInternalObjectFlags GetInternalFlags() const
+	UE_FORCEINLINE_HINT EInternalObjectFlags GetInternalFlags() const
 	{
 		return GUObjectArray.IndexToObject(InternalIndex)->GetFlags();
 	}
@@ -286,7 +286,7 @@ public:
 	 * @param FlagsToCheck	Object flags to check for.
 	 * @return				true if any of the passed in flags are set, false otherwise  (including no flags passed in).
 	 */
-	FORCEINLINE bool HasAnyInternalFlags(EInternalObjectFlags FlagsToCheck) const
+	UE_FORCEINLINE_HINT bool HasAnyInternalFlags(EInternalObjectFlags FlagsToCheck) const
 	{
 		return GUObjectArray.IndexToObject(InternalIndex)->HasAnyFlags(FlagsToCheck);
 	}
@@ -436,18 +436,18 @@ public:
 	 * 
 	 * @return Name of the object.
 	 */
-	FORCEINLINE FString GetName() const
+	UE_FORCEINLINE_HINT FString GetName() const
 	{
 		return GetFName().ToString();
 	}
 
 	/** Optimized version of GetName that overwrites an existing string */
-	FORCEINLINE void GetName(FString &ResultString) const
+	UE_FORCEINLINE_HINT void GetName(FString &ResultString) const
 	{
 		GetFName().ToString(ResultString);
 	}
 	/** Optimized version of GetName that appends to an existing string */
-	FORCEINLINE void AppendName(FString& ResultString) const
+	UE_FORCEINLINE_HINT void AppendName(FString& ResultString) const
 	{
 		GetFName().AppendString(ResultString);
 	}
@@ -604,7 +604,7 @@ public:
 
 private:
 	template <typename ClassType>
-	static FORCEINLINE bool IsChildOfWorkaround(const ClassType* ObjClass, const ClassType* TestCls)
+	static UE_FORCEINLINE_HINT bool IsChildOfWorkaround(const ClassType* ObjClass, const ClassType* TestCls)
 	{
 		return ObjClass->IsChildOf(TestCls);
 	}
@@ -767,7 +767,7 @@ public:
 	 *
 	 * @return true if this object's name is lexicographically smaller than the other object's name
 	 */
-	FORCEINLINE bool operator<( const UObjectBaseUtility& Other ) const
+	UE_FORCEINLINE_HINT bool operator<( const UObjectBaseUtility& Other ) const
 	{
 		return GetName() < Other.GetName();
 	}
@@ -1086,7 +1086,7 @@ public:
 #endif
 	}
 
-	FORCEINLINE_STATS FScopeCycleCounterUObject(const UObjectBaseUtility *Object, TStatId OtherStat)
+	UE_FORCEINLINE_HINT FScopeCycleCounterUObject(const UObjectBaseUtility *Object, TStatId OtherStat)
 		: FScopeCycleCounterUObject(Object)
 	{
 	}
@@ -1117,7 +1117,7 @@ class FScopeCycleCounterUObject
 {
 	const UObject* StatObject;
 public:
-	FORCEINLINE  FScopeCycleCounterUObject(const UObject* InStatObject, TStatId OtherStat = TStatId())
+	UE_FORCEINLINE_HINT  FScopeCycleCounterUObject(const UObject* InStatObject, TStatId OtherStat = TStatId())
 	{
 		StatObject = GHitchDetected ? nullptr : InStatObject;
 	}
@@ -1143,10 +1143,10 @@ public:
 class FScopeCycleCounterUObject
 {
 public:
-	FORCEINLINE_STATS FScopeCycleCounterUObject(const UObjectBaseUtility *Object)
+	UE_FORCEINLINE_HINT FScopeCycleCounterUObject(const UObjectBaseUtility *Object)
 	{
 	}
-	FORCEINLINE_STATS FScopeCycleCounterUObject(const UObjectBaseUtility *Object, TStatId OtherStat)
+	UE_FORCEINLINE_HINT FScopeCycleCounterUObject(const UObjectBaseUtility *Object, TStatId OtherStat)
 	{
 	}
 };
diff --git a/Engine/Source/Runtime/CoreUObject/Public/UObject/UObjectGlobals.h b/Engine/Source/Runtime/CoreUObject/Public/UObject/UObjectGlobals.h
--- a/Engine/Source/Runtime/CoreUObject/Public/UObject/UObjectGlobals.h
+++ b/Engine/Source/Runtime/CoreUObject/Public/UObject/UObjectGlobals.h
@@ -313,7 +313,7 @@ COREUOBJECT_API int32 UpdateSuffixForNextNewObject(UObject* Parent, const UClass
 COREUOBJECT_API UObject* StaticFindObjectFast(UClass* Class, UObject* InOuter, FName InName, EFindObjectFlags Flags = EFindObjectFlags::None, EObjectFlags ExclusiveFlags = RF_NoFlags, EInternalObjectFlags ExclusiveInternalFlags = EInternalObjectFlags::None);
 template <UE::CSameAs<bool> ExactClassType>
 UE_EXACTCLASS_BOOL_DEPRECATED("StaticFindObjectFast")
-UE_NODEBUG FORCEINLINE UObject* StaticFindObjectFast(UClass* Class, UObject* InOuter, FName InName, ExactClassType bExactClass, EObjectFlags ExclusiveFlags = RF_NoFlags, EInternalObjectFlags ExclusiveInternalFlags = EInternalObjectFlags::None)
+UE_NODEBUG UE_FORCEINLINE_HINT UObject* StaticFindObjectFast(UClass* Class, UObject* InOuter, FName InName, ExactClassType bExactClass, EObjectFlags ExclusiveFlags = RF_NoFlags, EInternalObjectFlags ExclusiveInternalFlags = EInternalObjectFlags::None)
 {
 	return StaticFindObjectFast(Class, InOuter, InName, bExactClass ? EFindObjectFlags::ExactClass : EFindObjectFlags::None, ExclusiveFlags, ExclusiveInternalFlags);
 }
@@ -335,7 +335,7 @@ UE_NODEBUG FORCEINLINE UObject* StaticFindObjectFast(UClass* Class, UObject* InO
 COREUOBJECT_API UObject* StaticFindObjectFastSafe(UClass* Class, UObject* InOuter, FName InName, EFindObjectFlags Flags = EFindObjectFlags::None, EObjectFlags ExclusiveFlags = RF_NoFlags, EInternalObjectFlags ExclusiveInternalFlags = EInternalObjectFlags::None);
 template <UE::CSameAs<bool> ExactClassType>
 UE_EXACTCLASS_BOOL_DEPRECATED("StaticFindObjectFastSafe")
-UE_NODEBUG FORCEINLINE UObject* StaticFindObjectFastSafe(UClass* Class, UObject* InOuter, FName InName, ExactClassType bExactClass, EObjectFlags ExclusiveFlags = RF_NoFlags, EInternalObjectFlags ExclusiveInternalFlags = EInternalObjectFlags::None)
+UE_NODEBUG UE_FORCEINLINE_HINT UObject* StaticFindObjectFastSafe(UClass* Class, UObject* InOuter, FName InName, ExactClassType bExactClass, EObjectFlags ExclusiveFlags = RF_NoFlags, EInternalObjectFlags ExclusiveInternalFlags = EInternalObjectFlags::None)
 {
 	return StaticFindObjectFastSafe(Class, InOuter, InName, bExactClass ? EFindObjectFlags::ExactClass : EFindObjectFlags::None, ExclusiveFlags, ExclusiveInternalFlags);
 }
@@ -353,7 +353,7 @@ UE_NODEBUG FORCEINLINE UObject* StaticFindObjectFastSafe(UClass* Class, UObject*
 COREUOBJECT_API UObject* StaticFindObject( UClass* Class, UObject* InOuter, FStringView Name, EFindObjectFlags Flags = EFindObjectFlags::None );
 template <UE::CSameAs<bool> ExactClassType>
 UE_EXACTCLASS_BOOL_DEPRECATED("StaticFindObject")
-UE_NODEBUG FORCEINLINE UObject* StaticFindObject( UClass* Class, UObject* InOuter, FStringView Name, ExactClassType bExactClass )
+UE_NODEBUG UE_FORCEINLINE_HINT UObject* StaticFindObject( UClass* Class, UObject* InOuter, FStringView Name, ExactClassType bExactClass )
 {
 	return StaticFindObject(Class, InOuter, Name, bExactClass ? EFindObjectFlags::ExactClass : EFindObjectFlags::None);
 }
@@ -370,7 +370,7 @@ UE_NODEBUG FORCEINLINE UObject* StaticFindObject( UClass* Class, UObject* InOute
 COREUOBJECT_API UObject* StaticFindObject(UClass* Class, FTopLevelAssetPath ObjectPath, EFindObjectFlags Flags = EFindObjectFlags::None);
 template <UE::CSameAs<bool> ExactClassType>
 UE_EXACTCLASS_BOOL_DEPRECATED("StaticFindObject")
-UE_NODEBUG FORCEINLINE UObject* StaticFindObject(UClass* Class, FTopLevelAssetPath ObjectPath, ExactClassType bExactClass)
+UE_NODEBUG UE_FORCEINLINE_HINT UObject* StaticFindObject(UClass* Class, FTopLevelAssetPath ObjectPath, ExactClassType bExactClass)
 {
 	return StaticFindObject(Class, ObjectPath, bExactClass ? EFindObjectFlags::ExactClass : EFindObjectFlags::None);
 }
@@ -379,7 +379,7 @@ UE_NODEBUG FORCEINLINE UObject* StaticFindObject(UClass* Class, FTopLevelAssetPa
 COREUOBJECT_API UObject* StaticFindObjectChecked( UClass* Class, UObject* InOuter, FStringView Name, EFindObjectFlags Flags = EFindObjectFlags::None );
 template <UE::CSameAs<bool> ExactClassType>
 UE_EXACTCLASS_BOOL_DEPRECATED("StaticFindObjectChecked")
-UE_NODEBUG FORCEINLINE UObject* StaticFindObjectChecked( UClass* Class, UObject* InOuter, FStringView Name, ExactClassType bExactClass)
+UE_NODEBUG UE_FORCEINLINE_HINT UObject* StaticFindObjectChecked( UClass* Class, UObject* InOuter, FStringView Name, ExactClassType bExactClass)
 {
 	return StaticFindObjectChecked(Class, InOuter, Name, bExactClass ? EFindObjectFlags::ExactClass : EFindObjectFlags::None);
 }
@@ -388,7 +388,7 @@ UE_NODEBUG FORCEINLINE UObject* StaticFindObjectChecked( UClass* Class, UObject*
 COREUOBJECT_API UObject* StaticFindObjectSafe( UClass* Class, UObject* InOuter, FStringView Name, EFindObjectFlags Flags = EFindObjectFlags::None );
 template <UE::CSameAs<bool> ExactClassType>
 UE_EXACTCLASS_BOOL_DEPRECATED("StaticFindObjectSafe")
-UE_NODEBUG FORCEINLINE UObject* StaticFindObjectSafe( UClass* Class, UObject* InOuter, FStringView Name, ExactClassType bExactClass )
+UE_NODEBUG UE_FORCEINLINE_HINT UObject* StaticFindObjectSafe( UClass* Class, UObject* InOuter, FStringView Name, ExactClassType bExactClass )
 {
 	return StaticFindObjectSafe(Class, InOuter, Name, bExactClass ? EFindObjectFlags::ExactClass : EFindObjectFlags::None);
 }
@@ -406,7 +406,7 @@ UE_NODEBUG FORCEINLINE UObject* StaticFindObjectSafe( UClass* Class, UObject* In
 COREUOBJECT_API UObject* StaticFindObjectSafe(UClass* Class, FTopLevelAssetPath ObjectPath, EFindObjectFlags Flags = EFindObjectFlags::None);
 template <UE::CSameAs<bool> ExactClassType>
 UE_EXACTCLASS_BOOL_DEPRECATED("StaticFindObjectSafe")
-UE_NODEBUG FORCEINLINE UObject* StaticFindObjectSafe(UClass* Class, FTopLevelAssetPath ObjectPath, ExactClassType bExactClass)
+UE_NODEBUG UE_FORCEINLINE_HINT UObject* StaticFindObjectSafe(UClass* Class, FTopLevelAssetPath ObjectPath, ExactClassType bExactClass)
 {
 	return StaticFindObjectSafe(Class, ObjectPath, bExactClass ? EFindObjectFlags::ExactClass : EFindObjectFlags::None);
 }
@@ -427,7 +427,7 @@ UE_NODEBUG FORCEINLINE UObject* StaticFindObjectSafe(UClass* Class, FTopLevelAss
 COREUOBJECT_API bool StaticFindAllObjectsFast(TArray<UObject*>& OutFoundObjects, UClass* ObjectClass, FName ObjectName, EFindObjectFlags Flags = EFindObjectFlags::None, EObjectFlags ExclusiveFlags = RF_NoFlags, EInternalObjectFlags ExclusiveInternalFlags = EInternalObjectFlags::None);
 template <UE::CSameAs<bool> ExactClassType>
 UE_EXACTCLASS_BOOL_DEPRECATED("StaticFindAllObjectsFast")
-UE_NODEBUG FORCEINLINE bool StaticFindAllObjectsFast(TArray<UObject*>& OutFoundObjects, UClass* ObjectClass, FName ObjectName, ExactClassType bExactClass, EObjectFlags ExclusiveFlags = RF_NoFlags, EInternalObjectFlags ExclusiveInternalFlags = EInternalObjectFlags::None)
+UE_NODEBUG UE_FORCEINLINE_HINT bool StaticFindAllObjectsFast(TArray<UObject*>& OutFoundObjects, UClass* ObjectClass, FName ObjectName, ExactClassType bExactClass, EObjectFlags ExclusiveFlags = RF_NoFlags, EInternalObjectFlags ExclusiveInternalFlags = EInternalObjectFlags::None)
 {
 	return StaticFindAllObjectsFast(OutFoundObjects, ObjectClass, ObjectName, bExactClass ? EFindObjectFlags::ExactClass : EFindObjectFlags::None, ExclusiveFlags, ExclusiveInternalFlags);
 }
@@ -449,7 +449,7 @@ UE_NODEBUG FORCEINLINE bool StaticFindAllObjectsFast(TArray<UObject*>& OutFoundO
 COREUOBJECT_API bool StaticFindAllObjectsFastSafe(TArray<UObject*>& OutFoundObjects, UClass* ObjectClass, FName ObjectName, EFindObjectFlags Flags = EFindObjectFlags::None, EObjectFlags ExclusiveFlags = RF_NoFlags, EInternalObjectFlags ExclusiveInternalFlags = EInternalObjectFlags::None);
 template <UE::CSameAs<bool> ExactClassType>
 UE_EXACTCLASS_BOOL_DEPRECATED("StaticFindAllObjectsFastSafe")
-UE_NODEBUG FORCEINLINE bool StaticFindAllObjectsFastSafe(TArray<UObject*>& OutFoundObjects, UClass* ObjectClass, FName ObjectName, ExactClassType bExactClass, EObjectFlags ExclusiveFlags = RF_NoFlags, EInternalObjectFlags ExclusiveInternalFlags = EInternalObjectFlags::None)
+UE_NODEBUG UE_FORCEINLINE_HINT bool StaticFindAllObjectsFastSafe(TArray<UObject*>& OutFoundObjects, UClass* ObjectClass, FName ObjectName, ExactClassType bExactClass, EObjectFlags ExclusiveFlags = RF_NoFlags, EInternalObjectFlags ExclusiveInternalFlags = EInternalObjectFlags::None)
 {
 	return StaticFindAllObjectsFastSafe(OutFoundObjects, ObjectClass, ObjectName, bExactClass ? EFindObjectFlags::ExactClass : EFindObjectFlags::None, ExclusiveFlags, ExclusiveInternalFlags);
 }
@@ -467,7 +467,7 @@ UE_NODEBUG FORCEINLINE bool StaticFindAllObjectsFastSafe(TArray<UObject*>& OutFo
 COREUOBJECT_API bool StaticFindAllObjects(TArray<UObject*>& OutFoundObjects, UClass* Class, const TCHAR* Name, EFindObjectFlags Flags = EFindObjectFlags::None);
 template <UE::CSameAs<bool> ExactClassType>
 UE_EXACTCLASS_BOOL_DEPRECATED("StaticFindAllObjects")
-UE_NODEBUG FORCEINLINE bool StaticFindAllObjects(TArray<UObject*>& OutFoundObjects, UClass* Class, const TCHAR* Name, ExactClassType bExactClass)
+UE_NODEBUG UE_FORCEINLINE_HINT bool StaticFindAllObjects(TArray<UObject*>& OutFoundObjects, UClass* Class, const TCHAR* Name, ExactClassType bExactClass)
 {
 	return StaticFindAllObjects(OutFoundObjects, Class, Name, bExactClass ? EFindObjectFlags::ExactClass : EFindObjectFlags::None);
 }
@@ -486,7 +486,7 @@ UE_NODEBUG FORCEINLINE bool StaticFindAllObjects(TArray<UObject*>& OutFoundObjec
 COREUOBJECT_API bool StaticFindAllObjectsSafe(TArray<UObject*>& OutFoundObjects, UClass* Class, const TCHAR* Name, EFindObjectFlags Flags = EFindObjectFlags::None);
 template <UE::CSameAs<bool> ExactClassType>
 UE_EXACTCLASS_BOOL_DEPRECATED("StaticFindAllObjectsSafe")
-UE_NODEBUG FORCEINLINE bool StaticFindAllObjectsSafe(TArray<UObject*>& OutFoundObjects, UClass* Class, const TCHAR* Name, ExactClassType bExactClass)
+UE_NODEBUG UE_FORCEINLINE_HINT bool StaticFindAllObjectsSafe(TArray<UObject*>& OutFoundObjects, UClass* Class, const TCHAR* Name, ExactClassType bExactClass)
 {
 	return StaticFindAllObjectsSafe(OutFoundObjects, Class, Name, bExactClass ? EFindObjectFlags::ExactClass : EFindObjectFlags::None);
 }
@@ -1318,7 +1318,7 @@ public:
 	/** 
 	 * Return the archetype that this object will copy properties from later
 	**/
-	FORCEINLINE UObject* GetArchetype() const
+	UE_FORCEINLINE_HINT UObject* GetArchetype() const
 	{
 		return ObjectArchetype;
 	}
@@ -1326,17 +1326,17 @@ public:
 	/**
 	* Return the object that is being constructed
 	**/
-	FORCEINLINE UObject* GetObj() const
+	UE_FORCEINLINE_HINT UObject* GetObj() const
 	{
 		return Obj;
 	}
 
-	FORCEINLINE struct FObjectInstancingGraph* GetInstancingGraph()
+	UE_FORCEINLINE_HINT struct FObjectInstancingGraph* GetInstancingGraph()
 	{
 		return InstanceGraph;
 	}
 
-	FORCEINLINE void AddPropertyPostInitCallback(TFunction<void()>&& Callback)
+	UE_FORCEINLINE_HINT void AddPropertyPostInitCallback(TFunction<void()>&& Callback)
 	{
 		PropertyPostInitCallbacks.Add(MoveTemp(Callback));
 	}
@@ -1527,7 +1527,7 @@ public:
 	 **/
 	static COREUOBJECT_API void AssertIfInConstructor(UObject* Outer, const TCHAR* ErrorMessage);
 
-	FORCEINLINE void FinalizeSubobjectClassInitialization()
+	UE_FORCEINLINE_HINT void FinalizeSubobjectClassInitialization()
 	{
 		bSubobjectClassInitializationAllowed = false;
 	}
@@ -2017,7 +2017,7 @@ inline T* FindObjectFast(UObject* Outer, FName Name, EFindObjectFlags Flags = EF
 }
 template <typename T, UE::CSameAs<bool> ExactClassType>
 UE_EXACTCLASS_BOOL_DEPRECATED("FindObjectFast")
-UE_NODEBUG FORCEINLINE T* FindObjectFast(UObject* Outer, FName Name, ExactClassType bExactClass, EObjectFlags ExclusiveFlags = RF_NoFlags)
+UE_NODEBUG UE_FORCEINLINE_HINT T* FindObjectFast(UObject* Outer, FName Name, ExactClassType bExactClass, EObjectFlags ExclusiveFlags = RF_NoFlags)
 {
 	return FindObjectFast<T>(Outer, Name, bExactClass ? EFindObjectFlags::ExactClass : EFindObjectFlags::None, ExclusiveFlags);
 }
@@ -2033,7 +2033,7 @@ inline T* FindObject( UObject* Outer, FStringView Name, EFindObjectFlags Flags =
 }
 template <typename T, UE::CSameAs<bool> ExactClassType>
 UE_EXACTCLASS_BOOL_DEPRECATED("FindObject")
-UE_NODEBUG FORCEINLINE T* FindObject(UObject* Outer, FStringView Name, ExactClassType bExactClass)
+UE_NODEBUG UE_FORCEINLINE_HINT T* FindObject(UObject* Outer, FStringView Name, ExactClassType bExactClass)
 {
 	return FindObject<T>(Outer, Name, bExactClass ? EFindObjectFlags::ExactClass : EFindObjectFlags::None);
 }
@@ -2049,7 +2049,7 @@ inline T* FindObject(FTopLevelAssetPath InPath, EFindObjectFlags Flags = EFindOb
 }
 template <typename T, UE::CSameAs<bool> ExactClassType>
 UE_EXACTCLASS_BOOL_DEPRECATED("FindObject")
-UE_NODEBUG FORCEINLINE T* FindObject(FTopLevelAssetPath InPath, ExactClassType bExactClass)
+UE_NODEBUG UE_FORCEINLINE_HINT T* FindObject(FTopLevelAssetPath InPath, ExactClassType bExactClass)
 {
 	return FindObject<T>(InPath, bExactClass ? EFindObjectFlags::ExactClass : EFindObjectFlags::None);
 }
@@ -2065,7 +2065,7 @@ inline T* FindObjectChecked( UObject* Outer, FStringView Name, EFindObjectFlags
 }
 template <typename T, UE::CSameAs<bool> ExactClassType>
 UE_EXACTCLASS_BOOL_DEPRECATED("FindObjectChecked")
-UE_NODEBUG FORCEINLINE T* FindObjectChecked( UObject* Outer, FStringView Name, ExactClassType bExactClass )
+UE_NODEBUG UE_FORCEINLINE_HINT T* FindObjectChecked( UObject* Outer, FStringView Name, ExactClassType bExactClass )
 {
 	return FindObjectChecked<T>(Outer, Name, bExactClass ? EFindObjectFlags::ExactClass : EFindObjectFlags::None);
 }
@@ -2081,7 +2081,7 @@ inline T* FindObjectSafe( UObject* Outer, FStringView Name, EFindObjectFlags Fla
 }
 template <typename T, UE::CSameAs<bool> ExactClassType>
 UE_EXACTCLASS_BOOL_DEPRECATED("FindObjectSafe")
-UE_NODEBUG FORCEINLINE T* FindObjectSafe( UObject* Outer, FStringView Name, ExactClassType bExactClass )
+UE_NODEBUG UE_FORCEINLINE_HINT T* FindObjectSafe( UObject* Outer, FStringView Name, ExactClassType bExactClass )
 {
 	return FindObjectSafe<T>(Outer, Name, bExactClass ? EFindObjectFlags::ExactClass : EFindObjectFlags::None);
 }
@@ -2097,7 +2097,7 @@ inline T* FindObjectSafe(FTopLevelAssetPath InPath, EFindObjectFlags Flags = EFi
 }
 template <typename T, UE::CSameAs<bool> ExactClassType>
 UE_EXACTCLASS_BOOL_DEPRECATED("FindObjectSafe")
-UE_NODEBUG FORCEINLINE T* FindObjectSafe(FTopLevelAssetPath InPath, ExactClassType bExactClass)
+UE_NODEBUG UE_FORCEINLINE_HINT T* FindObjectSafe(FTopLevelAssetPath InPath, ExactClassType bExactClass)
 {
 	return FindObjectSafe<T>(InPath, bExactClass ? EFindObjectFlags::ExactClass : EFindObjectFlags::None);
 }
@@ -2222,7 +2222,7 @@ bool ContainsObjectOfClass( const TArray<T*>& ObjectArray, UClass* ClassToCheck,
 }
 template <typename T, UE::CSameAs<bool> ExactClassType>
 UE_EXACTCLASS_BOOL_DEPRECATED("ContainsObjectOfClass")
-UE_NODEBUG FORCEINLINE bool ContainsObjectOfClass( const TArray<T*>& ObjectArray, UClass* ClassToCheck, ExactClassType bExactClass, TArray<T*>* out_Objects=nullptr )
+UE_NODEBUG UE_FORCEINLINE_HINT bool ContainsObjectOfClass( const TArray<T*>& ObjectArray, UClass* ClassToCheck, ExactClassType bExactClass, TArray<T*>* out_Objects=nullptr )
 {
 	return ContainsObjectOfClass(ObjectArray, ClassToCheck, bExactClass ? EFindObjectFlags::ExactClass : EFindObjectFlags::None, out_Objects );
 }
@@ -2310,24 +2310,24 @@ public:
 	}
 	template <UE::CSameAs<bool> ExactClassType>
 	UE_EXACTCLASS_BOOL_DEPRECATED("TObjectArrayIterator")
-	UE_NODEBUG FORCEINLINE TObjectArrayIterator( TArray<TObjectClass*>& InArray, UClass* InClassToCheck, ExactClassType bInExactClass)
+	UE_NODEBUG UE_FORCEINLINE_HINT TObjectArrayIterator( TArray<TObjectClass*>& InArray, UClass* InClassToCheck, ExactClassType bInExactClass)
 		: TObjectArrayIterator(InArray, InClassToCheck, bInExactClass ? EFindObjectFlags::ExactClass : EFindObjectFlags::None)
 	{
 	}
 	/**
 		* Iterator advance
 		*/
-	FORCEINLINE void operator++()
+	UE_FORCEINLINE_HINT void operator++()
 	{
 		Advance();
 	}
 	/** conversion to "bool" returning true if the iterator is valid. */
-	FORCEINLINE explicit operator bool() const
+	UE_FORCEINLINE_HINT explicit operator bool() const
 	{ 
 		return Index < Array.Num(); 
 	}
 	/** inverse of the "bool" operator */
-	FORCEINLINE bool operator !() const 
+	UE_FORCEINLINE_HINT bool operator !() const 
 	{
 		return !(bool)*this;
 	}
@@ -2355,7 +2355,7 @@ public:
 	  * Removes the current element from the array, slower, but preserves the order. 
 	  * Iterator is decremented for you so a loop will check all items.
 	*/
-	FORCEINLINE void RemoveCurrent()
+	UE_FORCEINLINE_HINT void RemoveCurrent()
 	{
 		Array.RemoveAt(Index--);
 	}
@@ -2363,7 +2363,7 @@ public:
 	  * Removes the current element from the array, faster, but does not preserves the array order. 
 	  * Iterator is decremented for you so a loop will check all items.
 	*/
-	FORCEINLINE void RemoveCurrentSwap()
+	UE_FORCEINLINE_HINT void RemoveCurrentSwap()
 	{
 		Array.RemoveSwap(Index--);
 	}
@@ -2373,7 +2373,7 @@ protected:
 		* Dereferences the iterator with an ordinary name for clarity in derived classes
 		* @return	the UObject at the iterator
 	*/
-	FORCEINLINE TObjectClass* GetObject() const 
+	UE_FORCEINLINE_HINT TObjectClass* GetObject() const 
 	{ 
 		return Array(Index);
 	}
@@ -2421,7 +2421,7 @@ class FReferenceCollectorArchive : public FArchiveUObject
 	class FReferenceCollector& Collector;
 
 protected:
-	FORCEINLINE class FReferenceCollector& GetCollector()
+	UE_FORCEINLINE_HINT class FReferenceCollector& GetCollector()
 	{
 		return Collector;
 	}
@@ -2565,19 +2565,19 @@ public:
 	}
 
 	template<class UObjectType>
-	FORCEINLINE void AddStableReference(TObjectPtr<UObjectType>* Object)
+	UE_FORCEINLINE_HINT void AddStableReference(TObjectPtr<UObjectType>* Object)
 	{
 		AddStableReferenceFwd(reinterpret_cast<FObjectPtr*>(Object));
 	}
 
 	template<class UObjectType>
-	FORCEINLINE void AddStableReferenceArray(TArray<TObjectPtr<UObjectType>>* Objects)
+	UE_FORCEINLINE_HINT void AddStableReferenceArray(TArray<TObjectPtr<UObjectType>>* Objects)
 	{
 		AddStableReferenceArrayFwd(reinterpret_cast<TArray<FObjectPtr>*>(Objects));
 	}
 
 	template<class UObjectType>
-	FORCEINLINE void AddStableReferenceSet(TSet<TObjectPtr<UObjectType>>* Objects)
+	UE_FORCEINLINE_HINT void AddStableReferenceSet(TSet<TObjectPtr<UObjectType>>* Objects)
 	{
 		AddStableReferenceSetFwd(reinterpret_cast<TSet<FObjectPtr>*>(Objects));
 	}
diff --git a/Engine/Source/Runtime/CoreUObject/Public/UObject/UObjectHash.h b/Engine/Source/Runtime/CoreUObject/Public/UObject/UObjectHash.h
--- a/Engine/Source/Runtime/CoreUObject/Public/UObject/UObjectHash.h
+++ b/Engine/Source/Runtime/CoreUObject/Public/UObject/UObjectHash.h
@@ -46,7 +46,7 @@ DECLARE_DWORD_COUNTER_STAT_EXTERN(TEXT("NumObjects"), STAT_Hash_NumObjects, STAT
 COREUOBJECT_API UObject* StaticFindObjectFastInternal(const UClass* Class, const UObject* InOuter, FName InName, EFindObjectFlags Flags = EFindObjectFlags::None, EObjectFlags ExclusiveFlags = RF_NoFlags, EInternalObjectFlags ExclusiveInternalFlags = EInternalObjectFlags::None);
 template <UE::CSameAs<bool> ExactClassType>
 UE_EXACTCLASS_BOOL_DEPRECATED("StaticFindObjectFastInternal")
-UE_NODEBUG FORCEINLINE UObject* StaticFindObjectFastInternal(const UClass* Class, const UObject* InOuter, FName InName, ExactClassType bExactClass, EObjectFlags ExclusiveFlags = RF_NoFlags, EInternalObjectFlags ExclusiveInternalFlags = EInternalObjectFlags::None)
+UE_NODEBUG UE_FORCEINLINE_HINT UObject* StaticFindObjectFastInternal(const UClass* Class, const UObject* InOuter, FName InName, ExactClassType bExactClass, EObjectFlags ExclusiveFlags = RF_NoFlags, EInternalObjectFlags ExclusiveInternalFlags = EInternalObjectFlags::None)
 {
 	return StaticFindObjectFastInternal(Class, InOuter, InName, bExactClass ? EFindObjectFlags::ExactClass : EFindObjectFlags::None, ExclusiveFlags, ExclusiveInternalFlags);
 }
@@ -65,7 +65,7 @@ UE_NODEBUG FORCEINLINE UObject* StaticFindObjectFastInternal(const UClass* Class
 COREUOBJECT_API bool StaticFindAllObjectsFastInternal(TArray<UObject*>& OutFoundObjects, const UClass* ObjectClass, FName ObjectName, EFindObjectFlags Flags = EFindObjectFlags::None, EObjectFlags ExcludeFlags = RF_NoFlags, EInternalObjectFlags ExclusiveInternalFlags = EInternalObjectFlags::None);
 template <UE::CSameAs<bool> ExactClassType>
 UE_EXACTCLASS_BOOL_DEPRECATED("StaticFindAllObjectsFastInternal")
-UE_NODEBUG FORCEINLINE bool StaticFindAllObjectsFastInternal(TArray<UObject*>& OutFoundObjects, const UClass* ObjectClass, FName ObjectName, ExactClassType bExactClass, EObjectFlags ExcludeFlags = RF_NoFlags, EInternalObjectFlags ExclusiveInternalFlags = EInternalObjectFlags::None)
+UE_NODEBUG UE_FORCEINLINE_HINT bool StaticFindAllObjectsFastInternal(TArray<UObject*>& OutFoundObjects, const UClass* ObjectClass, FName ObjectName, ExactClassType bExactClass, EObjectFlags ExcludeFlags = RF_NoFlags, EInternalObjectFlags ExclusiveInternalFlags = EInternalObjectFlags::None)
 {
 	return StaticFindAllObjectsFastInternal(OutFoundObjects, ObjectClass, ObjectName, bExactClass ? EFindObjectFlags::ExactClass : EFindObjectFlags::None, ExcludeFlags, ExclusiveInternalFlags);
 }
@@ -83,7 +83,7 @@ UE_NODEBUG FORCEINLINE bool StaticFindAllObjectsFastInternal(TArray<UObject*>& O
 COREUOBJECT_API UObject* StaticFindFirstObjectFastInternal(const UClass* ObjectClass, FName ObjectName, EFindObjectFlags Flags = EFindObjectFlags::None, EObjectFlags ExcludeFlag = RF_NoFlags, EInternalObjectFlags ExclusiveInternalFlags = EInternalObjectFlags::None);
 template <UE::CSameAs<bool> ExactClassType>
 UE_EXACTCLASS_BOOL_DEPRECATED("StaticFindFirstObjectFastInternal")
-UE_NODEBUG FORCEINLINE UObject* StaticFindFirstObjectFastInternal(const UClass* ObjectClass, FName ObjectName, ExactClassType bExactClass, EObjectFlags ExcludeFlag = RF_NoFlags, EInternalObjectFlags ExclusiveInternalFlags = EInternalObjectFlags::None)
+UE_NODEBUG UE_FORCEINLINE_HINT UObject* StaticFindFirstObjectFastInternal(const UClass* ObjectClass, FName ObjectName, ExactClassType bExactClass, EObjectFlags ExcludeFlag = RF_NoFlags, EInternalObjectFlags ExclusiveInternalFlags = EInternalObjectFlags::None)
 {
 	return StaticFindFirstObjectFastInternal(ObjectClass, ObjectName, bExactClass ? EFindObjectFlags::ExactClass : EFindObjectFlags::None, ExcludeFlag, ExclusiveInternalFlags);
 }
@@ -101,7 +101,7 @@ UE_NODEBUG FORCEINLINE UObject* StaticFindFirstObjectFastInternal(const UClass*
 UObject* StaticFindObjectFastExplicit(const UClass* ObjectClass, FName ObjectName, const FString& ObjectPathName, EFindObjectFlags Flags = EFindObjectFlags::None, EObjectFlags ExcludeFlags = RF_NoFlags);
 template <UE::CSameAs<bool> ExactClassType>
 UE_EXACTCLASS_BOOL_DEPRECATED("StaticFindObjectFastExplicit")
-UE_NODEBUG FORCEINLINE UObject* StaticFindObjectFastExplicit(const UClass* ObjectClass, FName ObjectName, const FString& ObjectPathName, ExactClassType bExactClass, EObjectFlags ExcludeFlags = RF_NoFlags)
+UE_NODEBUG UE_FORCEINLINE_HINT UObject* StaticFindObjectFastExplicit(const UClass* ObjectClass, FName ObjectName, const FString& ObjectPathName, ExactClassType bExactClass, EObjectFlags ExcludeFlags = RF_NoFlags)
 {
 	return StaticFindObjectFastExplicit(ObjectClass, ObjectName, ObjectPathName, bExactClass ? EFindObjectFlags::ExactClass : EFindObjectFlags::None, ExcludeFlags);
 }
@@ -357,11 +357,11 @@ void UnlockUObjectHashTables();
 class FScopedUObjectHashTablesLock
 {
 public:
-	FORCEINLINE FScopedUObjectHashTablesLock()
+	UE_FORCEINLINE_HINT FScopedUObjectHashTablesLock()
 	{
 		LockUObjectHashTables();
 	}
-	FORCEINLINE ~FScopedUObjectHashTablesLock()
+	UE_FORCEINLINE_HINT ~FScopedUObjectHashTablesLock()
 	{
 		UnlockUObjectHashTables();
 	}
diff --git a/Engine/Source/Runtime/CoreUObject/Public/UObject/UObjectIterator.h b/Engine/Source/Runtime/CoreUObject/Public/UObject/UObjectIterator.h
--- a/Engine/Source/Runtime/CoreUObject/Public/UObject/UObjectIterator.h
+++ b/Engine/Source/Runtime/CoreUObject/Public/UObject/UObjectIterator.h
@@ -32,7 +32,7 @@ public:
 	 * Iterator dereference
 	 * @return	the object pointer pointed at by the iterator
 	 */
-	FORCEINLINE FUObjectItem* operator*() const
+	UE_FORCEINLINE_HINT FUObjectItem* operator*() const
 	{
 		// casting UObjectBase to UObject for clients
 		return GetObject();
@@ -41,7 +41,7 @@ public:
 	 * Iterator dereference
 	 * @return	the object pointer pointed at by the iterator
 	 */
-	FORCEINLINE FUObjectItem* operator->() const
+	UE_FORCEINLINE_HINT FUObjectItem* operator->() const
 	{
 		return GetObject();
 	}
@@ -282,18 +282,18 @@ public:
 	/**
 	 * Iterator advance
 	 */
-	FORCEINLINE void operator++()
+	UE_FORCEINLINE_HINT void operator++()
 	{
 		Advance();
 	}
 
 	/** Conversion to "bool" returning true if the iterator is valid. */
-	FORCEINLINE explicit operator bool() const
+	UE_FORCEINLINE_HINT explicit operator bool() const
 	{ 
 		return ObjectArray.IsValidIndex(Index); 
 	}
 	/** Conversion to "bool" returning true if the iterator is valid. */
-	FORCEINLINE bool operator !() const 
+	UE_FORCEINLINE_HINT bool operator !() const 
 	{
 		return !(bool)*this;
 	}
@@ -303,7 +303,7 @@ public:
 	 * @return	the object pointer pointed at by the iterator
 	 */
 	FUNCTION_NON_NULL_RETURN_START
-		FORCEINLINE T* operator* () const
+		UE_FORCEINLINE_HINT T* operator* () const
 	FUNCTION_NON_NULL_RETURN_END
 	{
 		return (T*)GetObject();
@@ -313,14 +313,14 @@ public:
 	 * @return	the object pointer pointed at by the iterator
 	 */
 	FUNCTION_NON_NULL_RETURN_START
-		FORCEINLINE T* operator-> () const
+		UE_FORCEINLINE_HINT T* operator-> () const
 	FUNCTION_NON_NULL_RETURN_END
 	{
 		return (T*)GetObject();
 	}
 
-	FORCEINLINE bool operator==(const TObjectIterator& Rhs) const { return Index == Rhs.Index; }
-	FORCEINLINE bool operator!=(const TObjectIterator& Rhs) const { return Index != Rhs.Index; }
+	UE_FORCEINLINE_HINT bool operator==(const TObjectIterator& Rhs) const { return Index == Rhs.Index; }
+	UE_FORCEINLINE_HINT bool operator!=(const TObjectIterator& Rhs) const { return Index != Rhs.Index; }
 
 protected:
 	/**
@@ -328,7 +328,7 @@ protected:
 	 *
 	 * @return	the UObject at the iterator
 	 */
-	FORCEINLINE UObject* GetObject() const 
+	UE_FORCEINLINE_HINT UObject* GetObject() const 
 	{ 
 		return ObjectArray[Index];
 	}
diff --git a/Engine/Source/Runtime/CoreUObject/Public/UObject/UObjectStats.h b/Engine/Source/Runtime/CoreUObject/Public/UObject/UObjectStats.h
--- a/Engine/Source/Runtime/CoreUObject/Public/UObject/UObjectStats.h
+++ b/Engine/Source/Runtime/CoreUObject/Public/UObject/UObjectStats.h
@@ -15,17 +15,17 @@ namespace UObjectStats
 {
 	extern COREUOBJECT_API std::atomic<int32> GUObjectCount;
 
-	FORCEINLINE void IncrementUObjectCount()
+	UE_FORCEINLINE_HINT void IncrementUObjectCount()
 	{
 		GUObjectCount.fetch_add(1, std::memory_order_relaxed);
 	}
 
-	FORCEINLINE void DecrementUObjectCount()
+	UE_FORCEINLINE_HINT void DecrementUObjectCount()
 	{
 		GUObjectCount.fetch_sub(1, std::memory_order_relaxed);
 	}
 
-	FORCEINLINE int32 GetUObjectCount()
+	UE_FORCEINLINE_HINT int32 GetUObjectCount()
 	{
 		return GUObjectCount.load(std::memory_order_relaxed);
 	}
