Moved the scope of the Remote Executor Verbosity Log

diff --git a/Engine/Source/Runtime/CoreUObject/Private/UObject/RemoteExecutor.cpp b/Engine/Source/Runtime/CoreUObject/Private/UObject/RemoteExecutor.cpp
--- a/Engine/Source/Runtime/CoreUObject/Private/UObject/RemoteExecutor.cpp
+++ b/Engine/Source/Runtime/CoreUObject/Private/UObject/RemoteExecutor.cpp
@@ -121,14 +121,6 @@ public:
 
 		double LastStallPrintTime = FPlatformTime::Seconds();
 
-#if !UE_BUILD_SHIPPING && !NO_LOGGING
-		ELogVerbosity::Type InitialLogVerbosity = LogRemoteExec.GetVerbosity();
-		ON_SCOPE_EXIT
-		{
-			LogRemoteExec.SetVerbosity(InitialLogVerbosity);
-		};
-#endif
-
 		for (int32 LocalIterationNumber = 0; ; LocalIterationNumber++)
 		{
 			TRACE_CPUPROFILER_EVENT_SCOPE(ResumeExecutePendingWork)
@@ -137,6 +129,12 @@ public:
 			bool bPrintStall = false;
 
 #if !UE_BUILD_SHIPPING && !NO_LOGGING
+			ELogVerbosity::Type InitialLogVerbosity = LogRemoteExec.GetVerbosity();
+			ON_SCOPE_EXIT
+			{
+				LogRemoteExec.SetVerbosity(InitialLogVerbosity);
+			};
+
 			constexpr int32 MaxIterationsBeforeStall = 100000;
 			if (LocalIterationNumber >= MaxIterationsBeforeStall && !(LocalIterationNumber % MaxIterationsBeforeStall))
 			{
