Reducing verbosity of the incremental purge logs

diff --git a/Engine/Source/Runtime/CoreUObject/Private/UObject/GarbageCollection.cpp b/Engine/Source/Runtime/CoreUObject/Private/UObject/GarbageCollection.cpp
--- a/Engine/Source/Runtime/CoreUObject/Private/UObject/GarbageCollection.cpp
+++ b/Engine/Source/Runtime/CoreUObject/Private/UObject/GarbageCollection.cpp
@@ -5056,15 +5056,6 @@ bool IncrementalDestroyGarbage(bool bUseTimeLimit, double TimeLimit)
 			GObjFinishDestroyHasBeenRoutedToAllObjects		= false;
 			GObjPurgeIsRequired								= false;
 			GObjCurrentPurgeObjectIndexNeedsReset			= true;
-
-			// Log status information.
-			const int32 PurgedObjectCountSinceLastMarkPhase = GUObjectPurge.GetObjectsDestroyedSinceLastMarkPhase();
-			UE_LOG(LogGarbage, Log, TEXT("GC purged %i objects (%i -> %i) in %.3fms"), PurgedObjectCountSinceLastMarkPhase, 
-				GObjectCountDuringLastMarkPhase.GetValue(), 
-				GObjectCountDuringLastMarkPhase.GetValue() - PurgedObjectCountSinceLastMarkPhase,
-				(FPlatformTime::Seconds() - IncrementalDestroyGarbageStartTime) * 1000);
-			UE::GC::GDetailedStats.LogPurgeStats(PurgedObjectCountSinceLastMarkPhase);
-			GUObjectPurge.ResetObjectsDestroyedSinceLastMarkPhase();
 		}
 	}
 
@@ -5072,14 +5063,26 @@ bool IncrementalDestroyGarbage(bool bUseTimeLimit, double TimeLimit)
 	GGCStats.DestroyGarbageTime += ElapsedTime;
 	GGCStats.DestroyGarbageTimeLimit = TimeLimit;
 
+	// Log status information.
 	if (bUseTimeLimit && !bCompleted)
 	{
-		UE_LOG(LogGarbage, Log, TEXT("%.3f ms for incrementally purging unreachable objects (FinishDestroyed: %d, Destroyed: %d / %d)"),
+		UE_LOG(LogGarbage, Verbose, TEXT("%.3f ms for incrementally purging unreachable objects (FinishDestroyed: %d, Destroyed: %d / %d)"),
 			ElapsedTime * 1000,
 			GObjCurrentPurgeObjectIndex,
 			GUObjectPurge.GetObjectsDestroyedSinceLastMarkPhase(),
 			GUnreachableObjects.Num());
 	}
+	else if (bCompleted)
+	{		
+		const int32 PurgedObjectCountSinceLastMarkPhase = GUObjectPurge.GetObjectsDestroyedSinceLastMarkPhase();
+		UE_LOG(LogGarbage, Log, TEXT("GC purged %i objects (%i -> %i) in %.3fms (%d iteration(s))"), PurgedObjectCountSinceLastMarkPhase,
+			GObjectCountDuringLastMarkPhase.GetValue(),
+			GObjectCountDuringLastMarkPhase.GetValue() - PurgedObjectCountSinceLastMarkPhase,
+			GGCStats.DestroyGarbageTime.Total * 1000,
+			GGCStats.DestroyGarbageTime.NumIterations);
+		UE::GC::GDetailedStats.LogPurgeStats(PurgedObjectCountSinceLastMarkPhase);
+		GUObjectPurge.ResetObjectsDestroyedSinceLastMarkPhase();
+	}
 
 	return bCompleted;
 }
