[ZenLoader]
- Fix innacurate handling of sync context cleanup that could cause delays or deadlocks under ALT

diff --git a/Engine/Source/Runtime/CoreUObject/Private/Serialization/AsyncLoading2.cpp b/Engine/Source/Runtime/CoreUObject/Private/Serialization/AsyncLoading2.cpp
--- a/Engine/Source/Runtime/CoreUObject/Private/Serialization/AsyncLoading2.cpp
+++ b/Engine/Source/Runtime/CoreUObject/Private/Serialization/AsyncLoading2.cpp
@@ -4577,10 +4577,18 @@ public:
 			}
 		}
 
-		// Any removed pending request is of interest to main thread as it might unblock a flush.
-		if (RemovedCount > 0 && ThreadState.bIsAsyncLoadingThread)
+		if (RemovedCount > 0)
 		{
-			MainThreadWakeEvent.Notify();
+			// Any removed pending request is of interest to main thread as it might unblock a flush.
+			if (ThreadState.bIsAsyncLoadingThread)
+			{
+				MainThreadWakeEvent.Notify();
+			}
+			// Removal from the GT might unblock the ALT for flushes coming from other threads.
+			else if (IsMultithreaded())
+			{
+				AltZenaphore.NotifyOne();
+			}
 		}
 	}
 
@@ -4780,7 +4788,7 @@ private:
 
 	void IncludePackageInSyncLoadContextRecursive(FAsyncLoadingThreadState2& ThreadState, uint64 ContextId, FAsyncPackage2* Package);
 	void UpdateSyncLoadContext(FAsyncLoadingThreadState2& ThreadState, bool bAutoHandleSyncLoadContext = true);
-
+	void CleanupSyncLoadContext(FAsyncLoadingThreadState2& ThreadState);
 	bool CreateAsyncPackagesFromQueue(FAsyncLoadingThreadState2& ThreadState, bool bForceExitForGarbageCollect = false);
 
 	FAsyncPackage2* CreateAsyncPackage(FAsyncLoadingThreadState2& ThreadState, const FAsyncPackageDesc2& Desc)
@@ -9596,44 +9604,73 @@ FEventLoadNode2& FAsyncPackage2::GetExportBundleNode(EEventLoadNode2 Phase, uint
 	return Data.ExportBundleNodes[ExportBundleNodeIndex];
 }
 
+void FAsyncLoadingThread2::CleanupSyncLoadContext(FAsyncLoadingThreadState2& ThreadState)
+{
+	if (ThreadState.bIsAsyncLoadingThread)
+	{
+		while (!ThreadState.SyncLoadContextStack.IsEmpty())
+		{
+			FAsyncLoadingSyncLoadContext* SyncLoadContext = ThreadState.SyncLoadContextStack.Top();
+
+			if (ContainsAnyRequestID(SyncLoadContext->RequestIDs))
+			{
+				// We still have some processing to do for this request, abort cleanup
+				break;
+			}
+
+			// Request is fulfilled, release it
+			UE_LOG(LogStreaming, VeryVerbose, TEXT("Popping ALT SyncLoadContext %" UINT64_FMT), SyncLoadContext->ContextId);
+			SyncLoadContext->DoneEvent.Notify();
+			SyncLoadContext->ReleaseRef();
+			ThreadState.SyncLoadContextStack.Pop();
+		}
+	}
+}
+
 void FAsyncLoadingThread2::UpdateSyncLoadContext(FAsyncLoadingThreadState2& ThreadState, bool bAutoHandleSyncLoadContext)
 {
 	if (ThreadState.bIsAsyncLoadingThread && bAutoHandleSyncLoadContext)
 	{
+		// Retire complete/invalid contexts for which we aren't loading any requests
+		// before pushing any more contexts on top
+		CleanupSyncLoadContext(ThreadState);
+
+		bool bRequiresSecondCleanup = false;
 		FAsyncLoadingSyncLoadContext* CreatedOnOtherThread;
 		while (ThreadState.SyncLoadContextsCreatedOnOtherThreads.Dequeue(CreatedOnOtherThread))
 		{
 			ThreadState.SyncLoadContextStack.Push(CreatedOnOtherThread);
 
+			bRequiresSecondCleanup = true;
 			UE_LOG(LogStreaming, VeryVerbose, TEXT("Pushing ALT SyncLoadContext %" UINT64_FMT), CreatedOnOtherThread->ContextId);
 		}
+
+		// Retire newly pushed complete/invalid contexts for which we aren't loading any requests
+		if (bRequiresSecondCleanup)
+		{
+			CleanupSyncLoadContext(ThreadState);
+		}
 	}
+	else
+	{
+		// Auto handling is only for async loading thread, all other threads manage
+		// their sync context on their own.
+		bAutoHandleSyncLoadContext = false;
+	}
+
 	if (ThreadState.SyncLoadContextStack.IsEmpty())
 	{
 		return;
 	}
+
 	FAsyncLoadingSyncLoadContext* SyncLoadContext = ThreadState.SyncLoadContextStack.Top();
-	if (ThreadState.bIsAsyncLoadingThread && bAutoHandleSyncLoadContext)
-	{
-		// Retire complete/invalid contexts for which we aren't loading any requests
-		while (!ContainsAnyRequestID(SyncLoadContext->RequestIDs))
-		{
-			UE_LOG(LogStreaming, VeryVerbose, TEXT("Popping ALT SyncLoadContext %" UINT64_FMT), SyncLoadContext->ContextId);
 
-			SyncLoadContext->DoneEvent.Notify();
-			SyncLoadContext->ReleaseRef();
-			ThreadState.SyncLoadContextStack.Pop();
-			if (ThreadState.SyncLoadContextStack.IsEmpty())
-			{
-				return;
-			}
-			SyncLoadContext = ThreadState.SyncLoadContextStack.Top();
-		}
-	}
-	else if (!ContainsAnyRequestID(SyncLoadContext->RequestIDs))
+	// When manually managed, we need to verify if the request on the stack is already fulfilled
+	if (!bAutoHandleSyncLoadContext && !ContainsAnyRequestID(SyncLoadContext->RequestIDs))
 	{
 		return;
 	}
+
 	if (ThreadState.bCanAccessAsyncLoadingThreadData && !SyncLoadContext->bHasFoundRequestedPackages.load(std::memory_order_relaxed))
 	{
 		// Ensure that we've created the package we're waiting for
@@ -9691,6 +9728,7 @@ void FAsyncLoadingThread2::UpdateSyncLoadContext(FAsyncLoadingThreadState2& Thre
 	}
 	if (SyncLoadContext->bHasFoundRequestedPackages.load(std::memory_order_acquire))
 	{
+		bool bMayHaveFulfilledRequest = false;
 		for (int32 i = 0; i < SyncLoadContext->RequestIDs.Num(); ++i)
 		{
 			int32 RequestID = SyncLoadContext->RequestIDs[i];
@@ -9701,8 +9739,14 @@ void FAsyncLoadingThread2::UpdateSyncLoadContext(FAsyncLoadingThreadState2& Thre
 				bool bPreloadIsDone = RequestedPackage->AsyncPackageLoadingState >= EAsyncPackageLoadingState2::DeferredPostLoad;
 				UE_CLOG(!bPreloadIsDone, LogStreaming, Warning, TEXT("Flushing package %s while it's being preloaded in the same callstack is not possible. Releasing request %d to unblock."), *RequestedPackage->Desc.UPackageName.ToString(), RequestID);
 				RemovePendingRequests(ThreadState, {RequestID});
+				bMayHaveFulfilledRequest = true;
 			}
 		}
+
+		if (bAutoHandleSyncLoadContext && bMayHaveFulfilledRequest)
+		{
+			CleanupSyncLoadContext(ThreadState);
+		}
 	}
 }
 
@@ -10396,6 +10440,9 @@ uint32 FAsyncLoadingThread2::Run()
 							{
 								if (EventQueue.ExecuteSyncLoadEvents(ThreadState))
 								{
+									// The first thing we should do after execution is making sure we release
+									// any sync load context that have been fulfilled
+									CleanupSyncLoadContext(ThreadState);
 									bPopped = true;
 									bDidSomething = true;
 								}
