[UObject]
- Fix potential deadlock caused by trying to load assets while holding the Object hash lock

diff --git a/Engine/Source/Runtime/CoreUObject/Private/UObject/Obj.cpp b/Engine/Source/Runtime/CoreUObject/Private/UObject/Obj.cpp
--- a/Engine/Source/Runtime/CoreUObject/Private/UObject/Obj.cpp
+++ b/Engine/Source/Runtime/CoreUObject/Private/UObject/Obj.cpp
@@ -2025,16 +2025,24 @@ void UObject::BuildSubobjectMapping(UObject* OtherObject, TMap<UObject*, UObject
 	UPackage* ThisPackage = GetOutermost();
 	UPackage* OtherPackage = OtherObject->GetOutermost();
 
+	// ForEachObjectWithOuter is taking a hash lock, and we can't load new objects inside a hash lock
+	// so gather objects first, and process them after the hash lock is released.
+	TArray<UObject*> ObjectsToProcess;
 	ForEachObjectWithOuter(this, [&](UObject* InSubObject)
 	{
-		if (ObjectMapping.Contains(InSubObject))
+		ObjectsToProcess.Add(InSubObject);
+	}, false, RF_NoFlags, EInternalObjectFlags::Garbage);
+
+	for (UObject* SubObject : ObjectsToProcess)
+	{
+		if (ObjectMapping.Contains(SubObject))
 		{
 			return;
 		}
 
-		FString NewSubObjectName = InSubObject->GetName();
+		FString NewSubObjectName = SubObject->GetName();
 
-		UClass* OtherSubObjectClass = InSubObject->GetClass();
+		UClass* OtherSubObjectClass = SubObject->GetClass();
 #if WITH_EDITORONLY_DATA
 		if (OtherSubObjectClass->ClassGeneratedBy && OtherSubObjectClass->ClassGeneratedBy->GetOutermost() == ThisPackage)
 		{
@@ -2047,13 +2055,13 @@ void UObject::BuildSubobjectMapping(UObject* OtherObject, TMap<UObject*, UObject
 #endif
 
 		UObject* OtherSubObject = StaticFindObjectFast(OtherSubObjectClass, OtherObject, *NewSubObjectName);
-		ObjectMapping.Emplace(InSubObject, OtherSubObject);
+		ObjectMapping.Emplace(SubObject, OtherSubObject);
 
 		if (OtherSubObject)
 		{
-			InSubObject->BuildSubobjectMapping(OtherSubObject, ObjectMapping);
+			SubObject->BuildSubobjectMapping(OtherSubObject, ObjectMapping);
 		}
-	}, false, RF_NoFlags, EInternalObjectFlags::Garbage);
+	};
 }
 
 
