OFPE: Fix the level package keeping hard references to external entities through subobjects shadow serialization.

diff --git a/Engine/Source/Runtime/CoreUObject/Private/UObject/PropertyArray.cpp b/Engine/Source/Runtime/CoreUObject/Private/UObject/PropertyArray.cpp
--- a/Engine/Source/Runtime/CoreUObject/Private/UObject/PropertyArray.cpp
+++ b/Engine/Source/Runtime/CoreUObject/Private/UObject/PropertyArray.cpp
@@ -512,7 +512,7 @@ void FArrayProperty::SerializeItem(FStructuredArchive::FSlot Slot, void* Value,
 							const int32 ArrayNum = ArrayHelper.Num();
 							for (int i = 0; i < ArrayNum; ++i)
 							{
-								if (UObject* CurrentObject = InnerObjectProperty->GetObjectPropertyValue(ArrayHelper.GetElementPtr(i)); CurrentObject && (!bFilterExternalObjects || !CurrentObject->IsPackageExternal()))
+								if (UObject* CurrentObject = InnerObjectProperty->GetObjectPropertyValue(ArrayHelper.GetElementPtr(i)))
 								{
 									if (ObjectToFind == FOverriddenPropertyNodeID(CurrentObject))
 									{
@@ -537,10 +537,6 @@ void FArrayProperty::SerializeItem(FStructuredArchive::FSlot Slot, void* Value,
 										const int32 DefaultIndex = FindObject(SubNode.NodeID, DefaultsArrayHelper);
 										if (DefaultIndex != INDEX_NONE)
 										{
-											// We don't support removing external entities from overriden arrays yet
-											UObject* DefaultObject = InnerObjectProperty->GetObjectPropertyValue(DefaultsArrayHelper.GetElementPtr(DefaultIndex));
-											checkf(DefaultObject && !DefaultObject->IsPackageExternal(), TEXT("Unsupported remove operation on item %s(0x%p)"), *GetNameSafe(DefaultObject), DefaultObject);
-
 											RemovedIndices.Add(DefaultIndex);
 										}
 										else
@@ -552,21 +548,13 @@ void FArrayProperty::SerializeItem(FStructuredArchive::FSlot Slot, void* Value,
 								case EOverriddenPropertyOperation::Add:
 									{
 										const int32 Index = FindObject(SubNode.NodeID, ArrayHelper);
-										if (Index != INDEX_NONE)
+ 										if (Index != INDEX_NONE)
 										{
-											if (!bFilterExternalObjects)
-											{
-												AddedIndices.Add(Index);
-												ShadowIndices.Remove(Index);
-											}
-											else
+											ShadowIndices.Remove(Index);
+
+											if (!bFilterExternalObjects || !InnerObjectProperty->GetObjectPropertyValue(ArrayHelper.GetElementPtr(Index))->IsPackageExternal())
 											{
-												UObject* Object = InnerObjectProperty->GetObjectPropertyValue(ArrayHelper.GetElementPtr(Index));
-												if(Object && !Object->IsPackageExternal())
-												{
-													AddedIndices.Add(Index);
-													ShadowIndices.Remove(Index);
-												}
+												AddedIndices.Add(Index);												
 											}
 										}
 										else
