Manually update the value of ro.InitRemoteServerIdBeforeUObjectInit from ini because console variables are automatically updated with ini values after InitRemoteObjects is called

diff --git a/Engine/Source/Runtime/CoreUObject/Private/UObject/RemoteObject.cpp b/Engine/Source/Runtime/CoreUObject/Private/UObject/RemoteObject.cpp
--- a/Engine/Source/Runtime/CoreUObject/Private/UObject/RemoteObject.cpp
+++ b/Engine/Source/Runtime/CoreUObject/Private/UObject/RemoteObject.cpp
@@ -17,6 +17,7 @@
 #include "UObject/UObjectMigrationContext.h"
 #include "UObject/UObjectThreadContext.h"
 #include "Misc/CommandLine.h"
+#include "Misc/ConfigCacheIni.h"
 #include "Misc/CoreDelegates.h"
 #include "String/LexFromString.h"
 #include "Hash/CityHash.h"
@@ -257,7 +258,12 @@ void InitRemoteObjects()
 	}
 
 	ObjectMaps = new FRemoteObjectMaps();
-
+	
+	if (GConfig)
+	{
+		// Get the config value now because console vars are updated automatically after InitRemoteObjects is called
+		GConfig->GetInt(TEXT("ConsoleVariables"), TEXT("ro.InitRemoteServerIdBeforeUObjectInit"), GInitRemoteServerIdBeforeUObjectInit, GEngineIni);
+	}
 	if (GInitRemoteServerIdBeforeUObjectInit)
 	{
 		InitServerId();
