Add support for platform specific console variables in cook dependencies :
- New class FConsoleVariableData used to store all the information about console variable dependencies. Also has functions to serialize and resolve the value of the console variable.
- In FCookDependency, add a unique ptr to a FConsoleVariableData and update the class to support this new pointer.

#rnx
#incrementalcook

diff --git a/Engine/Source/Editor/UnrealEd/Private/Cooker/CookIncrementalVersion.cpp b/Engine/Source/Editor/UnrealEd/Private/Cooker/CookIncrementalVersion.cpp
--- a/Engine/Source/Editor/UnrealEd/Private/Cooker/CookIncrementalVersion.cpp
+++ b/Engine/Source/Editor/UnrealEd/Private/Cooker/CookIncrementalVersion.cpp
@@ -8,6 +8,6 @@ namespace UE::Cook
 {
 
 // Change CookIncrementalVersion to a new guid when all packages in an incremental cook need to be invalidated
-FGuid CookIncrementalVersion(0x8EB02510, 0xA8414F57, 0x9D68A4A8, 0x2F8BC58E);
+FGuid CookIncrementalVersion(0x93DE0E19, 0x721143D5, 0x9487FF18, 0xEF03D6B1);
 
 }
diff --git a/Engine/Source/Runtime/CoreUObject/Private/Cooker/ConsoleVariableData.cpp b/Engine/Source/Runtime/CoreUObject/Private/Cooker/ConsoleVariableData.cpp
new file mode 100644
--- /dev/null
+++ b/Engine/Source/Runtime/CoreUObject/Private/Cooker/ConsoleVariableData.cpp
@@ -0,0 +1,176 @@
+// Copyright Epic Games, Inc. All Rights Reserved.
+
+#if WITH_EDITOR
+
+#include "ConsoleVariableData.h"
+
+#include "HAL/IConsoleManager.h"
+#include "Interfaces/ITargetPlatform.h"
+#include "Interfaces/ITargetPlatformManagerModule.h"
+#include "Misc/CoreMisc.h"
+#include "Serialization/CompactBinarySerialization.h"
+#include "Serialization/CompactBinaryWriter.h"
+
+namespace UE::Cook
+{
+	FConsoleVariableData::FConsoleVariableData()
+		: VariableName()
+		, TargetPlatform(nullptr)
+		, bFallbackToNonPlatformValue(false)
+	{ }
+
+	FConsoleVariableData::FConsoleVariableData(const FStringView& InVariableName, const ITargetPlatform* InTargetPlatform, bool bInFallbackToNonPlatformValue)
+		: VariableName(InVariableName)
+		, TargetPlatform(InTargetPlatform)
+		, bFallbackToNonPlatformValue(bInFallbackToNonPlatformValue)
+	{ }
+
+	void FConsoleVariableData::Save(FCbWriter& Writer) const
+	{
+		Writer.BeginArray();
+		Writer << VariableName;
+
+		bool bHasTargetPlatform = TargetPlatform != nullptr;
+		Writer << bHasTargetPlatform;
+
+		if (bHasTargetPlatform)
+		{
+			Writer << TargetPlatform->PlatformName();
+		}
+
+		Writer << bFallbackToNonPlatformValue;
+
+		Writer.EndArray();
+	}
+
+	bool FConsoleVariableData::TryLoad(FCbFieldView Value)
+	{
+		if (!Value.IsArray())
+		{
+			return false;
+		}
+
+		FCbFieldViewIterator ArrayIterator = Value.CreateViewIterator();
+
+		if (!LoadFromCompactBinary(ArrayIterator++, VariableName))
+		{
+			return false;
+		}
+
+		bool bHasTargetPlatform = false;
+		if (!LoadFromCompactBinary(ArrayIterator++, bHasTargetPlatform))
+		{
+			return false;
+		}
+
+		if (bHasTargetPlatform)
+		{
+			FString TargetPlatformName;
+			if (!LoadFromCompactBinary(ArrayIterator++, TargetPlatformName))
+			{
+				return false;
+			}
+
+			ITargetPlatformManagerModule* TargetPlatformManager = GetTargetPlatformManager();
+			TargetPlatform = TargetPlatformManager->FindTargetPlatform(TargetPlatformName);
+			if (!TargetPlatform)
+			{
+				return false;
+			}
+		}
+
+		if (!LoadFromCompactBinary(ArrayIterator++, bFallbackToNonPlatformValue))
+		{
+			return false;
+		}
+
+		return true;
+	}
+
+	FStringView FConsoleVariableData::GetConsoleVariableName() const
+	{
+		return VariableName;
+	}
+
+	bool FConsoleVariableData::TryResolveValue(FString& Value) const
+	{
+		IConsoleVariable* VariableInstance = IConsoleManager::Get().FindConsoleVariable(*VariableName, false);
+		if (VariableInstance == nullptr)
+		{
+			return false;
+		}
+
+		if (TargetPlatform)
+		{
+			FName IniPlatformName = *TargetPlatform->IniPlatformName();
+
+			//don't call directly GetPlatformValueVariable because it will create the cvar
+			bool bHasPlaftormValue = VariableInstance->HasPlatformValueVariable(IniPlatformName);
+			if (bHasPlaftormValue)
+			{
+				if (const IConsoleVariable* PlatformCVar = VariableInstance->GetPlatformValueVariable(IniPlatformName).Get())
+				{
+					Value = PlatformCVar->GetString();
+				}
+				else
+				{
+					return false;
+				}
+			}
+			else if (bFallbackToNonPlatformValue)
+			{
+				Value = VariableInstance->GetString();
+			}
+			else
+			{
+				return false;
+			}
+		}
+		else
+		{
+			Value = VariableInstance->GetString();
+		}
+
+		return true;
+	}
+
+	bool FConsoleVariableData::operator<(const FConsoleVariableData& Other) const
+	{
+		int32 CompareResult = VariableName.Compare(Other.VariableName);
+		if (CompareResult != 0)
+		{
+			return CompareResult < 0;
+		}
+
+		if (TargetPlatform != Other.TargetPlatform)
+		{
+			if (!TargetPlatform)
+			{
+				return true;
+			}
+			else if (!Other.TargetPlatform)
+			{
+				return false;
+			}
+			else
+			{
+				CompareResult = TargetPlatform->PlatformName().Compare(Other.TargetPlatform->PlatformName());
+				return CompareResult < 0;
+			}
+		}
+
+		if (bFallbackToNonPlatformValue != Other.bFallbackToNonPlatformValue)
+		{
+			return bFallbackToNonPlatformValue;
+		}
+
+		return false;
+	}
+
+	bool FConsoleVariableData::operator==(const FConsoleVariableData& Other) const
+	{
+		return VariableName == Other.VariableName && TargetPlatform == Other.TargetPlatform && bFallbackToNonPlatformValue == Other.bFallbackToNonPlatformValue;
+	}
+}
+
+#endif // #if WITH_EDITOR
diff --git a/Engine/Source/Runtime/CoreUObject/Private/Cooker/ConsoleVariableData.h b/Engine/Source/Runtime/CoreUObject/Private/Cooker/ConsoleVariableData.h
new file mode 100644
--- /dev/null
+++ b/Engine/Source/Runtime/CoreUObject/Private/Cooker/ConsoleVariableData.h
@@ -0,0 +1,52 @@
+// Copyright Epic Games, Inc. All Rights Reserved.
+
+#pragma once
+
+#if WITH_EDITOR
+
+#include "Containers/StringFwd.h"
+#include "Serialization/CompactBinary.h"
+
+class FCbWriter;
+class ITargetPlatform;
+
+namespace UE::Cook
+{
+
+/* Describe how a console variable is loaded for FCookDependency. */
+class FConsoleVariableData
+{
+public:
+	FConsoleVariableData();
+	FConsoleVariableData(const FStringView& InVariableName, const ITargetPlatform* InTargetPlatform, bool bInFallbackToNonPlatformValue);
+
+	void Save(FCbWriter& Writer) const;
+	bool TryLoad(FCbFieldView Value);
+
+	FStringView GetConsoleVariableName() const;
+
+	bool TryResolveValue(FString& Value) const;
+
+	bool operator<(const FConsoleVariableData& Other) const;
+	bool operator==(const FConsoleVariableData& Other) const;
+
+private:
+	FString VariableName;
+	const ITargetPlatform* TargetPlatform;		// If non-null, query the platform specific value
+	bool bFallbackToNonPlatformValue;			// If querying the platform specific value but it's not found then fallback to the non-platform value.
+
+	friend FCbWriter& operator<<(FCbWriter& Writer, const UE::Cook::FConsoleVariableData& Data)
+	{
+		Data.Save(Writer);
+		return Writer;
+	}
+
+	friend bool LoadFromCompactBinary(FCbFieldView Value, UE::Cook::FConsoleVariableData& Data)
+	{
+		return Data.TryLoad(Value);
+	}
+};
+
+}
+
+#endif // #if WITH_EDITOR
diff --git a/Engine/Source/Runtime/CoreUObject/Private/Cooker/CookDependency.cpp b/Engine/Source/Runtime/CoreUObject/Private/Cooker/CookDependency.cpp
--- a/Engine/Source/Runtime/CoreUObject/Private/Cooker/CookDependency.cpp
+++ b/Engine/Source/Runtime/CoreUObject/Private/Cooker/CookDependency.cpp
@@ -8,6 +8,7 @@
 #include "AssetRegistry/AssetData.h"
 #include "Containers/Array.h"
 #include "Containers/Map.h"
+#include "Cooker/ConsoleVariableData.h"
 #include "CoreGlobals.h"
 #include "HAL/FileManager.h"
 #include "HAL/IConsoleManager.h"
@@ -118,10 +119,15 @@ FCookDependency FCookDependency::Package(FName PackageName)
 	return Result;
 }
 
-FCookDependency FCookDependency::ConsoleVariable(FStringView VariableName)
+FCookDependency FCookDependency::ConsoleVariable(const FStringView& VariableName)
+{
+	return FCookDependency::ConsoleVariable(VariableName, nullptr, false);
+}
+
+FCookDependency FCookDependency::ConsoleVariable(const FStringView& VariableName, const ITargetPlatform* TargetPlatform, bool bFallbackToNonPlatform)
 {
 	FCookDependency Result(ECookDependency::ConsoleVariable);
-	Result.StringData = VariableName;
+	Result.ConsoleVariableData = MakeUnique<FConsoleVariableData>(VariableName, TargetPlatform, bFallbackToNonPlatform);
 	return Result;
 }
 
@@ -272,10 +278,12 @@ FCookDependency& FCookDependency::operator=(const FCookDependency& Other)
 	case ECookDependency::None:
 		break;
 	case ECookDependency::File:
-	case ECookDependency::ConsoleVariable:
 	case ECookDependency::NativeClass:
 		StringData = Other.StringData;
 		break;
+	case ECookDependency::ConsoleVariable:
+		ConsoleVariableData.Reset(new FConsoleVariableData(*Other.ConsoleVariableData));
+		break;
 	case ECookDependency::Function:
 		FunctionData = Other.FunctionData;
 		break;
@@ -317,10 +325,12 @@ FCookDependency& FCookDependency::operator=(FCookDependency&& Other)
 	case ECookDependency::None:
 		break;
 	case ECookDependency::File:
-	case ECookDependency::ConsoleVariable:
 	case ECookDependency::NativeClass:
 		StringData = MoveTemp(Other.StringData);
 		break;
+	case ECookDependency::ConsoleVariable:
+		ConsoleVariableData = MoveTemp(Other.ConsoleVariableData);
+		break;
 	case ECookDependency::Function:
 		FunctionData = MoveTemp(Other.FunctionData);
 		break;
@@ -356,9 +366,10 @@ FString FCookDependency::GetDebugIdentifier() const
 	case ECookDependency::None:
 		return FString();
 	case ECookDependency::File:
-	case ECookDependency::ConsoleVariable:
 	case ECookDependency::NativeClass:
 		return StringData;
+	case ECookDependency::ConsoleVariable:
+		return ConsoleVariableData.IsValid() ? FString(ConsoleVariableData->GetConsoleVariableName()) : FString();
 	case ECookDependency::Function:
 		return FunctionData.Name.ToString();
 	case ECookDependency::TransitiveBuild:
@@ -499,21 +510,29 @@ void FCookDependency::UpdateHash(FCookDependencyContext& Context, FBlake3Hash* O
 		return;
 	case ECookDependency::ConsoleVariable:
 	{
-		IConsoleVariable* VariableInstance = IConsoleManager::Get().FindConsoleVariable(*StringData, false /* bTrackFrequentCalls */);
-		if (VariableInstance == nullptr)
+		if (!ConsoleVariableData.IsValid())
 		{
+			return;
+		}
+
+		FString Value;
+		bool bVariableFound = ConsoleVariableData->TryResolveValue(Value);
+
+		if (!bVariableFound)
+		{
+			FStringView ConsoleVariableName = ConsoleVariableData->GetConsoleVariableName();
 			Context.LogError(FString::Printf(
-				TEXT("FCookDependency::ConsoleVariable('%s') failed to UpdateHash: could not find console variable."),
-				*StringData));
+				TEXT("FCookDependency::ConsoleVariable('%.*s') failed to UpdateHash: could not find console variable."),
+				ConsoleVariableName.Len(), ConsoleVariableName.GetData()));
+
 			if (OutCurrentValue)
 			{
 				*OutCurrentValue = FBlake3Hash();
 			}
 			return;
 		}
-		FString VariableAsString = VariableInstance->GetString();
 
-		FBlake3Hash Hash = ConvertToHash(FUtf8String(VariableAsString));
+		FBlake3Hash Hash = ConvertToHash(FUtf8String(Value));
 		if (OutCurrentValue)
 		{
 			*OutCurrentValue = Hash;
@@ -649,10 +668,13 @@ void FCookDependency::Construct()
 	case ECookDependency::None:
 		break;
 	case ECookDependency::File:
-	case ECookDependency::ConsoleVariable:
 	case ECookDependency::NativeClass:
 		new(&StringData) FString();
 		break;
+	case ECookDependency::ConsoleVariable:
+		// Set the union's bytes equal to a TUniquePtr with a null internal pointer
+		new(&ConsoleVariableData) TUniquePtr<FConsoleVariableData>(nullptr);
+		break;
 	case ECookDependency::Function:
 		new(&FunctionData) FFunctionData();
 		break;
@@ -689,10 +711,12 @@ void FCookDependency::Destruct()
 	case ECookDependency::None:
 		break;
 	case ECookDependency::File:
-	case ECookDependency::ConsoleVariable:
 	case ECookDependency::NativeClass:
 		StringData.~FString();
 		break;
+	case ECookDependency::ConsoleVariable:
+		ConsoleVariableData.~TUniquePtr<FConsoleVariableData>();
+		break;
 	case ECookDependency::Function:
 		FunctionData.~FFunctionData();
 		break;
@@ -731,10 +755,20 @@ void FCookDependency::Save(FCbWriter& Writer) const
 	case ECookDependency::None:
 		break;
 	case ECookDependency::File:
-	case ECookDependency::ConsoleVariable:
 	case ECookDependency::NativeClass:
 		Writer << StringData;
 		break;
+	case ECookDependency::ConsoleVariable:
+	{
+		bool bIsValid = ConsoleVariableData.IsValid();
+		Writer << bIsValid;
+
+		if (bIsValid)
+		{
+			Writer << *ConsoleVariableData;
+		}
+		break;
+	}
 	case ECookDependency::Function:
 		Writer << FunctionData.Name;
 		Writer.BeginArray();
@@ -859,12 +893,25 @@ bool FCookDependency::Load(FCbFieldView Value)
 
 	case ECookDependency::ConsoleVariable:
 	{
-		FString LocalConsoleVariableName;
-		if (!LoadFromCompactBinary(Field++, LocalConsoleVariableName))
+		bool bIsValid = false;
+		if (!LoadFromCompactBinary(Field++, bIsValid))
+		{
+			return false;
+		}
+
+		if (!bIsValid)
+		{
+			return false;
+		}
+
+		TUniquePtr<FConsoleVariableData> LocalConsoleVariableData = MakeUnique<FConsoleVariableData>();
+		if (!LoadFromCompactBinary(Field++, *LocalConsoleVariableData))
 		{
 			return false;
 		}
-		*this = FCookDependency::ConsoleVariable(LocalConsoleVariableName);
+
+		*this = FCookDependency(ECookDependency::ConsoleVariable);
+		ConsoleVariableData = MoveTemp(LocalConsoleVariableData);
 	}
 	break;
 
@@ -975,6 +1022,16 @@ bool FCookDependency::ARFilterEqual(const FARFilter& A, const FARFilter& B)
 	return A == B;
 }
 
+bool FCookDependency::ConsoleVariableDataLessThan(const FConsoleVariableData& A, const FConsoleVariableData& B)
+{
+	return A < B;
+}
+
+bool FCookDependency::ConsoleVariableDataEqual(const FConsoleVariableData& A, const FConsoleVariableData& B)
+{
+	return A == B;
+}
+
 void FCookDependencyContext::Update(const void* Data, uint64 Size)
 {
 	FBlake3& Blake3Hasher = *(reinterpret_cast<FBlake3*>(Hasher));
diff --git a/Engine/Source/Runtime/CoreUObject/Public/Cooker/CookDependency.h b/Engine/Source/Runtime/CoreUObject/Public/Cooker/CookDependency.h
--- a/Engine/Source/Runtime/CoreUObject/Public/Cooker/CookDependency.h
+++ b/Engine/Source/Runtime/CoreUObject/Public/Cooker/CookDependency.h
@@ -18,6 +18,7 @@ class FCbWriter;
 class UObject;
 namespace UE::ConfigAccessTracking { enum class ELoadType : uint8; }
 namespace UE::ConfigAccessTracking { struct FConfigAccessData; }
+namespace UE::Cook { class FConsoleVariableData; }
 struct FARFilter;
 #endif
 
@@ -230,7 +231,8 @@ public:
 	 * Create a dependency on the value of a cvar. The cvar will be read and its value (as a string) will be hashed into the oplog data
 	 * If the cvar value is changed, the packages that depend on it will be invalidated
 	 */
-	COREUOBJECT_API static FCookDependency ConsoleVariable(FStringView VariableName);
+	COREUOBJECT_API static FCookDependency ConsoleVariable(const FStringView& VariableName);
+	COREUOBJECT_API static FCookDependency ConsoleVariable(const FStringView& VariableName, const ITargetPlatform* TargetPlatform, bool bFallbackToNonPlatform);
 
 	/** Create a dependency on the value of a config variable. */
 	COREUOBJECT_API static FCookDependency Config(UE::ConfigAccessTracking::FConfigAccessData AccessData);
@@ -349,6 +351,9 @@ private:
 	static COREUOBJECT_API bool ARFilterLessThan(const FARFilter& A, const FARFilter& B);
 	static COREUOBJECT_API bool ARFilterEqual(const FARFilter& A, const FARFilter& B);
 
+	COREUOBJECT_API static bool ConsoleVariableDataLessThan(const FConsoleVariableData& A, const FConsoleVariableData& B);
+	COREUOBJECT_API static bool ConsoleVariableDataEqual(const FConsoleVariableData& A, const FConsoleVariableData& B);
+
 	/** Public hidden friend for operator<< into an FCbWriter. */
 	friend FCbWriter& operator<<(FCbWriter& Writer, const FCookDependency& CookDependency)
 	{
@@ -382,6 +387,7 @@ private:
 		const UObject* ObjectPtr;
 		TUniquePtr<UE::ConfigAccessTracking::FConfigAccessData> ConfigAccessData;
 		TUniquePtr<FARFilter> ARFilter;
+		TUniquePtr<FConsoleVariableData> ConsoleVariableData;
 	};
 
 	/** 
@@ -563,9 +569,24 @@ inline bool FCookDependency::operator<(const FCookDependency& Other) const
 	case ECookDependency::None:
 		return false;
 	case ECookDependency::File:
-	case ECookDependency::ConsoleVariable:
 	case ECookDependency::NativeClass:
 		return StringData.Compare(Other.StringData, ESearchCase::IgnoreCase) < 0;
+
+	case ECookDependency::ConsoleVariable:
+	{
+		if (ConsoleVariableData.IsValid() != Other.ConsoleVariableData.IsValid())
+		{
+			return !ConsoleVariableData.IsValid();
+		}
+
+		if (ConsoleVariableData.IsValid())
+		{
+			return ConsoleVariableDataLessThan(*ConsoleVariableData, *Other.ConsoleVariableData);
+		}
+
+		return false; // Both pointers are invalid
+	}
+
 	case ECookDependency::Function:
 	{
 		int32 Compare = FunctionData.Name.Compare(Other.FunctionData.Name);
@@ -639,9 +660,24 @@ inline bool FCookDependency::operator==(const FCookDependency& Other) const
 	case ECookDependency::None:
 		return true;
 	case ECookDependency::File:
-	case ECookDependency::ConsoleVariable:
 	case ECookDependency::NativeClass:
 		return StringData.Compare(Other.StringData, ESearchCase::IgnoreCase) == 0;
+
+	case ECookDependency::ConsoleVariable:
+	{
+		if (ConsoleVariableData.IsValid() != Other.ConsoleVariableData.IsValid())
+		{
+			return false;
+		}
+
+		if (!ConsoleVariableData.IsValid())
+		{
+			return true;
+		}
+
+		return ConsoleVariableDataEqual(*ConsoleVariableData, *Other.ConsoleVariableData);
+	}
+
 	case ECookDependency::Function:
 	{
 		if (FunctionData.Name.Compare(Other.FunctionData.Name) != 0)
