[UE] Add LowLevelTests for subobject instancing ground truth and regression testing around related APIs.

diff --git a/Engine/Source/Runtime/CoreUObject/Private/UObject/CoreNative.cpp b/Engine/Source/Runtime/CoreUObject/Private/UObject/CoreNative.cpp
--- a/Engine/Source/Runtime/CoreUObject/Private/UObject/CoreNative.cpp
+++ b/Engine/Source/Runtime/CoreUObject/Private/UObject/CoreNative.cpp
@@ -251,8 +251,9 @@ UObject* FObjectInstancingGraph::GetInstancedSubobject( UObject* SourceSubobject
 						const bool bShouldLoadForClient = SourceSubobject->NeedsLoadForClient();
 						const bool bShouldLoadForServer = SourceSubobject->NeedsLoadForServer();
 						const bool bShouldLoadForEditor = ( GIsEditor && ( bShouldLoadForClient || !CurrentObject->RootPackageHasAnyFlags(PKG_PlayInEditor) ) );
+						const bool bShouldLoadForStandalone = (!GIsEditor && !GIsClient && !GIsServer);
 
-						if ( ((GIsClient && bShouldLoadForClient) || (GIsServer && bShouldLoadForServer) || bShouldLoadForEditor) )
+						if ( (GIsClient && bShouldLoadForClient) || (GIsServer && bShouldLoadForServer) || bShouldLoadForEditor || bShouldLoadForStandalone )
 						{
 							// this is the first time the instance corresponding to SourceSubobject has been requested
 
diff --git a/Engine/Source/Runtime/CoreUObject/Tests/UObject/InstancedPropertiesTest.cpp b/Engine/Source/Runtime/CoreUObject/Tests/UObject/InstancedPropertiesTest.cpp
new file mode 100644
--- /dev/null
+++ b/Engine/Source/Runtime/CoreUObject/Tests/UObject/InstancedPropertiesTest.cpp
@@ -0,0 +1,464 @@
+// Copyright Epic Games, Inc. All Rights Reserved.
+
+#if WITH_TESTS
+
+#include "Misc/AssertionMacros.h"
+#include "SubobjectInstancingTest.h"
+#include "Tests/TestHarnessAdapter.h"
+#include "UObject/PropertyOptional.h"
+#include "UObject/UnrealType.h"
+
+namespace UE
+{
+
+TEST_CASE_NAMED(FInstancedPropertiesTest, "CoreUObject::InstancedProperties", "[CoreUObject][EngineFilter]")
+{
+	SECTION("Instanced object property")
+	{
+		const FName InstancedObjectPropertyName = GET_MEMBER_NAME_CHECKED(USubobjectInstancingTestOuterObject, InnerObject);
+		const FProperty* InstancedObjectProperty = USubobjectInstancingTestOuterObject::StaticClass()->FindPropertyByName(InstancedObjectPropertyName);
+
+		REQUIRE(InstancedObjectProperty != nullptr);
+		CHECK(InstancedObjectProperty->IsA<FObjectProperty>());
+
+		CHECK(InstancedObjectProperty->ContainsInstancedObjectProperty());
+		CHECK(InstancedObjectProperty->HasAllPropertyFlags(CPF_InstancedReference));
+		CHECK(InstancedObjectProperty->HasAllPropertyFlags(CPF_PersistentInstance));
+#if WITH_METADATA
+		CHECK(InstancedObjectProperty->GetBoolMetaData(TEXT("EditInline")));
+#endif
+	}
+
+	SECTION("Non-instanced object property")
+	{
+		const FName NonInstancedObjectPropertyName = GET_MEMBER_NAME_CHECKED(USubobjectInstancingTestOuterObject, ExternalObject);
+		const FProperty* NonInstancedObjectProperty = USubobjectInstancingTestOuterObject::StaticClass()->FindPropertyByName(NonInstancedObjectPropertyName);
+
+		REQUIRE(NonInstancedObjectProperty != nullptr);
+		CHECK(NonInstancedObjectProperty->IsA<FObjectProperty>());
+
+		CHECK_FALSE(NonInstancedObjectProperty->ContainsInstancedObjectProperty());
+		CHECK_FALSE(NonInstancedObjectProperty->HasAllPropertyFlags(CPF_InstancedReference));
+		CHECK_FALSE(NonInstancedObjectProperty->HasAllPropertyFlags(CPF_PersistentInstance));
+#if WITH_METADATA
+		CHECK_FALSE(NonInstancedObjectProperty->GetBoolMetaData(TEXT("EditInline")));
+#endif
+	}
+
+	SECTION("DefaultToInstanced object property")
+	{
+		const FName DefaultToInstancedObjectPropertyName = GET_MEMBER_NAME_CHECKED(USubobjectInstancingTestOuterObject, InnerObjectFromType);
+		const FProperty* DefaultToInstancedObjectProperty = USubobjectInstancingTestOuterObject::StaticClass()->FindPropertyByName(DefaultToInstancedObjectPropertyName);
+
+		REQUIRE(DefaultToInstancedObjectProperty != nullptr);
+		CHECK(DefaultToInstancedObjectProperty->IsA<FObjectProperty>());
+
+		CHECK(DefaultToInstancedObjectProperty->ContainsInstancedObjectProperty());
+		CHECK(DefaultToInstancedObjectProperty->HasAllPropertyFlags(CPF_InstancedReference));
+		CHECK_FALSE(DefaultToInstancedObjectProperty->HasAllPropertyFlags(CPF_PersistentInstance));
+#if WITH_METADATA
+		// Note: DefaultToInstanced also implies 'EditInline' (via UHT)
+		CHECK(DefaultToInstancedObjectProperty->GetBoolMetaData(TEXT("EditInline")));
+#endif
+	}
+
+	SECTION("Transient instanced object property")
+	{
+		const FName TransientInstancedObjectPropertyName = GET_MEMBER_NAME_CHECKED(USubobjectInstancingTestOuterObject, TransientInnerObject);
+		const FProperty* TransientInstancedObjectProperty = USubobjectInstancingTestOuterObject::StaticClass()->FindPropertyByName(TransientInstancedObjectPropertyName);
+
+		REQUIRE(TransientInstancedObjectProperty != nullptr);
+		CHECK(TransientInstancedObjectProperty->IsA<FObjectProperty>());
+
+		CHECK(TransientInstancedObjectProperty->ContainsInstancedObjectProperty());
+		CHECK(TransientInstancedObjectProperty->HasAllPropertyFlags(CPF_Transient));
+		CHECK(TransientInstancedObjectProperty->HasAllPropertyFlags(CPF_InstancedReference));
+		CHECK(TransientInstancedObjectProperty->HasAllPropertyFlags(CPF_PersistentInstance));
+#if WITH_METADATA
+		CHECK(TransientInstancedObjectProperty->GetBoolMetaData(TEXT("EditInline")));
+#endif
+	}
+
+	SECTION("Instanced object array property")
+	{
+		const FName InstancedObjectArrayPropertyName = GET_MEMBER_NAME_CHECKED(USubobjectInstancingTestOuterObject, InnerObjectArray);
+		const FArrayProperty* InstancedObjectArrayProperty = CastField<FArrayProperty>(USubobjectInstancingTestOuterObject::StaticClass()->FindPropertyByName(InstancedObjectArrayPropertyName));
+
+		REQUIRE(InstancedObjectArrayProperty != nullptr);
+		CHECK(InstancedObjectArrayProperty->ContainsInstancedObjectProperty());
+		CHECK(InstancedObjectArrayProperty->HasAllPropertyFlags(CPF_ContainsInstancedReference));
+		CHECK_FALSE(InstancedObjectArrayProperty->HasAllPropertyFlags(CPF_InstancedReference));
+		CHECK_FALSE(InstancedObjectArrayProperty->HasAllPropertyFlags(CPF_PersistentInstance));
+#if WITH_METADATA
+		CHECK(InstancedObjectArrayProperty->GetBoolMetaData(TEXT("EditInline")));
+#endif
+
+		REQUIRE(InstancedObjectArrayProperty->Inner != nullptr);
+		CHECK(InstancedObjectArrayProperty->Inner->IsA<FObjectProperty>());
+
+		CHECK(InstancedObjectArrayProperty->Inner->ContainsInstancedObjectProperty());
+		CHECK_FALSE(InstancedObjectArrayProperty->Inner->HasAllPropertyFlags(CPF_ContainsInstancedReference));
+		CHECK(InstancedObjectArrayProperty->Inner->HasAllPropertyFlags(CPF_InstancedReference));
+		CHECK(InstancedObjectArrayProperty->Inner->HasAllPropertyFlags(CPF_PersistentInstance));
+#if WITH_METADATA
+		CHECK(InstancedObjectArrayProperty->Inner->GetBoolMetaData(TEXT("EditInline")));
+#endif
+	}
+
+	SECTION("DefaultToInstanced object array property")
+	{
+		const FName DefaultToInstancedObjectArrayPropertyName = GET_MEMBER_NAME_CHECKED(USubobjectInstancingTestOuterObject, InnerObjectFromTypeArray);
+		const FArrayProperty* DefaultToInstancedObjectArrayProperty = CastField<FArrayProperty>(USubobjectInstancingTestOuterObject::StaticClass()->FindPropertyByName(DefaultToInstancedObjectArrayPropertyName));
+
+		REQUIRE(DefaultToInstancedObjectArrayProperty != nullptr);
+		CHECK(DefaultToInstancedObjectArrayProperty->ContainsInstancedObjectProperty());
+		CHECK(DefaultToInstancedObjectArrayProperty->HasAllPropertyFlags(CPF_ContainsInstancedReference));
+		CHECK_FALSE(DefaultToInstancedObjectArrayProperty->HasAllPropertyFlags(CPF_InstancedReference));
+		CHECK_FALSE(DefaultToInstancedObjectArrayProperty->HasAllPropertyFlags(CPF_PersistentInstance));
+#if WITH_METADATA
+		// Note: DefaultToInstanced also implies 'EditInline' (via UHT). 
+		CHECK(DefaultToInstancedObjectArrayProperty->GetBoolMetaData(TEXT("EditInline")));
+#endif
+
+		REQUIRE(DefaultToInstancedObjectArrayProperty->Inner != nullptr);
+		CHECK(DefaultToInstancedObjectArrayProperty->Inner->IsA<FObjectProperty>());
+
+		CHECK(DefaultToInstancedObjectArrayProperty->Inner->ContainsInstancedObjectProperty());
+		CHECK_FALSE(DefaultToInstancedObjectArrayProperty->Inner->HasAllPropertyFlags(CPF_ContainsInstancedReference));
+		CHECK(DefaultToInstancedObjectArrayProperty->Inner->HasAllPropertyFlags(CPF_InstancedReference));
+		CHECK_FALSE(DefaultToInstancedObjectArrayProperty->Inner->HasAllPropertyFlags(CPF_PersistentInstance));
+#if WITH_METADATA
+		// Note: Unlike explicitly-instanced array properties, the Inner is not also flagged as 'EditInline'.
+		CHECK_FALSE(DefaultToInstancedObjectArrayProperty->Inner->GetBoolMetaData(TEXT("EditInline")));
+#endif
+	}
+
+	SECTION("Instanced object set container property")
+	{
+		const FName InstancedObjectSetPropertyName = GET_MEMBER_NAME_CHECKED(USubobjectInstancingTestOuterObject, InnerObjectSet);
+		const FSetProperty* InstancedObjectSetProperty = CastField<FSetProperty>(USubobjectInstancingTestOuterObject::StaticClass()->FindPropertyByName(InstancedObjectSetPropertyName));
+
+		REQUIRE(InstancedObjectSetProperty != nullptr);
+		CHECK(InstancedObjectSetProperty->ContainsInstancedObjectProperty());
+		CHECK(InstancedObjectSetProperty->HasAllPropertyFlags(CPF_ContainsInstancedReference));
+		CHECK_FALSE(InstancedObjectSetProperty->HasAllPropertyFlags(CPF_InstancedReference));
+		CHECK_FALSE(InstancedObjectSetProperty->HasAllPropertyFlags(CPF_PersistentInstance));
+#if WITH_METADATA
+		CHECK(InstancedObjectSetProperty->GetBoolMetaData(TEXT("EditInline")));
+#endif
+
+		REQUIRE(InstancedObjectSetProperty->ElementProp != nullptr);
+		CHECK(InstancedObjectSetProperty->ElementProp->IsA<FObjectProperty>());
+
+		CHECK(InstancedObjectSetProperty->ElementProp->ContainsInstancedObjectProperty());
+		CHECK_FALSE(InstancedObjectSetProperty->ElementProp->HasAllPropertyFlags(CPF_ContainsInstancedReference));
+		CHECK(InstancedObjectSetProperty->ElementProp->HasAllPropertyFlags(CPF_InstancedReference));
+		CHECK(InstancedObjectSetProperty->ElementProp->HasAllPropertyFlags(CPF_PersistentInstance));
+#if WITH_METADATA
+		CHECK(InstancedObjectSetProperty->ElementProp->GetBoolMetaData(TEXT("EditInline")));
+#endif
+	}
+
+	SECTION("DefaultToInstanced object set container property")
+	{
+		const FName DefaultToInstancedObjectSetPropertyName = GET_MEMBER_NAME_CHECKED(USubobjectInstancingTestOuterObject, InnerObjectFromTypeSet);
+		const FSetProperty* DefaultToInstancedObjectSetProperty = CastField<FSetProperty>(USubobjectInstancingTestOuterObject::StaticClass()->FindPropertyByName(DefaultToInstancedObjectSetPropertyName));
+
+		REQUIRE(DefaultToInstancedObjectSetProperty != nullptr);
+		CHECK(DefaultToInstancedObjectSetProperty->ContainsInstancedObjectProperty());
+		CHECK(DefaultToInstancedObjectSetProperty->HasAllPropertyFlags(CPF_ContainsInstancedReference));
+		CHECK_FALSE(DefaultToInstancedObjectSetProperty->HasAllPropertyFlags(CPF_InstancedReference));
+		CHECK_FALSE(DefaultToInstancedObjectSetProperty->HasAllPropertyFlags(CPF_PersistentInstance));
+#if WITH_METADATA
+		// Note: DefaultToInstanced also implies 'EditInline' (via UHT).
+		CHECK(DefaultToInstancedObjectSetProperty->GetBoolMetaData(TEXT("EditInline")));
+#endif
+
+		REQUIRE(DefaultToInstancedObjectSetProperty->ElementProp != nullptr);
+		CHECK(DefaultToInstancedObjectSetProperty->ElementProp->IsA<FObjectProperty>());
+
+		CHECK(DefaultToInstancedObjectSetProperty->ElementProp->ContainsInstancedObjectProperty());
+		CHECK_FALSE(DefaultToInstancedObjectSetProperty->ElementProp->HasAllPropertyFlags(CPF_ContainsInstancedReference));
+		CHECK(DefaultToInstancedObjectSetProperty->ElementProp->HasAllPropertyFlags(CPF_InstancedReference));
+		CHECK_FALSE(DefaultToInstancedObjectSetProperty->ElementProp->HasAllPropertyFlags(CPF_PersistentInstance));
+#if WITH_METADATA
+		// Note: Unlike explicitly-instanced set properties, the ElementProp is not also flagged as 'EditInline'.
+		CHECK_FALSE(DefaultToInstancedObjectSetProperty->ElementProp->GetBoolMetaData(TEXT("EditInline")));
+#endif
+	}
+
+	SECTION("Instanced object map container property")
+	{
+		const FName InstancedObjectMapPropertyName = GET_MEMBER_NAME_CHECKED(USubobjectInstancingTestOuterObject, InnerObjectMap);
+		const FMapProperty* InstancedObjectMapProperty = CastField<FMapProperty>(USubobjectInstancingTestOuterObject::StaticClass()->FindPropertyByName(InstancedObjectMapPropertyName));
+
+		REQUIRE(InstancedObjectMapProperty != nullptr);
+		CHECK(InstancedObjectMapProperty->ContainsInstancedObjectProperty());
+		CHECK(InstancedObjectMapProperty->HasAllPropertyFlags(CPF_ContainsInstancedReference));
+		CHECK_FALSE(InstancedObjectMapProperty->HasAllPropertyFlags(CPF_InstancedReference));
+		CHECK_FALSE(InstancedObjectMapProperty->HasAllPropertyFlags(CPF_PersistentInstance));
+#if WITH_METADATA
+		// Note: 'EditInline' is set via UHT.
+		CHECK(InstancedObjectMapProperty->GetBoolMetaData(TEXT("EditInline")));
+#endif
+
+		REQUIRE(InstancedObjectMapProperty->KeyProp != nullptr);
+		CHECK(InstancedObjectMapProperty->KeyProp->IsA<FObjectProperty>());
+
+		CHECK(InstancedObjectMapProperty->KeyProp->ContainsInstancedObjectProperty());
+		CHECK_FALSE(InstancedObjectMapProperty->KeyProp->HasAllPropertyFlags(CPF_ContainsInstancedReference));
+		CHECK(InstancedObjectMapProperty->KeyProp->HasAllPropertyFlags(CPF_InstancedReference));
+		CHECK(InstancedObjectMapProperty->KeyProp->HasAllPropertyFlags(CPF_PersistentInstance));
+#if WITH_METADATA
+		// Note: Currently the KeyProp is not also flagged as 'EditInline' when the map is declared as 'Instanced'.
+		CHECK_FALSE(InstancedObjectMapProperty->KeyProp->GetBoolMetaData(TEXT("EditInline")));
+#endif
+
+		REQUIRE(InstancedObjectMapProperty->ValueProp != nullptr);
+		CHECK(InstancedObjectMapProperty->ValueProp->IsA<FObjectProperty>());
+
+		CHECK(InstancedObjectMapProperty->ValueProp->ContainsInstancedObjectProperty());
+		CHECK_FALSE(InstancedObjectMapProperty->ValueProp->HasAllPropertyFlags(CPF_ContainsInstancedReference));
+		CHECK(InstancedObjectMapProperty->ValueProp->HasAllPropertyFlags(CPF_InstancedReference));
+		CHECK(InstancedObjectMapProperty->ValueProp->HasAllPropertyFlags(CPF_PersistentInstance));
+#if WITH_METADATA
+		CHECK(InstancedObjectMapProperty->ValueProp->GetBoolMetaData(TEXT("EditInline")));
+#endif
+	}
+
+	SECTION("DefaultToInstanced object map container property")
+	{
+		const FName DefaultToInstancedObjectMapPropertyName = GET_MEMBER_NAME_CHECKED(USubobjectInstancingTestOuterObject, InnerObjectFromTypeMap);
+		const FMapProperty* DefaultToInstancedObjectMapProperty = CastField<FMapProperty>(USubobjectInstancingTestOuterObject::StaticClass()->FindPropertyByName(DefaultToInstancedObjectMapPropertyName));
+
+		REQUIRE(DefaultToInstancedObjectMapProperty != nullptr);
+		CHECK(DefaultToInstancedObjectMapProperty->ContainsInstancedObjectProperty());
+		CHECK(DefaultToInstancedObjectMapProperty->HasAllPropertyFlags(CPF_ContainsInstancedReference));
+		CHECK_FALSE(DefaultToInstancedObjectMapProperty->HasAllPropertyFlags(CPF_InstancedReference));
+		CHECK_FALSE(DefaultToInstancedObjectMapProperty->HasAllPropertyFlags(CPF_PersistentInstance));
+#if WITH_METADATA
+		// Note: DefaultToInstanced also implies 'EditInline' (via UHT).
+		CHECK(DefaultToInstancedObjectMapProperty->GetBoolMetaData(TEXT("EditInline")));
+#endif
+
+		REQUIRE(DefaultToInstancedObjectMapProperty->KeyProp != nullptr);
+		CHECK(DefaultToInstancedObjectMapProperty->KeyProp->IsA<FObjectProperty>());
+
+		CHECK(DefaultToInstancedObjectMapProperty->KeyProp->ContainsInstancedObjectProperty());
+		CHECK_FALSE(DefaultToInstancedObjectMapProperty->KeyProp->HasAllPropertyFlags(CPF_ContainsInstancedReference));
+		CHECK(DefaultToInstancedObjectMapProperty->KeyProp->HasAllPropertyFlags(CPF_InstancedReference));
+		CHECK_FALSE(DefaultToInstancedObjectMapProperty->KeyProp->HasAllPropertyFlags(CPF_PersistentInstance));
+#if WITH_METADATA
+		// Note: Unlike explicitly-instanced map properties, the Key/ValueProp is not also flagged as 'EditInline'.
+		CHECK_FALSE(DefaultToInstancedObjectMapProperty->KeyProp->GetBoolMetaData(TEXT("EditInline")));
+#endif
+
+		REQUIRE(DefaultToInstancedObjectMapProperty->ValueProp != nullptr);
+		CHECK(DefaultToInstancedObjectMapProperty->ValueProp->IsA<FObjectProperty>());
+
+		CHECK(DefaultToInstancedObjectMapProperty->ValueProp->ContainsInstancedObjectProperty());
+		CHECK_FALSE(DefaultToInstancedObjectMapProperty->ValueProp->HasAllPropertyFlags(CPF_ContainsInstancedReference));
+		CHECK(DefaultToInstancedObjectMapProperty->ValueProp->HasAllPropertyFlags(CPF_InstancedReference));
+		CHECK_FALSE(DefaultToInstancedObjectMapProperty->ValueProp->HasAllPropertyFlags(CPF_PersistentInstance));
+#if WITH_METADATA
+		// Note: Unlike explicitly-instanced map properties, the Key/ValueProp is not also flagged as 'EditInline'.
+		CHECK_FALSE(DefaultToInstancedObjectMapProperty->ValueProp->GetBoolMetaData(TEXT("EditInline")));
+#endif
+	}
+
+	SECTION("Struct with instanced object property")
+	{
+		const FName StructWithInstancedObjectPropertyName = GET_MEMBER_NAME_CHECKED(USubobjectInstancingTestOuterObject, StructWithInnerObjects);
+		const FStructProperty* StructWithInstancedObjectProperty = CastField<FStructProperty>(USubobjectInstancingTestOuterObject::StaticClass()->FindPropertyByName(StructWithInstancedObjectPropertyName));
+
+		REQUIRE(StructWithInstancedObjectProperty != nullptr);
+		CHECK(StructWithInstancedObjectProperty->ContainsInstancedObjectProperty());
+		CHECK(StructWithInstancedObjectProperty->HasAllPropertyFlags(CPF_ContainsInstancedReference));
+		CHECK_FALSE(StructWithInstancedObjectProperty->HasAllPropertyFlags(CPF_InstancedReference));
+		CHECK_FALSE(StructWithInstancedObjectProperty->HasAllPropertyFlags(CPF_PersistentInstance));
+#if WITH_METADATA
+		CHECK_FALSE(StructWithInstancedObjectProperty->GetBoolMetaData(TEXT("EditInline")));
+#endif
+
+		REQUIRE(StructWithInstancedObjectProperty->Struct != nullptr);
+		const FName InstancedObjectPropertyName = GET_MEMBER_NAME_CHECKED(FSubobjectInstancingTestStructType, InnerObject);
+		const FProperty* InstancedObjectProperty = StructWithInstancedObjectProperty->Struct->FindPropertyByName(InstancedObjectPropertyName);
+
+		REQUIRE(InstancedObjectProperty != nullptr);
+		CHECK(InstancedObjectProperty->IsA<FObjectProperty>());
+
+		CHECK(InstancedObjectProperty->ContainsInstancedObjectProperty());
+		CHECK_FALSE(InstancedObjectProperty->HasAllPropertyFlags(CPF_ContainsInstancedReference));
+		CHECK(InstancedObjectProperty->HasAllPropertyFlags(CPF_InstancedReference));
+		CHECK(InstancedObjectProperty->HasAllPropertyFlags(CPF_PersistentInstance));
+#if WITH_METADATA
+		CHECK(InstancedObjectProperty->GetBoolMetaData(TEXT("EditInline")));
+#endif
+	}
+
+	SECTION("Struct with DefaultToInstanced object property")
+	{
+		const FName StructWithInstancedObjectPropertyName = GET_MEMBER_NAME_CHECKED(USubobjectInstancingTestOuterObject, StructWithInnerObjects);
+		const FStructProperty* StructWithInstancedObjectProperty = CastField<FStructProperty>(USubobjectInstancingTestOuterObject::StaticClass()->FindPropertyByName(StructWithInstancedObjectPropertyName));
+
+		REQUIRE(StructWithInstancedObjectProperty != nullptr);
+		CHECK(StructWithInstancedObjectProperty->ContainsInstancedObjectProperty());
+		CHECK(StructWithInstancedObjectProperty->HasAllPropertyFlags(CPF_ContainsInstancedReference));
+		CHECK_FALSE(StructWithInstancedObjectProperty->HasAllPropertyFlags(CPF_InstancedReference));
+		CHECK_FALSE(StructWithInstancedObjectProperty->HasAllPropertyFlags(CPF_PersistentInstance));
+#if WITH_METADATA
+		CHECK_FALSE(StructWithInstancedObjectProperty->GetBoolMetaData(TEXT("EditInline")));
+#endif
+
+		REQUIRE(StructWithInstancedObjectProperty->Struct != nullptr);
+		const FName DefaultToInstancedObjectPropertyName = GET_MEMBER_NAME_CHECKED(FSubobjectInstancingTestStructType, InnerObjectFromType);
+		const FProperty* DefaultToInstancedObjectProperty = StructWithInstancedObjectProperty->Struct->FindPropertyByName(DefaultToInstancedObjectPropertyName);
+
+		REQUIRE(DefaultToInstancedObjectProperty != nullptr);
+		CHECK(DefaultToInstancedObjectProperty->IsA<FObjectProperty>());
+
+		CHECK(DefaultToInstancedObjectProperty->ContainsInstancedObjectProperty());
+		CHECK_FALSE(DefaultToInstancedObjectProperty->HasAllPropertyFlags(CPF_ContainsInstancedReference));
+		CHECK(DefaultToInstancedObjectProperty->HasAllPropertyFlags(CPF_InstancedReference));
+		CHECK_FALSE(DefaultToInstancedObjectProperty->HasAllPropertyFlags(CPF_PersistentInstance));
+#if WITH_METADATA
+		CHECK(DefaultToInstancedObjectProperty->GetBoolMetaData(TEXT("EditInline")));
+#endif
+	}
+
+	SECTION("Struct with instanced object array property")
+	{
+		const FName StructWithInstancedObjectArrayPropertyName = GET_MEMBER_NAME_CHECKED(USubobjectInstancingTestOuterObject, StructWithInnerObjects);
+		const FStructProperty* StructWithInstancedObjectArrayProperty = CastField<FStructProperty>(USubobjectInstancingTestOuterObject::StaticClass()->FindPropertyByName(StructWithInstancedObjectArrayPropertyName));
+
+		REQUIRE(StructWithInstancedObjectArrayProperty != nullptr);
+		CHECK(StructWithInstancedObjectArrayProperty->ContainsInstancedObjectProperty());
+		CHECK(StructWithInstancedObjectArrayProperty->HasAllPropertyFlags(CPF_ContainsInstancedReference));
+		CHECK_FALSE(StructWithInstancedObjectArrayProperty->HasAllPropertyFlags(CPF_InstancedReference));
+		CHECK_FALSE(StructWithInstancedObjectArrayProperty->HasAllPropertyFlags(CPF_PersistentInstance));
+#if WITH_METADATA
+		CHECK_FALSE(StructWithInstancedObjectArrayProperty->GetBoolMetaData(TEXT("EditInline")));
+#endif
+
+		REQUIRE(StructWithInstancedObjectArrayProperty->Struct != nullptr);
+		const FName InstancedObjectArrayPropertyName = GET_MEMBER_NAME_CHECKED(FSubobjectInstancingTestStructType, InnerObjectArray);
+		const FArrayProperty* InstancedObjectArrayProperty = CastField<FArrayProperty>(StructWithInstancedObjectArrayProperty->Struct->FindPropertyByName(InstancedObjectArrayPropertyName));
+
+		REQUIRE(InstancedObjectArrayProperty != nullptr);
+		CHECK(InstancedObjectArrayProperty->ContainsInstancedObjectProperty());
+		CHECK(InstancedObjectArrayProperty->HasAllPropertyFlags(CPF_ContainsInstancedReference));
+		CHECK_FALSE(InstancedObjectArrayProperty->HasAllPropertyFlags(CPF_InstancedReference));
+		CHECK_FALSE(InstancedObjectArrayProperty->HasAllPropertyFlags(CPF_PersistentInstance));
+#if WITH_METADATA
+		CHECK(InstancedObjectArrayProperty->GetBoolMetaData(TEXT("EditInline")));
+#endif
+
+		REQUIRE(InstancedObjectArrayProperty->Inner != nullptr);
+		CHECK(InstancedObjectArrayProperty->Inner->IsA<FObjectProperty>());
+
+		CHECK(InstancedObjectArrayProperty->Inner->ContainsInstancedObjectProperty());
+		CHECK_FALSE(InstancedObjectArrayProperty->Inner->HasAllPropertyFlags(CPF_ContainsInstancedReference));
+		CHECK(InstancedObjectArrayProperty->Inner->HasAllPropertyFlags(CPF_InstancedReference));
+		CHECK(InstancedObjectArrayProperty->Inner->HasAllPropertyFlags(CPF_PersistentInstance));
+#if WITH_METADATA
+		CHECK(InstancedObjectArrayProperty->Inner->GetBoolMetaData(TEXT("EditInline")));
+#endif
+	}
+
+	SECTION("Optional instanced object property")
+	{
+		const FName OptionalInstancedObjectPropertyName = GET_MEMBER_NAME_CHECKED(USubobjectInstancingTestOuterObject, OptionalInnerObject);
+		const FOptionalProperty* OptionalInstancedObjectProperty = CastField<FOptionalProperty>(USubobjectInstancingTestOuterObject::StaticClass()->FindPropertyByName(OptionalInstancedObjectPropertyName));
+
+		REQUIRE(OptionalInstancedObjectProperty != nullptr);
+		CHECK(OptionalInstancedObjectProperty->ContainsInstancedObjectProperty());
+		CHECK(OptionalInstancedObjectProperty->HasAllPropertyFlags(CPF_ContainsInstancedReference));
+		CHECK_FALSE(OptionalInstancedObjectProperty->HasAllPropertyFlags(CPF_InstancedReference));
+		CHECK_FALSE(OptionalInstancedObjectProperty->HasAllPropertyFlags(CPF_PersistentInstance));
+#if WITH_METADATA
+		CHECK(OptionalInstancedObjectProperty->GetBoolMetaData(TEXT("EditInline")));
+#endif
+
+		const FProperty* OptionalInstancedObjectValueProperty = OptionalInstancedObjectProperty->GetValueProperty();
+		REQUIRE(OptionalInstancedObjectValueProperty != nullptr);
+		CHECK(OptionalInstancedObjectValueProperty->IsA<FObjectProperty>());
+
+		CHECK(OptionalInstancedObjectValueProperty->ContainsInstancedObjectProperty());
+		CHECK_FALSE(OptionalInstancedObjectValueProperty->HasAllPropertyFlags(CPF_ContainsInstancedReference));
+		CHECK(OptionalInstancedObjectValueProperty->HasAllPropertyFlags(CPF_InstancedReference));
+		CHECK(OptionalInstancedObjectValueProperty->HasAllPropertyFlags(CPF_PersistentInstance));
+#if WITH_METADATA
+		CHECK(OptionalInstancedObjectValueProperty->GetBoolMetaData(TEXT("EditInline")));
+#endif
+	}
+
+	SECTION("Optional DefaultToInstanced object property")
+	{
+		const FName OptionalDefaultToInstancedObjectPropertyName = GET_MEMBER_NAME_CHECKED(USubobjectInstancingTestOuterObject, OptionalInnerObjectFromType);
+		const FOptionalProperty* OptionalDefaultToInstancedObjectProperty = CastField<FOptionalProperty>(USubobjectInstancingTestOuterObject::StaticClass()->FindPropertyByName(OptionalDefaultToInstancedObjectPropertyName));
+
+		REQUIRE(OptionalDefaultToInstancedObjectProperty != nullptr);
+		CHECK(OptionalDefaultToInstancedObjectProperty->ContainsInstancedObjectProperty());
+		CHECK(OptionalDefaultToInstancedObjectProperty->HasAllPropertyFlags(CPF_ContainsInstancedReference));
+		CHECK_FALSE(OptionalDefaultToInstancedObjectProperty->HasAllPropertyFlags(CPF_InstancedReference));
+		CHECK_FALSE(OptionalDefaultToInstancedObjectProperty->HasAllPropertyFlags(CPF_PersistentInstance));
+#if WITH_METADATA
+		// Note: DefaultToInstanced also implies 'EditInline' (via UHT).
+		CHECK(OptionalDefaultToInstancedObjectProperty->GetBoolMetaData(TEXT("EditInline")));
+#endif
+
+		const FProperty* OptionalDefaultToInstancedObjectValueProperty = OptionalDefaultToInstancedObjectProperty->GetValueProperty();
+		REQUIRE(OptionalDefaultToInstancedObjectValueProperty != nullptr);
+		CHECK(OptionalDefaultToInstancedObjectValueProperty->IsA<FObjectProperty>());
+
+		CHECK(OptionalDefaultToInstancedObjectValueProperty->ContainsInstancedObjectProperty());
+		CHECK_FALSE(OptionalDefaultToInstancedObjectValueProperty->HasAllPropertyFlags(CPF_ContainsInstancedReference));
+		CHECK(OptionalDefaultToInstancedObjectValueProperty->HasAllPropertyFlags(CPF_InstancedReference));
+		CHECK_FALSE(OptionalDefaultToInstancedObjectValueProperty->HasAllPropertyFlags(CPF_PersistentInstance));
+#if WITH_METADATA
+		// Note: Unlike explicitly-instanced optional properties, the value property is not also flagged as 'EditInline'.
+		CHECK_FALSE(OptionalDefaultToInstancedObjectValueProperty->GetBoolMetaData(TEXT("EditInline")));
+#endif
+	}
+
+	SECTION("Optional instanced object array property")
+	{
+		const FName OptionalInstancedObjectArrayPropertyName = GET_MEMBER_NAME_CHECKED(USubobjectInstancingTestOuterObject, OptionalInnerObjectArray);
+		const FOptionalProperty* OptionalInstancedObjectArrayProperty = CastField<FOptionalProperty>(USubobjectInstancingTestOuterObject::StaticClass()->FindPropertyByName(OptionalInstancedObjectArrayPropertyName));
+
+		REQUIRE(OptionalInstancedObjectArrayProperty != nullptr);
+		CHECK(OptionalInstancedObjectArrayProperty->ContainsInstancedObjectProperty());
+		CHECK(OptionalInstancedObjectArrayProperty->HasAllPropertyFlags(CPF_ContainsInstancedReference));
+		CHECK_FALSE(OptionalInstancedObjectArrayProperty->HasAllPropertyFlags(CPF_InstancedReference));
+		CHECK_FALSE(OptionalInstancedObjectArrayProperty->HasAllPropertyFlags(CPF_PersistentInstance));
+#if WITH_METADATA
+		CHECK(OptionalInstancedObjectArrayProperty->GetBoolMetaData(TEXT("EditInline")));
+#endif
+
+		FArrayProperty* OptionalInstancedObjectArrayValueProperty = CastField<FArrayProperty>(OptionalInstancedObjectArrayProperty->GetValueProperty());
+		REQUIRE(OptionalInstancedObjectArrayValueProperty != nullptr);
+		CHECK(OptionalInstancedObjectArrayValueProperty->ContainsInstancedObjectProperty());
+		CHECK(OptionalInstancedObjectArrayValueProperty->HasAllPropertyFlags(CPF_ContainsInstancedReference));
+		CHECK(OptionalInstancedObjectArrayValueProperty->HasAllPropertyFlags(CPF_InstancedReference));
+		CHECK(OptionalInstancedObjectArrayValueProperty->HasAllPropertyFlags(CPF_PersistentInstance));
+#if WITH_METADATA
+		CHECK(OptionalInstancedObjectArrayValueProperty->GetBoolMetaData(TEXT("EditInline")));
+#endif
+
+		REQUIRE(OptionalInstancedObjectArrayValueProperty->Inner != nullptr);
+		CHECK(OptionalInstancedObjectArrayValueProperty->Inner->IsA<FObjectProperty>());
+
+		CHECK(OptionalInstancedObjectArrayValueProperty->Inner->ContainsInstancedObjectProperty());
+		CHECK_FALSE(OptionalInstancedObjectArrayValueProperty->Inner->HasAllPropertyFlags(CPF_ContainsInstancedReference));
+		CHECK(OptionalInstancedObjectArrayValueProperty->Inner->HasAllPropertyFlags(CPF_InstancedReference));
+		CHECK(OptionalInstancedObjectArrayValueProperty->Inner->HasAllPropertyFlags(CPF_PersistentInstance));
+#if WITH_METADATA
+		// Note: Currently 'EditInline' does not propagate to the array's inner property.
+		CHECK_FALSE(OptionalInstancedObjectArrayValueProperty->Inner->GetBoolMetaData(TEXT("EditInline")));
+#endif
+	}
+}
+
+}
+
+#endif	// WITH_TESTS
\ No newline at end of file
diff --git a/Engine/Source/Runtime/CoreUObject/Tests/UObject/ObjectInstancingGraphTest.cpp b/Engine/Source/Runtime/CoreUObject/Tests/UObject/ObjectInstancingGraphTest.cpp
new file mode 100644
--- /dev/null
+++ b/Engine/Source/Runtime/CoreUObject/Tests/UObject/ObjectInstancingGraphTest.cpp
@@ -0,0 +1,163 @@
+// Copyright Epic Games, Inc. All Rights Reserved.
+
+#if WITH_TESTS
+
+#include "SubobjectInstancingTest.h"
+#include "Tests/TestHarnessAdapter.h"
+#include "UObject/Package.h"
+#include "UObject/UnrealType.h"
+
+namespace UE
+{
+
+TEST_CASE_NAMED(FObjectInstancingGraphTest, "CoreUObject::FObjectInstancingGraph", "[CoreUObject][EngineFilter]")
+{
+	SECTION("Default construction")
+	{
+		FObjectInstancingGraph InstanceGraph;
+		CHECK_FALSE(InstanceGraph.HasDestinationRoot());
+		CHECK(InstanceGraph.IsSubobjectInstancingEnabled());
+	}
+
+	SECTION("Construct with options")
+	{
+		FObjectInstancingGraph InstanceGraph(EObjectInstancingGraphOptions::DisableInstancing);
+		CHECK_FALSE(InstanceGraph.HasDestinationRoot());
+		CHECK_FALSE(InstanceGraph.IsSubobjectInstancingEnabled());
+	}
+
+	SECTION("Construct with root object")
+	{
+		USubobjectInstancingTestObject* RootObject = NewObject<USubobjectInstancingTestObject>();
+
+		FObjectInstancingGraph InstanceGraph(RootObject);
+		CHECK(InstanceGraph.HasDestinationRoot());
+		CHECK(InstanceGraph.GetDestinationRoot() == RootObject);
+		CHECK(InstanceGraph.GetDestinationObject(RootObject->GetArchetype()) == RootObject);
+	}
+
+	SECTION("Set destination root on first add")
+	{
+		USubobjectInstancingTestObject* RootObject = NewObject<USubobjectInstancingTestObject>();
+
+		FObjectInstancingGraph InstanceGraph;
+		InstanceGraph.AddNewObject(RootObject);
+		CHECK(InstanceGraph.HasDestinationRoot());
+		CHECK(InstanceGraph.GetDestinationRoot() == RootObject);
+		CHECK(InstanceGraph.GetDestinationObject(RootObject->GetArchetype()) == RootObject);
+	}
+
+	SECTION("Exclude native class property from instancing when initialized from CDO")
+	{
+		FObjectInstancingGraph InstanceGraph;
+
+		UClass* TestClass = USubobjectInstancingTestOuterObject::StaticClass();
+		const FName InnerObjectPropertyName = GET_MEMBER_NAME_CHECKED(USubobjectInstancingTestOuterObject, InnerObject);
+		const FObjectProperty* InnerObjectProperty = CastField<FObjectProperty>(TestClass->FindPropertyByName(InnerObjectPropertyName));
+
+		REQUIRE(InnerObjectProperty != nullptr);
+		InstanceGraph.AddPropertyToSubobjectExclusionList(InnerObjectProperty);
+
+		FStaticConstructObjectParameters Params(TestClass);
+		Params.Outer = GetTransientPackage();
+		Params.InstanceGraph = &InstanceGraph;
+		USubobjectInstancingTestOuterObject* OuterObject = CastChecked<USubobjectInstancingTestOuterObject>(StaticConstructObject_Internal(Params));
+
+		UObject* InnerObject = InnerObjectProperty->GetObjectPropertyValue_InContainer(OuterObject);
+
+		// Note: Because the class/property is native and we are initializing from the CDO, excluding it will have no effect and you are left with the constructed (instanced) value
+		// This differs from the non-native (e.g. Blueprint) case, which will always initialize non-transient object properties from the CDO, and you are left with the default value
+		// @see FObjectInitializer::InitProperties()
+		REQUIRE(InnerObject != nullptr);
+		CHECK(InnerObject == OuterObject->InnerObject);
+		CHECK(InnerObject->IsInOuter(OuterObject));
+		CHECK_FALSE(InnerObject->IsTemplate());
+		CHECK(InnerObject->IsDefaultSubobject());
+		CHECK_FALSE(InnerObject->IsTemplateForSubobjects());
+	}
+
+	SECTION("Exclude native class property from instancing when initialized from template")
+	{
+		FObjectInstancingGraph InstanceGraph;
+
+		UClass* TestClass = USubobjectInstancingTestOuterObject::StaticClass();
+		const FName InnerObjectPropertyName = GET_MEMBER_NAME_CHECKED(USubobjectInstancingTestOuterObject, InnerObject);
+		const FObjectProperty* InnerObjectProperty = CastField<FObjectProperty>(TestClass->FindPropertyByName(InnerObjectPropertyName));
+
+		REQUIRE(InnerObjectProperty != nullptr);
+		InstanceGraph.AddPropertyToSubobjectExclusionList(InnerObjectProperty);
+
+		FStaticConstructObjectParameters Params(TestClass);
+		Params.Outer = GetTransientPackage();
+		Params.Template = NewObject<USubobjectInstancingTestOuterObject>(GetTransientPackage(), NAME_None, RF_ArchetypeObject);
+		Params.InstanceGraph = &InstanceGraph;
+		USubobjectInstancingTestOuterObject* OuterObject = CastChecked<USubobjectInstancingTestOuterObject>(StaticConstructObject_Internal(Params));
+
+		UObject* InnerObject = InnerObjectProperty->GetObjectPropertyValue_InContainer(OuterObject);
+
+		// Note: The expectation here is that we are left with the default (initialized) value
+		REQUIRE(InnerObject != nullptr);
+		CHECK(InnerObject == OuterObject->InnerObject);
+		CHECK_FALSE(InnerObject->IsInOuter(OuterObject));
+		CHECK(InnerObject->IsTemplate());
+		CHECK(InnerObject->IsDefaultSubobject());
+		CHECK(InnerObject->IsTemplateForSubobjects());
+	}
+
+	SECTION("Exclude non-native class property from instancing when initialized from CDO")
+	{
+		FObjectInstancingGraph InstanceGraph;
+
+		UClass* ParentClass = USubobjectInstancingTestOuterObject::StaticClass();
+		UClass* TestClass = FSubobjectInstancingTestUtils::CreateNonNativeInstancingTestClass(ParentClass);
+		const FName NonNativeInnerObjectPropertyName = FSubobjectInstancingTestUtils::GetNonNativeInnerObjectPropertyName();
+		const FObjectProperty* NonNativeInnerObjectProperty = CastField<FObjectProperty>(TestClass->FindPropertyByName(NonNativeInnerObjectPropertyName));
+
+		REQUIRE(NonNativeInnerObjectProperty != nullptr);
+		InstanceGraph.AddPropertyToSubobjectExclusionList(NonNativeInnerObjectProperty);
+
+		FStaticConstructObjectParameters Params(TestClass);
+		Params.Outer = GetTransientPackage();
+		Params.InstanceGraph = &InstanceGraph;
+		USubobjectInstancingTestOuterObject* OuterObject = CastChecked<USubobjectInstancingTestOuterObject>(StaticConstructObject_Internal(Params));
+
+		UObject* NonNativeInnerObject = NonNativeInnerObjectProperty->GetObjectPropertyValue_InContainer(OuterObject);
+
+		REQUIRE(NonNativeInnerObject != nullptr);
+		CHECK_FALSE(NonNativeInnerObject->IsInOuter(OuterObject));
+		CHECK(NonNativeInnerObject->IsTemplate());
+		CHECK(NonNativeInnerObject->IsDefaultSubobject());
+		CHECK(NonNativeInnerObject->IsTemplateForSubobjects());
+	}
+
+	SECTION("Exclude non-native class property from instancing when initialized from template")
+	{
+		FObjectInstancingGraph InstanceGraph;
+
+		UClass* ParentClass = USubobjectInstancingTestOuterObject::StaticClass();
+		UClass* TestClass = FSubobjectInstancingTestUtils::CreateNonNativeInstancingTestClass(ParentClass);
+		const FName NonNativeInnerObjectPropertyName = FSubobjectInstancingTestUtils::GetNonNativeInnerObjectPropertyName();
+		const FObjectProperty* NonNativeInnerObjectProperty = CastField<FObjectProperty>(TestClass->FindPropertyByName(NonNativeInnerObjectPropertyName));
+
+		REQUIRE(NonNativeInnerObjectProperty != nullptr);
+		InstanceGraph.AddPropertyToSubobjectExclusionList(NonNativeInnerObjectProperty);
+
+		FStaticConstructObjectParameters Params(TestClass);
+		Params.Outer = GetTransientPackage();
+		Params.Template = NewObject<USubobjectInstancingTestOuterObject>(GetTransientPackage(), TestClass, NAME_None, RF_ArchetypeObject);
+		Params.InstanceGraph = &InstanceGraph;
+		USubobjectInstancingTestOuterObject* OuterObject = CastChecked<USubobjectInstancingTestOuterObject>(StaticConstructObject_Internal(Params));
+
+		UObject* NonNativeInnerObject = NonNativeInnerObjectProperty->GetObjectPropertyValue_InContainer(OuterObject);
+
+		REQUIRE(NonNativeInnerObject != nullptr);
+		CHECK_FALSE(NonNativeInnerObject->IsInOuter(OuterObject));
+		CHECK(NonNativeInnerObject->IsTemplate());
+		CHECK(NonNativeInnerObject->IsDefaultSubobject());
+		CHECK(NonNativeInnerObject->IsTemplateForSubobjects(RF_ArchetypeObject));
+	}
+}
+
+}
+
+#endif	// WITH_TESTS
diff --git a/Engine/Source/Runtime/CoreUObject/Tests/UObject/SubobjectInstancingTest.cpp b/Engine/Source/Runtime/CoreUObject/Tests/UObject/SubobjectInstancingTest.cpp
new file mode 100644
--- /dev/null
+++ b/Engine/Source/Runtime/CoreUObject/Tests/UObject/SubobjectInstancingTest.cpp
@@ -0,0 +1,1964 @@
+// Copyright Epic Games, Inc. All Rights Reserved.
+
+#if WITH_TESTS
+
+#include "SubobjectInstancingTest.h"
+#include "Tests/TestHarnessAdapter.h"
+#include "UObject/Package.h"
+#include "UObject/PropertyOptional.h"
+
+USubobjectInstancingTestOuterObject::USubobjectInstancingTestOuterObject(const FObjectInitializer& ObjectInitializer)
+	:Super(ObjectInitializer)
+{
+	// Self reference - no instancing, should always reference this object
+	// Since this is a native type, this must be marked 'Instanced' to work
+	SelfRef = this;
+
+	// Null reference - no instancing, should always be NULL after construction
+	NullObject = nullptr;
+
+	// Instanced subobject that is expected to be a unique ptr value after construction
+	InnerObject = CreateDefaultSubobject<USubobjectInstancingTestObject>(TEXT("InnerObject"));
+
+	// Object that exists within the scope of this object, but is otherwise not instanced
+	// With one exception for native class types (see below), not expected to be a unique ptr value
+	SharedObject = CreateDefaultSubobject<USubobjectInstancingTestObject>(TEXT("SharedObject"));
+
+	// Object that exists outside the scope of this object; e.g. a simple reference (no instancing)
+	ExternalObject = NewObject<USubobjectInstancingTestObject>(GetTransientPackage(), TEXT("ExternalObject"));
+
+	// Internal reference - no instancing, but should reference the unique InnerObject constructed above
+	// Since this is a native type that uses property-based instancing, this must be marked 'Instanced' to work
+	InternalObject = InnerObject;
+
+	// Instanced subobject that is not going to be serialized on save; note that this is a transient property
+	TransientInnerObject = CreateDefaultSubobject<USubobjectInstancingTestObject>(TEXT("TransientInnerObject"));
+
+	// Subobject that may be instanced at edit time rather than at construction time; this is NOT considered a default subobject
+	EditTimeInnerObject = nullptr;
+
+	// Subobject that does not inherit values from the CDO/template when instanced; 'transient' here does not affect serialization
+	// Note: This only applies to C++ types; otherwise it behaves the same as 'InnerObject' above and 'bTransient' is not supported
+	LocalOnlyInnerObject = CreateDefaultSubobject<USubobjectInstancingTestObject>(TEXT("LocalOnlyInnerObject"), /*bTransient =*/ true);
+
+	// Property with instanced semantics inferred from the class type (DefaultToInstanced); property flags exclude CPF_PersistentInstance
+	// Runtime behavior is otherwise expected to be the same as any other CPF_InstancedReference property (e.g. InnerObject above)
+	InnerObjectFromType = CreateDefaultSubobject<USubobjectInstancingDefaultToInstancedTestObject>(TEXT("InnerObjectFromType"));
+
+	// Instanced subobject that is constructed via NewObject() instead of CreateDefaultSubobject(), which is occasionally used instead
+	// Property flags will exclude RF_DefaultSubObject, and runtime behavior is similar to LocalOnlyInnerObject above (i.e. no inheritance)
+	InnerObjectUsingNew = NewObject<USubobjectInstancingTestObject>(this, TEXT("InnerObjectUsingNew"));
+
+	// Array container of references to instanced subobjects
+	InnerObjectArray.Add(CreateDefaultSubobject<USubobjectInstancingTestObject>(TEXT("InnerObjectArrayElem_0")));
+	InnerObjectArray.Add(CreateDefaultSubobject<USubobjectInstancingTestObject>(TEXT("InnerObjectArrayElem_1")));
+
+	// Array container of references to DefaultToInstanced-typed subobjects
+	InnerObjectFromTypeArray.Add(CreateDefaultSubobject<USubobjectInstancingDefaultToInstancedTestObject>(TEXT("InnerObjectFromTypeArrayElem_0")));
+	InnerObjectFromTypeArray.Add(CreateDefaultSubobject<USubobjectInstancingDefaultToInstancedTestObject>(TEXT("InnerObjectFromTypeArrayElem_1")));
+
+	// Set container of references to instanced subobjects
+	InnerObjectSet.Add(CreateDefaultSubobject<USubobjectInstancingTestObject>(TEXT("InnerObjectSetElem_0")));
+	InnerObjectSet.Add(CreateDefaultSubobject<USubobjectInstancingTestObject>(TEXT("InnerObjectSetElem_1")));
+
+	// Set container of references to DefaultToInstanced-typed subobjects
+	InnerObjectFromTypeSet.Add(CreateDefaultSubobject<USubobjectInstancingDefaultToInstancedTestObject>(TEXT("InnerObjectFromTypeSetElem_0")));
+	InnerObjectFromTypeSet.Add(CreateDefaultSubobject<USubobjectInstancingDefaultToInstancedTestObject>(TEXT("InnerObjectFromTypeSetElem_1")));
+
+	// Map container of pairs of references to instanced subobjects
+	InnerObjectMap.Add(CreateDefaultSubobject<USubobjectInstancingTestObject>(TEXT("InnerObjectMapKey_0")), CreateDefaultSubobject<USubobjectInstancingTestObject>(TEXT("InnerObjectMapVal_0")));
+	InnerObjectMap.Add(CreateDefaultSubobject<USubobjectInstancingTestObject>(TEXT("InnerObjectMapKey_1")), CreateDefaultSubobject<USubobjectInstancingTestObject>(TEXT("InnerObjectMapVal_1")));
+
+	// Map container of pairs of references to DefaultToInstanced-typed subobjects
+	InnerObjectFromTypeMap.Add(CreateDefaultSubobject<USubobjectInstancingDefaultToInstancedTestObject>(TEXT("InnerObjectFromTypeMapKey_0")), CreateDefaultSubobject<USubobjectInstancingDefaultToInstancedTestObject>(TEXT("InnerObjectFromTypeMapVal_0")));
+	InnerObjectFromTypeMap.Add(CreateDefaultSubobject<USubobjectInstancingDefaultToInstancedTestObject>(TEXT("InnerObjectFromTypeMapKey_1")), CreateDefaultSubobject<USubobjectInstancingDefaultToInstancedTestObject>(TEXT("InnerObjectFromTypeMapVal_1")));
+
+	// Struct container of references to instanced subobjects
+	StructWithInnerObjects.InnerObject = CreateDefaultSubobject<USubobjectInstancingTestObject>(TEXT("InnerObjectForStruct"));
+	StructWithInnerObjects.InnerObjectFromType = CreateDefaultSubobject<USubobjectInstancingDefaultToInstancedTestObject>(TEXT("InnerObjectFromTypeForStruct"));
+	StructWithInnerObjects.InnerObjectArray.Add(CreateDefaultSubobject<USubobjectInstancingTestObject>(TEXT("InnerObjectForStructArrayElem_0")));
+	StructWithInnerObjects.InnerObjectArray.Add(CreateDefaultSubobject<USubobjectInstancingTestObject>(TEXT("InnerObjectForStructArrayElem_1")));
+
+	// Optional reference to instanced subobject; creating as optional to also test the DoNotCreate override
+	// Note that this internally sets the 'bIsRequired' parameter to false, which could also be done here instead
+	OptionalInnerObject = CreateOptionalDefaultSubobject<USubobjectInstancingTestObject>(TEXT("OptionalInnerObject"));
+
+	// Optional reference to DefaultToInstanced subobject (with instanced property semantics inferred from ptr type)
+	OptionalInnerObjectFromType = CreateDefaultSubobject<USubobjectInstancingDefaultToInstancedTestObject>(TEXT("OptionalInnerObjectFromType"));
+
+	// Optional array of references to instanced subobjects
+	OptionalInnerObjectArray.Emplace();
+	OptionalInnerObjectArray->Add(CreateDefaultSubobject<USubobjectInstancingTestObject>(TEXT("OptionalInnerObjectArrayElem_0")));
+	OptionalInnerObjectArray->Add(CreateDefaultSubobject<USubobjectInstancingTestObject>(TEXT("OptionalInnerObjectArrayElem_1")));
+
+	// Instanced subobject that contains a directly-nested instanced subobject (i.e. one level deep)
+	InnerObjectWithDirectlyNestedObject = CreateDefaultSubobject<USubobjectInstancingTestDirectlyNestedObject>(TEXT("InnerObjectWithDirectlyNestedObject"));
+	InnerObjectWithDirectlyNestedObject->OwnerObject = this;
+
+	// Instanced subobject that contains an indirectly-nested instanced subobject (i.e. more than one level deep)
+	InnerObjectWithIndirectlyNestedObject = CreateDefaultSubobject<USubobjectInstancingTestIndirectlyNestedObject>(TEXT("InnerObjectWithIndirectlyNestedObject"));
+}
+
+void USubobjectInstancingTestOuterObject::PostInitProperties()
+{
+	Super::PostInitProperties();
+
+	// Subobject that's deferred from native construction, but still instanced as part of UObject initialization flow
+	InnerObjectPostInit = NewObject<USubobjectInstancingTestObject>(this, TEXT("InnerObjectPostInit"));
+}
+
+USubobjectInstancingTestDirectlyNestedObject::USubobjectInstancingTestDirectlyNestedObject()
+{
+	// Nested reference to self
+	SelfRef = this;
+
+	// Nested instanced subobject that is expected to be a unique ptr value after construction
+	InnerObject = CreateDefaultSubobject<USubobjectInstancingTestObject>(TEXT("InnerObject"));
+}
+
+USubobjectInstancingTestIndirectlyNestedObject::USubobjectInstancingTestIndirectlyNestedObject()
+{
+	InnerObject = CreateDefaultSubobject<USubobjectInstancingTestDirectlyNestedObject>(TEXT("InnerObject"));
+	InnerObject->OwnerObject = this;
+}
+
+// Derived outer type that is expected to instance its 'InnerObject' property at construction time using a subtype of the archetype instance's type
+USubobjectInstancingTestDerivedOuterObjectWithTypeOverride::USubobjectInstancingTestDerivedOuterObjectWithTypeOverride(const FObjectInitializer& ObjectInitializer)
+	:Super(ObjectInitializer.SetDefaultSubobjectClass<USubobjectInstancingTestDerivedObject>("InnerObject"))
+{
+}
+
+// Derived outer type that is not expected to instance the 'InnerObject' property at construction time - the value should be set to NULL after being default-initialized
+USubobjectInstancingTestDerivedOuterObjectWithDoNotCreateOverride::USubobjectInstancingTestDerivedOuterObjectWithDoNotCreateOverride(const FObjectInitializer& ObjectInitializer)
+	:Super(ObjectInitializer.DoNotCreateDefaultSubobject("OptionalInnerObject"))
+{
+}
+
+UDynamicSubobjectInstancingTestClass::UDynamicSubobjectInstancingTestClass()
+{
+	bNeedsDynamicSubobjectInstancing = true;
+}
+
+namespace UE
+{
+
+namespace SubobjectInstancingTest::Private
+{
+	static const FName NonNativeSelfReferencePropertyName(TEXT("NonNativeSelfRef"));
+	static const FName NonNativeInnerObjectPropertyName(TEXT("NonNativeInnerObject"));
+	static const FName NonNativeOuterObjectPropertyName(TEXT("NonNativeOwnerObject"));
+	static const FName NonNativeEditTimeInnerObjectPropertyName(TEXT("NonNativeEditTimeInnerObject"));
+
+	template<typename ClassType>
+	static UClass* NewTestClass(UClass* SuperClass)
+	{
+		if (!SuperClass)
+		{
+			SuperClass = USubobjectInstancingTestObject::StaticClass();
+		}
+
+		UClass* TestClass = NewObject<ClassType>(GetTransientPackage(), NAME_None, RF_Public | RF_Transient);
+
+		// Simulate creation as a non-native (e.g. Blueprint) type
+		TestClass->ClassFlags |= (SuperClass->ClassFlags & CLASS_Inherit) | CLASS_CompiledFromBlueprint;
+		TestClass->ClassCastFlags |= SuperClass->ClassCastFlags;
+
+		// Hint to editor/runtime that the class may contain at least one reference to an instanced subobject (activates certain code paths)
+		TestClass->ClassFlags |= CLASS_HasInstancedReference;
+
+		FField::FLinkedListBuilder PropertyListBuilder(&TestClass->ChildProperties);
+
+		EPropertyFlags NonNativeObjectPropertyFlags = CPF_TObjectPtrWrapper;
+		if (!TestClass->ShouldUseDynamicSubobjectInstancing())
+		{
+			NonNativeObjectPropertyFlags |= CPF_InstancedReference | CPF_PersistentInstance;
+		}
+
+		// Note: This flag was meant to be temporary in lieu of dynamic instancing, but allows non-native properties to work with self-referencing
+		EPropertyFlags NonNativeSelfReferencePropertyFlags = NonNativeObjectPropertyFlags;
+		if (!TestClass->ShouldUseDynamicSubobjectInstancing())
+		{
+			NonNativeSelfReferencePropertyFlags |= CPF_AllowSelfReference;
+		}
+
+		// Non-native object property to store a non-native reference to self - TObjectPtr<UObject>
+		FObjectProperty* NonNativeSelfReferenceProperty = new FObjectProperty(TestClass, NonNativeSelfReferencePropertyName, RF_Public | RF_Transient);
+		NonNativeSelfReferenceProperty->SetPropertyFlags(NonNativeSelfReferencePropertyFlags);
+		NonNativeSelfReferenceProperty->SetPropertyClass(UObject::StaticClass());
+		PropertyListBuilder.AppendNoTerminate(*NonNativeSelfReferenceProperty);
+
+		// Non-native object reference property to store a non-native instantiation by ctor - TObjectPtr<USubobjectInstancingTestObject>
+		FObjectProperty* NonNativeInnerObjectProperty = new FObjectProperty(TestClass, NonNativeInnerObjectPropertyName, RF_Public | RF_Transient);
+		NonNativeInnerObjectProperty->SetPropertyFlags(NonNativeObjectPropertyFlags);
+		NonNativeInnerObjectProperty->SetPropertyClass(USubobjectInstancingTestObject::StaticClass());
+		PropertyListBuilder.AppendNoTerminate(*NonNativeInnerObjectProperty);
+
+		// Non-native object property to store a non-native reference to an outer object - TObjectPtr<UObject>
+		FObjectProperty* NonNativeOuterObjectProperty = new FObjectProperty(TestClass, NonNativeOuterObjectPropertyName, RF_Public | RF_Transient);
+		NonNativeOuterObjectProperty->SetPropertyFlags(NonNativeObjectPropertyFlags);
+		NonNativeOuterObjectProperty->SetPropertyClass(UObject::StaticClass());
+		PropertyListBuilder.AppendNoTerminate(*NonNativeOuterObjectProperty);
+
+		// Non-native object reference property to simulate deferred edit-time instantiation - TObjectPtr<USubobjectInstancingTestObject>
+		FObjectProperty* NonNativeEditTimeInnerObjectProperty = new FObjectProperty(TestClass, NonNativeEditTimeInnerObjectPropertyName, RF_Public | RF_Transient);
+		NonNativeEditTimeInnerObjectProperty->SetPropertyFlags(NonNativeObjectPropertyFlags);
+		NonNativeEditTimeInnerObjectProperty->SetPropertyClass(USubobjectInstancingTestObject::StaticClass());
+		PropertyListBuilder.AppendNoTerminate(*NonNativeEditTimeInnerObjectProperty);
+
+		TestClass->SetSuperStruct(SuperClass);
+		TestClass->Bind();
+		TestClass->StaticLink(/*bRelinkExistingProperties =*/ true);
+		TestClass->AssembleReferenceTokenStream();
+
+		UObject* TestCDO = TestClass->GetDefaultObject(/*bCreateIfNeeded  =*/ true);
+		TestClass->PostLoadDefaultObject(TestCDO);
+
+		// Simulate a non-native class constructor assignment of 'self'
+		NonNativeSelfReferenceProperty->SetObjectPropertyValue_InContainer(TestCDO, TestCDO);
+
+		// Simulate a non-native class constructor instantiation of a default subobject, which will be "captured" by the CDO and used to initialize all new objects of this type
+		// Note: Non-native default subobjects should not include the RF_DefaultSubObject flag - that flag is used to identify subobjects constructed using the CreateDefaultSubobject() API
+		static const FName NonNativeInnerObjectName(TEXT("NonNativeInnerObject"));
+		UObject* NonNativeInnerDefaultSubobject = NewObject<USubobjectInstancingTestObject>(TestCDO, NonNativeInnerObjectName, RF_ArchetypeObject | TestCDO->GetMaskedFlags(RF_PropagateToSubObjects));
+		NonNativeInnerObjectProperty->SetObjectPropertyValue_InContainer(TestCDO, NonNativeInnerDefaultSubobject);
+
+		// Simulate a non-native class constructor initialization of the outer object property's default value
+		NonNativeOuterObjectProperty->SetObjectPropertyValue_InContainer(TestCDO, nullptr);
+
+		// Simulate a non-native class constructor initialization of the edit-time reference property's default value
+		NonNativeEditTimeInnerObjectProperty->SetObjectPropertyValue_InContainer(TestCDO, nullptr);
+
+		return TestClass;
+	}
+
+	void RunSubobjectInstancingTests_Base(UClass* InstancingTestClass)
+	{
+		SECTION("Null subobject initialized from CDO")
+		{
+			USubobjectInstancingTestOuterObject* OuterObject = NewObject<USubobjectInstancingTestOuterObject>(GetTransientPackage(), InstancingTestClass);
+
+			CHECK(OuterObject->NullObject == nullptr);
+		}
+
+		SECTION("Null subobject initialized from template")
+		{
+			USubobjectInstancingTestOuterObject* Template = NewObject<USubobjectInstancingTestOuterObject>(GetTransientPackage(), InstancingTestClass, NAME_None, RF_ArchetypeObject);
+			USubobjectInstancingTestOuterObject* OuterObject = NewObject<USubobjectInstancingTestOuterObject>(GetTransientPackage(), InstancingTestClass, NAME_None, RF_NoFlags, Template);
+
+			CHECK(OuterObject->NullObject == nullptr);
+		}
+
+		SECTION("Default subobject initialized from CDO")
+		{
+			USubobjectInstancingTestOuterObject* OuterObject = NewObject<USubobjectInstancingTestOuterObject>(GetTransientPackage(), InstancingTestClass);
+
+			REQUIRE(OuterObject->InnerObject != nullptr);
+			CHECK(OuterObject->InnerObject->IsInOuter(OuterObject));
+
+			CHECK(OuterObject->InnerObject->TestValue == 100);
+
+			CHECK(OuterObject->InnerObject->HasAllFlags(RF_DefaultSubObject));
+			CHECK(OuterObject->InnerObject->HasAllFlags(OuterObject->GetMaskedFlags(RF_PropagateToSubObjects)));
+
+			CHECK_FALSE(OuterObject->InnerObject->IsTemplate());
+			CHECK(OuterObject->InnerObject->IsDefaultSubobject());
+			CHECK_FALSE(OuterObject->InnerObject->IsTemplateForSubobjects());
+
+			USubobjectInstancingTestOuterObject* NewOuterObject = NewObject<USubobjectInstancingTestOuterObject>(GetTransientPackage(), InstancingTestClass);
+			CHECK(OuterObject->InnerObject != NewOuterObject->InnerObject);
+		}
+
+		SECTION("Default subobject initialized from template")
+		{
+			USubobjectInstancingTestOuterObject* Template = NewObject<USubobjectInstancingTestOuterObject>(GetTransientPackage(), InstancingTestClass, NAME_None, RF_ArchetypeObject);
+			Template->InnerObject->TestValue = 200;
+
+			USubobjectInstancingTestOuterObject* OuterObject = NewObject<USubobjectInstancingTestOuterObject>(GetTransientPackage(), InstancingTestClass, NAME_None, RF_NoFlags, Template);
+
+			REQUIRE(OuterObject->InnerObject != nullptr);
+			CHECK(OuterObject->InnerObject->IsInOuter(OuterObject));
+
+			CHECK(OuterObject->InnerObject->TestValue == 200);
+
+			CHECK(OuterObject->InnerObject->HasAllFlags(RF_DefaultSubObject));
+			CHECK(OuterObject->InnerObject->HasAllFlags(OuterObject->GetMaskedFlags(RF_PropagateToSubObjects)));
+
+			CHECK_FALSE(OuterObject->InnerObject->IsTemplate());
+			CHECK(OuterObject->InnerObject->IsDefaultSubobject());
+			CHECK_FALSE(OuterObject->InnerObject->IsTemplateForSubobjects());
+
+			USubobjectInstancingTestOuterObject* NewOuterObject = NewObject<USubobjectInstancingTestOuterObject>(GetTransientPackage(), InstancingTestClass, NAME_None, RF_NoFlags, Template);
+			CHECK(OuterObject->InnerObject != NewOuterObject->InnerObject);
+		}
+
+		SECTION("Default subobject from type initialized from CDO")
+		{
+			USubobjectInstancingTestOuterObject* OuterObject = NewObject<USubobjectInstancingTestOuterObject>(GetTransientPackage(), InstancingTestClass);
+
+			REQUIRE(OuterObject->InnerObjectFromType != nullptr);
+			CHECK(OuterObject->InnerObjectFromType->IsInOuter(OuterObject));
+
+			CHECK(OuterObject->InnerObjectFromType->TestValue == 100);
+
+			CHECK(OuterObject->InnerObjectFromType->HasAllFlags(RF_DefaultSubObject));
+			CHECK(OuterObject->InnerObjectFromType->HasAllFlags(OuterObject->GetMaskedFlags(RF_PropagateToSubObjects)));
+
+			CHECK_FALSE(OuterObject->InnerObjectFromType->IsTemplate());
+			CHECK(OuterObject->InnerObjectFromType->IsDefaultSubobject());
+			CHECK_FALSE(OuterObject->InnerObjectFromType->IsTemplateForSubobjects());
+
+			USubobjectInstancingTestOuterObject* NewOuterObject = NewObject<USubobjectInstancingTestOuterObject>(GetTransientPackage(), InstancingTestClass);
+			CHECK(OuterObject->InnerObjectFromType != NewOuterObject->InnerObjectFromType);
+		}
+
+		SECTION("Default subobject from type initialized from template")
+		{
+			USubobjectInstancingTestOuterObject* Template = NewObject<USubobjectInstancingTestOuterObject>(GetTransientPackage(), InstancingTestClass, NAME_None, RF_ArchetypeObject);
+			Template->InnerObjectFromType->TestValue = 200;
+
+			USubobjectInstancingTestOuterObject* OuterObject = NewObject<USubobjectInstancingTestOuterObject>(GetTransientPackage(), InstancingTestClass, NAME_None, RF_NoFlags, Template);
+
+			REQUIRE(OuterObject->InnerObjectFromType != nullptr);
+			CHECK(OuterObject->InnerObjectFromType->IsInOuter(OuterObject));
+
+			CHECK(OuterObject->InnerObjectFromType->TestValue == 200);
+
+			CHECK(OuterObject->InnerObjectFromType->HasAllFlags(RF_DefaultSubObject));
+			CHECK(OuterObject->InnerObjectFromType->HasAllFlags(OuterObject->GetMaskedFlags(RF_PropagateToSubObjects)));
+
+			CHECK_FALSE(OuterObject->InnerObjectFromType->IsTemplate());
+			CHECK(OuterObject->InnerObjectFromType->IsDefaultSubobject());
+			CHECK_FALSE(OuterObject->InnerObjectFromType->IsTemplateForSubobjects());
+
+			USubobjectInstancingTestOuterObject* NewOuterObject = NewObject<USubobjectInstancingTestOuterObject>(GetTransientPackage(), InstancingTestClass, NAME_None, RF_NoFlags, Template);
+			CHECK(OuterObject->InnerObjectFromType != NewOuterObject->InnerObjectFromType);
+		}
+
+		SECTION("Shared default subobject initialized from CDO")
+		{
+			const USubobjectInstancingTestOuterObject* CDO = CastChecked<USubobjectInstancingTestOuterObject>(InstancingTestClass->GetDefaultObject());
+			USubobjectInstancingTestOuterObject* OuterObject = NewObject<USubobjectInstancingTestOuterObject>(GetTransientPackage(), InstancingTestClass);
+
+			// Since this reference is initialized by the native ctor, it will have constructed a new instance for the outer object. And since the default data object
+			// is the CDO in this case, InitProperties() won't copy the default value to this property when the outer object is constructed. As a result, this value
+			// will remain set to reference a unique instance, so the expected outcome is as if this property were marked as instanced. It's likely that this behavior
+			// diverged from the template case when the fast path was added to InitProperties(), as it relies on the property always being initialized from default data.
+			REQUIRE(OuterObject->SharedObject != nullptr);
+			CHECK(OuterObject->SharedObject->IsInOuter(OuterObject));
+
+			// Expected result is that we are not mutating the archetype object that's owned by the CDO in this case.
+			OuterObject->SharedObject->TestValue = 300;
+			CHECK(OuterObject->SharedObject->TestValue != CDO->SharedObject->TestValue);
+
+			CHECK_FALSE(OuterObject->SharedObject->HasAllFlags(RF_ArchetypeObject));
+			CHECK(OuterObject->SharedObject->HasAllFlags(RF_DefaultSubObject));
+			CHECK(OuterObject->SharedObject->HasAllFlags(OuterObject->GetMaskedFlags(RF_PropagateToSubObjects)));
+
+			CHECK_FALSE(OuterObject->SharedObject->IsTemplate());
+			CHECK(OuterObject->SharedObject->IsDefaultSubobject());
+			CHECK_FALSE(OuterObject->SharedObject->IsTemplateForSubobjects());
+
+			USubobjectInstancingTestOuterObject* NewOuterObject = NewObject<USubobjectInstancingTestOuterObject>(GetTransientPackage(), InstancingTestClass);
+			CHECK(OuterObject->SharedObject != NewOuterObject->SharedObject);
+		}
+
+		SECTION("Shared default subobject initialized from template")
+		{
+			USubobjectInstancingTestOuterObject* Template = NewObject<USubobjectInstancingTestOuterObject>(GetTransientPackage(), InstancingTestClass, NAME_None, RF_ArchetypeObject);
+			Template->SharedObject->TestValue = 200;
+
+			USubobjectInstancingTestOuterObject* OuterObject = NewObject<USubobjectInstancingTestOuterObject>(GetTransientPackage(), InstancingTestClass, NAME_None, RF_NoFlags, Template);
+
+			REQUIRE(OuterObject->SharedObject != nullptr);
+
+			// Types using dynamic instancing do not support referencing back to the source archetype.
+			if (InstancingTestClass->ShouldUseDynamicSubobjectInstancing())
+			{
+				// Expected result should match an instanced reference in this case.
+				CHECK(OuterObject->SharedObject->IsInOuter(OuterObject));
+			}
+			else    // legacy path
+			{
+				// This reference is first initialized by the native ctor, which constructs a new instance for the outer object. But since the default data object is
+				// not the CDO in this case, InitProperties() will copy the template's reference to this property when the outer object is constructed. Because this
+				// property is not also marked as instanced, it will not be included in subobject instancing, and should equate to the template's value as a result.
+				CHECK(OuterObject->SharedObject->IsInOuter(Template));
+
+				// Expected result is that we are mutating the archetype object that's owned by the template in this case.
+				OuterObject->SharedObject->TestValue = 300;
+				CHECK(OuterObject->SharedObject->TestValue == Template->SharedObject->TestValue);
+
+				CHECK(OuterObject->SharedObject->HasAllFlags(RF_ArchetypeObject));
+				CHECK(OuterObject->SharedObject->HasAllFlags(RF_DefaultSubObject));
+				CHECK(OuterObject->SharedObject->HasAllFlags(OuterObject->GetMaskedFlags(RF_PropagateToSubObjects)));
+
+				CHECK(OuterObject->SharedObject->IsTemplate());
+				CHECK(OuterObject->SharedObject->IsDefaultSubobject());
+				CHECK(OuterObject->SharedObject->IsTemplateForSubobjects());
+
+				USubobjectInstancingTestOuterObject* NewOuterObject = NewObject<USubobjectInstancingTestOuterObject>(GetTransientPackage(), InstancingTestClass, NAME_None, RF_NoFlags, Template);
+				CHECK(OuterObject->SharedObject == NewOuterObject->SharedObject);
+			}
+		}
+
+		SECTION("Reference to external object initialized from CDO")
+		{
+			USubobjectInstancingTestOuterObject* OuterObject = NewObject<USubobjectInstancingTestOuterObject>(GetTransientPackage(), InstancingTestClass);
+
+			REQUIRE(OuterObject->ExternalObject != nullptr);
+			CHECK_FALSE(OuterObject->ExternalObject->IsInOuter(OuterObject));
+		}
+
+		SECTION("Reference to external object initialized from template")
+		{
+			USubobjectInstancingTestOuterObject* Template = NewObject<USubobjectInstancingTestOuterObject>(GetTransientPackage(), InstancingTestClass, NAME_None, RF_ArchetypeObject);
+			USubobjectInstancingTestOuterObject* OuterObject = NewObject<USubobjectInstancingTestOuterObject>(GetTransientPackage(), InstancingTestClass, NAME_None, RF_NoFlags, Template);
+
+			REQUIRE(OuterObject->ExternalObject != nullptr);
+			CHECK_FALSE(OuterObject->ExternalObject->IsInOuter(Template));
+			CHECK_FALSE(OuterObject->ExternalObject->IsInOuter(OuterObject));
+
+			CHECK(OuterObject->ExternalObject == Template->ExternalObject);
+		}
+
+		SECTION("Transient default subobject initialized from CDO")
+		{
+			USubobjectInstancingTestOuterObject* OuterObject = NewObject<USubobjectInstancingTestOuterObject>(GetTransientPackage(), InstancingTestClass);
+
+			REQUIRE(OuterObject->TransientInnerObject != nullptr);
+			CHECK(OuterObject->TransientInnerObject->IsInOuter(OuterObject));
+
+			CHECK(OuterObject->TransientInnerObject->TestValue == 100);
+
+			CHECK(OuterObject->TransientInnerObject->HasAllFlags(RF_DefaultSubObject));
+			CHECK(OuterObject->TransientInnerObject->HasAllFlags(OuterObject->GetMaskedFlags(RF_PropagateToSubObjects)));
+
+			// This must otherwise be explicitly set at construction/instancing time, ensuring here that it's not implied
+			CHECK_FALSE(OuterObject->TransientInnerObject->HasAnyFlags(RF_Transient));
+
+			CHECK_FALSE(OuterObject->TransientInnerObject->IsTemplate());
+			CHECK(OuterObject->TransientInnerObject->IsDefaultSubobject());
+			CHECK_FALSE(OuterObject->TransientInnerObject->IsTemplateForSubobjects());
+
+			USubobjectInstancingTestOuterObject* NewOuterObject = NewObject<USubobjectInstancingTestOuterObject>(GetTransientPackage(), InstancingTestClass);
+			CHECK(OuterObject->TransientInnerObject != NewOuterObject->TransientInnerObject);
+		}
+
+		SECTION("Transient default subobject initialized from template")
+		{
+			USubobjectInstancingTestOuterObject* Template = NewObject<USubobjectInstancingTestOuterObject>(GetTransientPackage(), InstancingTestClass, NAME_None, RF_ArchetypeObject);
+			Template->TransientInnerObject->TestValue = 200;
+
+			USubobjectInstancingTestOuterObject* OuterObject = NewObject<USubobjectInstancingTestOuterObject>(GetTransientPackage(), InstancingTestClass, NAME_None, RF_NoFlags, Template);
+
+			REQUIRE(OuterObject->TransientInnerObject != nullptr);
+			CHECK(OuterObject->TransientInnerObject->IsInOuter(OuterObject));
+
+			CHECK(OuterObject->TransientInnerObject->TestValue == 200);
+
+			CHECK(OuterObject->TransientInnerObject->HasAllFlags(RF_DefaultSubObject));
+			CHECK(OuterObject->TransientInnerObject->HasAllFlags(OuterObject->GetMaskedFlags(RF_PropagateToSubObjects)));
+
+			// This must otherwise be explicitly set at construction/instancing time, ensuring here that it's not implied
+			CHECK_FALSE(OuterObject->TransientInnerObject->HasAnyFlags(RF_Transient));
+
+			CHECK_FALSE(OuterObject->TransientInnerObject->IsTemplate());
+			CHECK(OuterObject->TransientInnerObject->IsDefaultSubobject());
+			CHECK_FALSE(OuterObject->TransientInnerObject->IsTemplateForSubobjects());
+
+			USubobjectInstancingTestOuterObject* NewOuterObject = NewObject<USubobjectInstancingTestOuterObject>(GetTransientPackage(), InstancingTestClass);
+			CHECK(OuterObject->TransientInnerObject != NewOuterObject->TransientInnerObject);
+		}
+
+		SECTION("Local-only default subobject initialized from template with transient data")
+		{
+			USubobjectInstancingTestOuterObject* Template = NewObject<USubobjectInstancingTestOuterObject>(GetTransientPackage(), InstancingTestClass, NAME_None, RF_ArchetypeObject);
+
+			REQUIRE(Template->LocalOnlyInnerObject != nullptr);
+			CHECK(Template->LocalOnlyInnerObject->TestValue == 100);
+
+			Template->LocalOnlyInnerObject->TestValue = 200;
+
+			USubobjectInstancingTestOuterObject* OuterObject = NewObject<USubobjectInstancingTestOuterObject>(GetTransientPackage(), InstancingTestClass, NAME_None, RF_NoFlags, Template);
+
+			CHECK(OuterObject->LocalOnlyInnerObject != nullptr);
+			CHECK(OuterObject->LocalOnlyInnerObject->IsInOuter(OuterObject));
+
+			if (InstancingTestClass->HasAnyClassFlags(CLASS_CompiledFromBlueprint))
+			{
+				// BP/non-native types do not support the transient option to CreateDefaultSubobject(). The reference will just be reconstructed at instancing time using the template's data.
+				// This is because we currently skip the code that looks for an existing instance on the outer object before construction (see FObjectInstancingGraph::GetInstancedSubobject()).
+				CHECK(OuterObject->LocalOnlyInnerObject->TestValue == 200);
+			}
+			else
+			{
+				// For native types, because we added the transient argument in the constructor, instance data is not expected to propagate through the template. It should keep the default value.
+				CHECK(OuterObject->LocalOnlyInnerObject->TestValue == 100);
+			}
+
+			CHECK_FALSE(OuterObject->LocalOnlyInnerObject->IsTemplate());
+			CHECK(OuterObject->LocalOnlyInnerObject->IsDefaultSubobject());
+			CHECK_FALSE(OuterObject->LocalOnlyInnerObject->IsTemplateForSubobjects());
+		}
+
+		SECTION("Native property initialized from CDO using NewObject()")
+		{
+			USubobjectInstancingTestOuterObject* OuterObject = NewObject<USubobjectInstancingTestOuterObject>(GetTransientPackage(), InstancingTestClass);
+
+			REQUIRE(OuterObject->InnerObjectUsingNew != nullptr);
+			CHECK(OuterObject->InnerObjectUsingNew->IsInOuter(OuterObject));
+
+			CHECK(OuterObject->InnerObjectUsingNew->TestValue == 100);
+
+			CHECK_FALSE(OuterObject->InnerObjectUsingNew->HasAllFlags(RF_DefaultSubObject));
+			CHECK(OuterObject->InnerObjectUsingNew->HasAllFlags(OuterObject->GetMaskedFlags(RF_PropagateToSubObjects)));
+
+			CHECK_FALSE(OuterObject->InnerObjectUsingNew->IsTemplate());
+			CHECK(OuterObject->InnerObjectUsingNew->IsDefaultSubobject());
+			CHECK_FALSE(OuterObject->InnerObjectUsingNew->IsTemplateForSubobjects());
+
+			USubobjectInstancingTestOuterObject* NewOuterObject = NewObject<USubobjectInstancingTestOuterObject>(GetTransientPackage(), InstancingTestClass);
+			CHECK(OuterObject->InnerObjectUsingNew != NewOuterObject->InnerObjectUsingNew);
+		}
+
+		SECTION("Native property initialized from template using NewObject()")
+		{
+			USubobjectInstancingTestOuterObject* Template = NewObject<USubobjectInstancingTestOuterObject>(GetTransientPackage(), InstancingTestClass, NAME_None, RF_ArchetypeObject);
+			Template->InnerObjectUsingNew->TestValue = 200;
+
+			USubobjectInstancingTestOuterObject* OuterObject = NewObject<USubobjectInstancingTestOuterObject>(GetTransientPackage(), InstancingTestClass, NAME_None, RF_NoFlags, Template);
+
+			REQUIRE(OuterObject->InnerObjectUsingNew != nullptr);
+			CHECK(OuterObject->InnerObjectUsingNew->IsInOuter(OuterObject));
+
+			if (InstancingTestClass->HasAnyClassFlags(CLASS_CompiledFromBlueprint))
+			{
+				// BP/non-native types behave differently and will not currently search for an existing subobject at instancing time
+				// before reconstructing the reference using the template object - see FObjectInstancingGraph::GetInstancedSubobject().
+				CHECK(OuterObject->InnerObjectUsingNew->TestValue == 200);
+			}
+			else
+			{
+				// Note: Because we did not construct using CreateDefaultSubobject(), instance data will not propagate through the template.
+				// Using NewObject() is effectively the same result as passing TRUE for the transient argument to CreateDefaultSubobject().
+				CHECK(OuterObject->InnerObjectUsingNew->TestValue == 100);
+			}
+
+			CHECK_FALSE(OuterObject->InnerObjectUsingNew->HasAllFlags(RF_DefaultSubObject));
+			CHECK(OuterObject->InnerObjectUsingNew->HasAllFlags(OuterObject->GetMaskedFlags(RF_PropagateToSubObjects)));
+
+			CHECK_FALSE(OuterObject->InnerObjectUsingNew->IsTemplate());
+			CHECK(OuterObject->InnerObjectUsingNew->IsDefaultSubobject());
+			CHECK_FALSE(OuterObject->InnerObjectUsingNew->IsTemplateForSubobjects());
+
+			USubobjectInstancingTestOuterObject* NewOuterObject = NewObject<USubobjectInstancingTestOuterObject>(GetTransientPackage(), InstancingTestClass, NAME_None, RF_NoFlags, Template);
+			CHECK(OuterObject->InnerObjectUsingNew != NewOuterObject->InnerObjectUsingNew);
+		}
+	}
+
+	void RunSubobjectInstancingTests_Array(UClass* InstancingTestClass)
+	{
+		SECTION("Array of default subobjects initialized from CDO")
+		{
+			USubobjectInstancingTestOuterObject* OuterObject = NewObject<USubobjectInstancingTestOuterObject>(GetTransientPackage(), InstancingTestClass);
+
+			CHECK_FALSE(OuterObject->InnerObjectArray.IsEmpty());
+			for (USubobjectInstancingTestObject* InnerObject : OuterObject->InnerObjectArray)
+			{
+				REQUIRE(InnerObject != nullptr);
+				CHECK(InnerObject->IsInOuter(OuterObject));
+
+				CHECK(InnerObject->TestValue == 100);
+
+				CHECK(InnerObject->HasAllFlags(RF_DefaultSubObject));
+				CHECK(InnerObject->HasAllFlags(OuterObject->GetMaskedFlags(RF_PropagateToSubObjects)));
+
+				CHECK_FALSE(InnerObject->IsTemplate());
+				CHECK(InnerObject->IsDefaultSubobject());
+				CHECK_FALSE(InnerObject->IsTemplateForSubobjects());
+			}
+		}
+
+		SECTION("Array of default subobjects initialized from template")
+		{
+			USubobjectInstancingTestOuterObject* Template = NewObject<USubobjectInstancingTestOuterObject>(GetTransientPackage(), InstancingTestClass, NAME_None, RF_ArchetypeObject);
+			for (USubobjectInstancingTestObject* InnerObject : Template->InnerObjectArray)
+			{
+				InnerObject->TestValue = 200;
+			}
+
+			USubobjectInstancingTestOuterObject* OuterObject = NewObject<USubobjectInstancingTestOuterObject>(GetTransientPackage(), InstancingTestClass, NAME_None, RF_NoFlags, Template);
+
+			CHECK_FALSE(OuterObject->InnerObjectArray.IsEmpty());
+			for (USubobjectInstancingTestObject* InnerObject : OuterObject->InnerObjectArray)
+			{
+				REQUIRE(InnerObject != nullptr);
+				CHECK(InnerObject->IsInOuter(OuterObject));
+
+				CHECK(InnerObject->TestValue == 200);
+
+				CHECK(InnerObject->HasAllFlags(RF_DefaultSubObject));
+				CHECK(InnerObject->HasAllFlags(OuterObject->GetMaskedFlags(RF_PropagateToSubObjects)));
+
+				CHECK_FALSE(InnerObject->IsTemplate());
+				CHECK(InnerObject->IsDefaultSubobject());
+				CHECK_FALSE(InnerObject->IsTemplateForSubobjects());
+			}
+		}
+
+		SECTION("Array of default subobjects from type initialized from CDO")
+		{
+			USubobjectInstancingTestOuterObject* OuterObject = NewObject<USubobjectInstancingTestOuterObject>(GetTransientPackage(), InstancingTestClass);
+
+			CHECK_FALSE(OuterObject->InnerObjectFromTypeArray.IsEmpty());
+			for (USubobjectInstancingTestObject* InnerObjectFromType : OuterObject->InnerObjectFromTypeArray)
+			{
+				REQUIRE(InnerObjectFromType != nullptr);
+				CHECK(InnerObjectFromType->IsInOuter(OuterObject));
+
+				CHECK(InnerObjectFromType->TestValue == 100);
+
+				CHECK(InnerObjectFromType->HasAllFlags(RF_DefaultSubObject));
+				CHECK(InnerObjectFromType->HasAllFlags(OuterObject->GetMaskedFlags(RF_PropagateToSubObjects)));
+
+				CHECK_FALSE(InnerObjectFromType->IsTemplate());
+				CHECK(InnerObjectFromType->IsDefaultSubobject());
+				CHECK_FALSE(InnerObjectFromType->IsTemplateForSubobjects());
+			}
+		}
+
+		SECTION("Array of default subobjects from type initialized from template")
+		{
+			USubobjectInstancingTestOuterObject* Template = NewObject<USubobjectInstancingTestOuterObject>(GetTransientPackage(), InstancingTestClass, NAME_None, RF_ArchetypeObject);
+			for (USubobjectInstancingTestObject* InnerObjectFromType : Template->InnerObjectFromTypeArray)
+			{
+				InnerObjectFromType->TestValue = 200;
+			}
+
+			USubobjectInstancingTestOuterObject* OuterObject = NewObject<USubobjectInstancingTestOuterObject>(GetTransientPackage(), InstancingTestClass, NAME_None, RF_NoFlags, Template);
+
+			CHECK_FALSE(OuterObject->InnerObjectFromTypeArray.IsEmpty());
+			for (USubobjectInstancingTestObject* InnerObjectFromType : OuterObject->InnerObjectFromTypeArray)
+			{
+				REQUIRE(InnerObjectFromType != nullptr);
+				CHECK(InnerObjectFromType->IsInOuter(OuterObject));
+
+				CHECK(InnerObjectFromType->TestValue == 200);
+
+				CHECK(InnerObjectFromType->HasAllFlags(RF_DefaultSubObject));
+				CHECK(InnerObjectFromType->HasAllFlags(OuterObject->GetMaskedFlags(RF_PropagateToSubObjects)));
+
+				CHECK_FALSE(InnerObjectFromType->IsTemplate());
+				CHECK(InnerObjectFromType->IsDefaultSubobject());
+				CHECK_FALSE(InnerObjectFromType->IsTemplateForSubobjects());
+			}
+		}
+	}
+
+	void RunSubobjectInstancingTests_Set(UClass* InstancingTestClass)
+	{
+		SECTION("Set of default subobjects initialized from CDO")
+		{
+			USubobjectInstancingTestOuterObject* OuterObject = NewObject<USubobjectInstancingTestOuterObject>(GetTransientPackage(), InstancingTestClass);
+
+			CHECK_FALSE(OuterObject->InnerObjectSet.IsEmpty());
+			for (USubobjectInstancingTestObject* InnerObject : OuterObject->InnerObjectSet)
+			{
+				REQUIRE(InnerObject != nullptr);
+				CHECK(InnerObject->IsInOuter(OuterObject));
+
+				CHECK(InnerObject->TestValue == 100);
+
+				CHECK(InnerObject->HasAllFlags(RF_DefaultSubObject));
+				CHECK(InnerObject->HasAllFlags(OuterObject->GetMaskedFlags(RF_PropagateToSubObjects)));
+
+				CHECK_FALSE(InnerObject->IsTemplate());
+				CHECK(InnerObject->IsDefaultSubobject());
+				CHECK_FALSE(InnerObject->IsTemplateForSubobjects());
+			}
+		}
+
+		SECTION("Set of default subobjects initialized from template")
+		{
+			USubobjectInstancingTestOuterObject* Template = NewObject<USubobjectInstancingTestOuterObject>(GetTransientPackage(), InstancingTestClass, NAME_None, RF_ArchetypeObject);
+			for (USubobjectInstancingTestObject* InnerObject : Template->InnerObjectSet)
+			{
+				InnerObject->TestValue = 200;
+			}
+
+			USubobjectInstancingTestOuterObject* OuterObject = NewObject<USubobjectInstancingTestOuterObject>(GetTransientPackage(), InstancingTestClass, NAME_None, RF_NoFlags, Template);
+
+			CHECK_FALSE(OuterObject->InnerObjectSet.IsEmpty());
+			for (USubobjectInstancingTestObject* InnerObject : OuterObject->InnerObjectSet)
+			{
+				REQUIRE(InnerObject != nullptr);
+				CHECK(InnerObject->IsInOuter(OuterObject));
+
+				CHECK(InnerObject->TestValue == 200);
+
+				CHECK(InnerObject->HasAllFlags(RF_DefaultSubObject));
+				CHECK(InnerObject->HasAllFlags(OuterObject->GetMaskedFlags(RF_PropagateToSubObjects)));
+
+				CHECK_FALSE(InnerObject->IsTemplate());
+				CHECK(InnerObject->IsDefaultSubobject());
+				CHECK_FALSE(InnerObject->IsTemplateForSubobjects());
+			}
+		}
+
+		SECTION("Set of default subobjects from type initialized from CDO")
+		{
+			USubobjectInstancingTestOuterObject* OuterObject = NewObject<USubobjectInstancingTestOuterObject>(GetTransientPackage(), InstancingTestClass);
+
+			CHECK_FALSE(OuterObject->InnerObjectFromTypeSet.IsEmpty());
+			for (USubobjectInstancingTestObject* InnerObjectFromType : OuterObject->InnerObjectFromTypeSet)
+			{
+				REQUIRE(InnerObjectFromType != nullptr);
+				CHECK(InnerObjectFromType->IsInOuter(OuterObject));
+
+				CHECK(InnerObjectFromType->TestValue == 100);
+
+				CHECK(InnerObjectFromType->HasAllFlags(RF_DefaultSubObject));
+				CHECK(InnerObjectFromType->HasAllFlags(OuterObject->GetMaskedFlags(RF_PropagateToSubObjects)));
+
+				CHECK_FALSE(InnerObjectFromType->IsTemplate());
+				CHECK(InnerObjectFromType->IsDefaultSubobject());
+				CHECK_FALSE(InnerObjectFromType->IsTemplateForSubobjects());
+			}
+		}
+
+		SECTION("Set of default subobjects from type initialized from template")
+		{
+			USubobjectInstancingTestOuterObject* Template = NewObject<USubobjectInstancingTestOuterObject>(GetTransientPackage(), InstancingTestClass, NAME_None, RF_ArchetypeObject);
+			for (USubobjectInstancingTestObject* InnerObjectFromType : Template->InnerObjectFromTypeSet)
+			{
+				InnerObjectFromType->TestValue = 200;
+			}
+
+			USubobjectInstancingTestOuterObject* OuterObject = NewObject<USubobjectInstancingTestOuterObject>(GetTransientPackage(), InstancingTestClass, NAME_None, RF_NoFlags, Template);
+
+			CHECK_FALSE(OuterObject->InnerObjectFromTypeSet.IsEmpty());
+			for (USubobjectInstancingTestObject* InnerObjectFromType : OuterObject->InnerObjectFromTypeSet)
+			{
+				REQUIRE(InnerObjectFromType != nullptr);
+				CHECK(InnerObjectFromType->IsInOuter(OuterObject));
+
+				CHECK(InnerObjectFromType->TestValue == 200);
+
+				CHECK(InnerObjectFromType->HasAllFlags(RF_DefaultSubObject));
+				CHECK(InnerObjectFromType->HasAllFlags(OuterObject->GetMaskedFlags(RF_PropagateToSubObjects)));
+
+				CHECK_FALSE(InnerObjectFromType->IsTemplate());
+				CHECK(InnerObjectFromType->IsDefaultSubobject());
+				CHECK_FALSE(InnerObjectFromType->IsTemplateForSubobjects());
+			}
+		}
+	}
+
+	void RunSubobjectInstancingTests_Map(UClass* InstancingTestClass)
+	{
+		SECTION("Map of default subobjects initialized from CDO")
+		{
+			USubobjectInstancingTestOuterObject* OuterObject = NewObject<USubobjectInstancingTestOuterObject>(GetTransientPackage(), InstancingTestClass);
+
+			CHECK_FALSE(OuterObject->InnerObjectMap.IsEmpty());
+			for (TPair<USubobjectInstancingTestObject*, USubobjectInstancingTestObject*> InnerObjectPair : OuterObject->InnerObjectMap)
+			{
+				REQUIRE(InnerObjectPair.Key != nullptr);
+				CHECK(InnerObjectPair.Key->IsInOuter(OuterObject));
+
+				CHECK(InnerObjectPair.Key->TestValue == 100);
+
+				CHECK(InnerObjectPair.Key->HasAllFlags(RF_DefaultSubObject));
+				CHECK(InnerObjectPair.Key->HasAllFlags(OuterObject->GetMaskedFlags(RF_PropagateToSubObjects)));
+
+				CHECK_FALSE(InnerObjectPair.Key->IsTemplate());
+				CHECK(InnerObjectPair.Key->IsDefaultSubobject());
+				CHECK_FALSE(InnerObjectPair.Key->IsTemplateForSubobjects());
+
+				REQUIRE(InnerObjectPair.Value != nullptr);
+				CHECK(InnerObjectPair.Value->IsInOuter(OuterObject));
+
+				CHECK(InnerObjectPair.Value->TestValue == 100);
+
+				CHECK(InnerObjectPair.Value->HasAllFlags(RF_DefaultSubObject));
+				CHECK(InnerObjectPair.Value->HasAllFlags(OuterObject->GetMaskedFlags(RF_PropagateToSubObjects)));
+
+				CHECK_FALSE(InnerObjectPair.Value->IsTemplate());
+				CHECK(InnerObjectPair.Value->IsDefaultSubobject());
+				CHECK_FALSE(InnerObjectPair.Value->IsTemplateForSubobjects());
+			}
+		}
+
+		SECTION("Map of default subobjects initialized from template")
+		{
+			USubobjectInstancingTestOuterObject* Template = NewObject<USubobjectInstancingTestOuterObject>(GetTransientPackage(), InstancingTestClass, NAME_None, RF_ArchetypeObject);
+			for (TPair<USubobjectInstancingTestObject*, USubobjectInstancingTestObject*> InnerObjectPair : Template->InnerObjectMap)
+			{
+				InnerObjectPair.Key->TestValue = 200;
+				InnerObjectPair.Value->TestValue = 300;
+			}
+
+			USubobjectInstancingTestOuterObject* OuterObject = NewObject<USubobjectInstancingTestOuterObject>(GetTransientPackage(), InstancingTestClass, NAME_None, RF_NoFlags, Template);
+
+			CHECK_FALSE(OuterObject->InnerObjectMap.IsEmpty());
+			for (TPair<USubobjectInstancingTestObject*, USubobjectInstancingTestObject*> InnerObjectPair : OuterObject->InnerObjectMap)
+			{
+				REQUIRE(InnerObjectPair.Key != nullptr);
+				CHECK(InnerObjectPair.Key->IsInOuter(OuterObject));
+
+				CHECK(InnerObjectPair.Key->TestValue == 200);
+
+				CHECK(InnerObjectPair.Key->HasAllFlags(RF_DefaultSubObject));
+				CHECK(InnerObjectPair.Key->HasAllFlags(OuterObject->GetMaskedFlags(RF_PropagateToSubObjects)));
+
+				CHECK_FALSE(InnerObjectPair.Key->IsTemplate());
+				CHECK(InnerObjectPair.Key->IsDefaultSubobject());
+				CHECK_FALSE(InnerObjectPair.Key->IsTemplateForSubobjects());
+
+				REQUIRE(InnerObjectPair.Value != nullptr);
+				CHECK(InnerObjectPair.Value->IsInOuter(OuterObject));
+
+				CHECK(InnerObjectPair.Value->TestValue == 300);
+
+				CHECK(InnerObjectPair.Value->HasAllFlags(RF_DefaultSubObject));
+				CHECK(InnerObjectPair.Value->HasAllFlags(OuterObject->GetMaskedFlags(RF_PropagateToSubObjects)));
+
+				CHECK_FALSE(InnerObjectPair.Value->IsTemplate());
+				CHECK(InnerObjectPair.Value->IsDefaultSubobject());
+				CHECK_FALSE(InnerObjectPair.Value->IsTemplateForSubobjects());
+			}
+		}
+
+		SECTION("Map of default subobjects from type initialized from CDO")
+		{
+			USubobjectInstancingTestOuterObject* OuterObject = NewObject<USubobjectInstancingTestOuterObject>(GetTransientPackage(), InstancingTestClass);
+
+			CHECK_FALSE(OuterObject->InnerObjectFromTypeMap.IsEmpty());
+			for (TPair<USubobjectInstancingTestObject*, USubobjectInstancingTestObject*> InnerObjectFromTypePair : OuterObject->InnerObjectFromTypeMap)
+			{
+				REQUIRE(InnerObjectFromTypePair.Key != nullptr);
+				CHECK(InnerObjectFromTypePair.Key->IsInOuter(OuterObject));
+
+				CHECK(InnerObjectFromTypePair.Key->TestValue == 100);
+
+				CHECK(InnerObjectFromTypePair.Key->HasAllFlags(RF_DefaultSubObject));
+				CHECK(InnerObjectFromTypePair.Key->HasAllFlags(OuterObject->GetMaskedFlags(RF_PropagateToSubObjects)));
+
+				CHECK_FALSE(InnerObjectFromTypePair.Key->IsTemplate());
+				CHECK(InnerObjectFromTypePair.Key->IsDefaultSubobject());
+				CHECK_FALSE(InnerObjectFromTypePair.Key->IsTemplateForSubobjects());
+
+				REQUIRE(InnerObjectFromTypePair.Value != nullptr);
+				CHECK(InnerObjectFromTypePair.Value->IsInOuter(OuterObject));
+
+				CHECK(InnerObjectFromTypePair.Value->TestValue == 100);
+
+				CHECK(InnerObjectFromTypePair.Value->HasAllFlags(RF_DefaultSubObject));
+				CHECK(InnerObjectFromTypePair.Value->HasAllFlags(OuterObject->GetMaskedFlags(RF_PropagateToSubObjects)));
+
+				CHECK_FALSE(InnerObjectFromTypePair.Value->IsTemplate());
+				CHECK(InnerObjectFromTypePair.Value->IsDefaultSubobject());
+				CHECK_FALSE(InnerObjectFromTypePair.Value->IsTemplateForSubobjects());
+			}
+		}
+
+		SECTION("Map of default subobjects from type initialized from template")
+		{
+			USubobjectInstancingTestOuterObject* Template = NewObject<USubobjectInstancingTestOuterObject>(GetTransientPackage(), InstancingTestClass, NAME_None, RF_ArchetypeObject);
+			for (TPair<USubobjectInstancingTestObject*, USubobjectInstancingTestObject*> InnerObjectFromTypePair : Template->InnerObjectFromTypeMap)
+			{
+				InnerObjectFromTypePair.Key->TestValue = 200;
+				InnerObjectFromTypePair.Value->TestValue = 300;
+			}
+
+			USubobjectInstancingTestOuterObject* OuterObject = NewObject<USubobjectInstancingTestOuterObject>(GetTransientPackage(), InstancingTestClass, NAME_None, RF_NoFlags, Template);
+
+			CHECK_FALSE(OuterObject->InnerObjectFromTypeMap.IsEmpty());
+			for (TPair<USubobjectInstancingTestObject*, USubobjectInstancingTestObject*> InnerObjectFromTypePair : OuterObject->InnerObjectFromTypeMap)
+			{
+				REQUIRE(InnerObjectFromTypePair.Key != nullptr);
+				CHECK(InnerObjectFromTypePair.Key->IsInOuter(OuterObject));
+
+				CHECK(InnerObjectFromTypePair.Key->TestValue == 200);
+
+				CHECK(InnerObjectFromTypePair.Key->HasAllFlags(RF_DefaultSubObject));
+				CHECK(InnerObjectFromTypePair.Key->HasAllFlags(OuterObject->GetMaskedFlags(RF_PropagateToSubObjects)));
+
+				CHECK_FALSE(InnerObjectFromTypePair.Key->IsTemplate());
+				CHECK(InnerObjectFromTypePair.Key->IsDefaultSubobject());
+				CHECK_FALSE(InnerObjectFromTypePair.Key->IsTemplateForSubobjects());
+
+				REQUIRE(InnerObjectFromTypePair.Value != nullptr);
+				CHECK(InnerObjectFromTypePair.Value->IsInOuter(OuterObject));
+
+				CHECK(InnerObjectFromTypePair.Value->TestValue == 300);
+
+				CHECK(InnerObjectFromTypePair.Value->HasAllFlags(RF_DefaultSubObject));
+				CHECK(InnerObjectFromTypePair.Value->HasAllFlags(OuterObject->GetMaskedFlags(RF_PropagateToSubObjects)));
+
+				CHECK_FALSE(InnerObjectFromTypePair.Value->IsTemplate());
+				CHECK(InnerObjectFromTypePair.Value->IsDefaultSubobject());
+				CHECK_FALSE(InnerObjectFromTypePair.Value->IsTemplateForSubobjects());
+			}
+		}
+	}
+
+	void RunSubobjectInstancingTests_Struct(UClass* InstancingTestClass)
+	{
+		SECTION("Struct member with default subobject initialized from CDO")
+		{
+			USubobjectInstancingTestOuterObject* OuterObject = NewObject<USubobjectInstancingTestOuterObject>(GetTransientPackage(), InstancingTestClass);
+
+			REQUIRE(OuterObject->StructWithInnerObjects.InnerObject != nullptr);
+			CHECK(OuterObject->StructWithInnerObjects.InnerObject->IsInOuter(OuterObject));
+
+			CHECK(OuterObject->StructWithInnerObjects.InnerObject->TestValue == 100);
+
+			CHECK(OuterObject->StructWithInnerObjects.InnerObject->HasAllFlags(RF_DefaultSubObject));
+			CHECK(OuterObject->StructWithInnerObjects.InnerObject->HasAllFlags(OuterObject->GetMaskedFlags(RF_PropagateToSubObjects)));
+
+			CHECK_FALSE(OuterObject->StructWithInnerObjects.InnerObject->IsTemplate());
+			CHECK(OuterObject->StructWithInnerObjects.InnerObject->IsDefaultSubobject());
+			CHECK_FALSE(OuterObject->StructWithInnerObjects.InnerObject->IsTemplateForSubobjects());
+
+			USubobjectInstancingTestOuterObject* NewOuterObject = NewObject<USubobjectInstancingTestOuterObject>(GetTransientPackage(), InstancingTestClass);
+			CHECK(OuterObject->StructWithInnerObjects.InnerObject != NewOuterObject->StructWithInnerObjects.InnerObject);
+		}
+
+		SECTION("Struct member with default subobject initialized from template")
+		{
+			USubobjectInstancingTestOuterObject* Template = NewObject<USubobjectInstancingTestOuterObject>(GetTransientPackage(), InstancingTestClass, NAME_None, RF_ArchetypeObject);
+			Template->StructWithInnerObjects.InnerObject->TestValue = 200;
+
+			USubobjectInstancingTestOuterObject* OuterObject = NewObject<USubobjectInstancingTestOuterObject>(GetTransientPackage(), InstancingTestClass, NAME_None, RF_NoFlags, Template);
+
+			REQUIRE(OuterObject->StructWithInnerObjects.InnerObject != nullptr);
+			CHECK(OuterObject->StructWithInnerObjects.InnerObject->IsInOuter(OuterObject));
+
+			CHECK(OuterObject->StructWithInnerObjects.InnerObject->TestValue == 200);
+
+			CHECK(OuterObject->StructWithInnerObjects.InnerObject->HasAllFlags(RF_DefaultSubObject));
+			CHECK(OuterObject->StructWithInnerObjects.InnerObject->HasAllFlags(OuterObject->GetMaskedFlags(RF_PropagateToSubObjects)));
+
+			CHECK_FALSE(OuterObject->StructWithInnerObjects.InnerObject->IsTemplate());
+			CHECK(OuterObject->StructWithInnerObjects.InnerObject->IsDefaultSubobject());
+			CHECK_FALSE(OuterObject->StructWithInnerObjects.InnerObject->IsTemplateForSubobjects());
+
+			USubobjectInstancingTestOuterObject* NewOuterObject = NewObject<USubobjectInstancingTestOuterObject>(GetTransientPackage(), InstancingTestClass, NAME_None, RF_NoFlags, Template);
+			CHECK(OuterObject->StructWithInnerObjects.InnerObject != NewOuterObject->StructWithInnerObjects.InnerObject);
+		}
+
+		SECTION("Struct member with default subobject from type initialized from CDO")
+		{
+			USubobjectInstancingTestOuterObject* OuterObject = NewObject<USubobjectInstancingTestOuterObject>(GetTransientPackage(), InstancingTestClass);
+
+			REQUIRE(OuterObject->StructWithInnerObjects.InnerObjectFromType != nullptr);
+			CHECK(OuterObject->StructWithInnerObjects.InnerObjectFromType->IsInOuter(OuterObject));
+
+			CHECK(OuterObject->StructWithInnerObjects.InnerObjectFromType->TestValue == 100);
+
+			CHECK(OuterObject->StructWithInnerObjects.InnerObjectFromType->HasAllFlags(RF_DefaultSubObject));
+			CHECK(OuterObject->StructWithInnerObjects.InnerObjectFromType->HasAllFlags(OuterObject->GetMaskedFlags(RF_PropagateToSubObjects)));
+
+			CHECK_FALSE(OuterObject->StructWithInnerObjects.InnerObjectFromType->IsTemplate());
+			CHECK(OuterObject->StructWithInnerObjects.InnerObjectFromType->IsDefaultSubobject());
+			CHECK_FALSE(OuterObject->StructWithInnerObjects.InnerObjectFromType->IsTemplateForSubobjects());
+
+			USubobjectInstancingTestOuterObject* NewOuterObject = NewObject<USubobjectInstancingTestOuterObject>(GetTransientPackage(), InstancingTestClass);
+			CHECK(OuterObject->StructWithInnerObjects.InnerObjectFromType != NewOuterObject->StructWithInnerObjects.InnerObjectFromType);
+		}
+
+		SECTION("Struct member with default subobject from type initialized from template")
+		{
+			USubobjectInstancingTestOuterObject* Template = NewObject<USubobjectInstancingTestOuterObject>(GetTransientPackage(), InstancingTestClass, NAME_None, RF_ArchetypeObject);
+			Template->StructWithInnerObjects.InnerObjectFromType->TestValue = 200;
+
+			USubobjectInstancingTestOuterObject* OuterObject = NewObject<USubobjectInstancingTestOuterObject>(GetTransientPackage(), InstancingTestClass, NAME_None, RF_NoFlags, Template);
+
+			REQUIRE(OuterObject->StructWithInnerObjects.InnerObjectFromType != nullptr);
+			CHECK(OuterObject->StructWithInnerObjects.InnerObjectFromType->IsInOuter(OuterObject));
+
+			CHECK(OuterObject->StructWithInnerObjects.InnerObjectFromType->TestValue == 200);
+
+			CHECK(OuterObject->StructWithInnerObjects.InnerObjectFromType->HasAllFlags(RF_DefaultSubObject));
+			CHECK(OuterObject->StructWithInnerObjects.InnerObjectFromType->HasAllFlags(OuterObject->GetMaskedFlags(RF_PropagateToSubObjects)));
+
+			CHECK_FALSE(OuterObject->StructWithInnerObjects.InnerObjectFromType->IsTemplate());
+			CHECK(OuterObject->StructWithInnerObjects.InnerObjectFromType->IsDefaultSubobject());
+			CHECK_FALSE(OuterObject->StructWithInnerObjects.InnerObjectFromType->IsTemplateForSubobjects());
+
+			USubobjectInstancingTestOuterObject* NewOuterObject = NewObject<USubobjectInstancingTestOuterObject>(GetTransientPackage(), InstancingTestClass, NAME_None, RF_NoFlags, Template);
+			CHECK(OuterObject->StructWithInnerObjects.InnerObjectFromType != NewOuterObject->StructWithInnerObjects.InnerObjectFromType);
+		}
+
+		SECTION("Struct member with default subobject array initialized from CDO")
+		{
+			USubobjectInstancingTestOuterObject* OuterObject = NewObject<USubobjectInstancingTestOuterObject>(GetTransientPackage(), InstancingTestClass);
+
+			CHECK_FALSE(OuterObject->StructWithInnerObjects.InnerObjectArray.IsEmpty());
+			for (USubobjectInstancingTestObject* InnerObject : OuterObject->StructWithInnerObjects.InnerObjectArray)
+			{
+				REQUIRE(InnerObject != nullptr);
+				CHECK(InnerObject->IsInOuter(OuterObject));
+
+				CHECK(InnerObject->TestValue == 100);
+
+				CHECK(InnerObject->HasAllFlags(RF_DefaultSubObject));
+				CHECK(InnerObject->HasAllFlags(OuterObject->GetMaskedFlags(RF_PropagateToSubObjects)));
+
+				CHECK_FALSE(InnerObject->IsTemplate());
+				CHECK(InnerObject->IsDefaultSubobject());
+				CHECK_FALSE(InnerObject->IsTemplateForSubobjects());
+			}
+		}
+
+		SECTION("Struct member with default subobject array initialized from template")
+		{
+			USubobjectInstancingTestOuterObject* Template = NewObject<USubobjectInstancingTestOuterObject>(GetTransientPackage(), InstancingTestClass, NAME_None, RF_ArchetypeObject);
+			for (USubobjectInstancingTestObject* InnerObject : Template->StructWithInnerObjects.InnerObjectArray)
+			{
+				InnerObject->TestValue = 200;
+			}
+
+			USubobjectInstancingTestOuterObject* OuterObject = NewObject<USubobjectInstancingTestOuterObject>(GetTransientPackage(), InstancingTestClass, NAME_None, RF_NoFlags, Template);
+
+			CHECK_FALSE(OuterObject->StructWithInnerObjects.InnerObjectArray.IsEmpty());
+			for (USubobjectInstancingTestObject* InnerObject : OuterObject->StructWithInnerObjects.InnerObjectArray)
+			{
+				REQUIRE(InnerObject != nullptr);
+				CHECK(InnerObject->IsInOuter(OuterObject));
+
+				CHECK(InnerObject->TestValue == 200);
+
+				CHECK(InnerObject->HasAllFlags(RF_DefaultSubObject));
+				CHECK(InnerObject->HasAllFlags(OuterObject->GetMaskedFlags(RF_PropagateToSubObjects)));
+
+				CHECK_FALSE(InnerObject->IsTemplate());
+				CHECK(InnerObject->IsDefaultSubobject());
+				CHECK_FALSE(InnerObject->IsTemplateForSubobjects());
+			}
+		}
+	}
+
+	void RunSubobjectInstancingTests_Optional(UClass* InstancingTestClass)
+	{
+		SECTION("Optional default subobject initialized from CDO")
+		{
+			USubobjectInstancingTestOuterObject* OuterObject = NewObject<USubobjectInstancingTestOuterObject>(GetTransientPackage(), InstancingTestClass);
+
+			CHECK(OuterObject->OptionalInnerObject.IsSet());
+
+			USubobjectInstancingTestObject* InnerObject = OuterObject->OptionalInnerObject.GetValue();
+			REQUIRE(InnerObject != nullptr);
+			CHECK(InnerObject->IsInOuter(OuterObject));
+
+			CHECK(InnerObject->TestValue == 100);
+
+			CHECK(InnerObject->HasAllFlags(RF_DefaultSubObject));
+			CHECK(InnerObject->HasAllFlags(OuterObject->GetMaskedFlags(RF_PropagateToSubObjects)));
+
+			CHECK_FALSE(InnerObject->IsTemplate());
+			CHECK(InnerObject->IsDefaultSubobject());
+			CHECK_FALSE(InnerObject->IsTemplateForSubobjects());
+
+			USubobjectInstancingTestOuterObject* NewOuterObject = NewObject<USubobjectInstancingTestOuterObject>(GetTransientPackage(), InstancingTestClass);
+			CHECK(InnerObject != NewOuterObject->OptionalInnerObject.GetValue());
+		}
+
+		SECTION("Optional default subobject initialized from template")
+		{
+			USubobjectInstancingTestOuterObject* Template = NewObject<USubobjectInstancingTestOuterObject>(GetTransientPackage(), InstancingTestClass, NAME_None, RF_ArchetypeObject);
+			Template->OptionalInnerObject.GetValue()->TestValue = 200;
+
+			USubobjectInstancingTestOuterObject* OuterObject = NewObject<USubobjectInstancingTestOuterObject>(GetTransientPackage(), InstancingTestClass, NAME_None, RF_NoFlags, Template);
+
+			CHECK(OuterObject->OptionalInnerObject.IsSet());
+
+			USubobjectInstancingTestObject* InnerObject = OuterObject->OptionalInnerObject.GetValue();
+			REQUIRE(InnerObject != nullptr);
+			CHECK(InnerObject->IsInOuter(OuterObject));
+
+			CHECK(InnerObject->TestValue == 200);
+
+			CHECK(InnerObject->HasAllFlags(RF_DefaultSubObject));
+			CHECK(InnerObject->HasAllFlags(OuterObject->GetMaskedFlags(RF_PropagateToSubObjects)));
+
+			CHECK_FALSE(InnerObject->IsTemplate());
+			CHECK(InnerObject->IsDefaultSubobject());
+			CHECK_FALSE(InnerObject->IsTemplateForSubobjects());
+
+			USubobjectInstancingTestOuterObject* NewOuterObject = NewObject<USubobjectInstancingTestOuterObject>(GetTransientPackage(), InstancingTestClass);
+			CHECK(InnerObject != NewOuterObject->OptionalInnerObject.GetValue());
+		}
+
+		SECTION("Optional default subobject from type initialized from CDO")
+		{
+			USubobjectInstancingTestOuterObject* OuterObject = NewObject<USubobjectInstancingTestOuterObject>(GetTransientPackage(), InstancingTestClass);
+
+			CHECK(OuterObject->OptionalInnerObjectFromType.IsSet());
+
+			USubobjectInstancingTestObject* InnerObject = OuterObject->OptionalInnerObjectFromType.GetValue();
+			REQUIRE(InnerObject != nullptr);
+			CHECK(InnerObject->IsInOuter(OuterObject));
+
+			CHECK(InnerObject->TestValue == 100);
+
+			CHECK(InnerObject->HasAllFlags(RF_DefaultSubObject));
+			CHECK(InnerObject->HasAllFlags(OuterObject->GetMaskedFlags(RF_PropagateToSubObjects)));
+
+			CHECK_FALSE(InnerObject->IsTemplate());
+			CHECK(InnerObject->IsDefaultSubobject());
+			CHECK_FALSE(InnerObject->IsTemplateForSubobjects());
+
+			USubobjectInstancingTestOuterObject* NewOuterObject = NewObject<USubobjectInstancingTestOuterObject>(GetTransientPackage(), InstancingTestClass);
+			CHECK(InnerObject != NewOuterObject->OptionalInnerObjectFromType.GetValue());
+		}
+
+		SECTION("Optional default subobject from type initialized from template")
+		{
+			USubobjectInstancingTestOuterObject* Template = NewObject<USubobjectInstancingTestOuterObject>(GetTransientPackage(), InstancingTestClass, NAME_None, RF_ArchetypeObject);
+			Template->OptionalInnerObjectFromType.GetValue()->TestValue = 200;
+
+			USubobjectInstancingTestOuterObject* OuterObject = NewObject<USubobjectInstancingTestOuterObject>(GetTransientPackage(), InstancingTestClass, NAME_None, RF_NoFlags, Template);
+
+			CHECK(OuterObject->OptionalInnerObjectFromType.IsSet());
+
+			USubobjectInstancingTestObject* InnerObject = OuterObject->OptionalInnerObjectFromType.GetValue();
+			REQUIRE(InnerObject != nullptr);
+			CHECK(InnerObject->IsInOuter(OuterObject));
+
+			CHECK(InnerObject->TestValue == 200);
+
+			CHECK(InnerObject->HasAllFlags(RF_DefaultSubObject));
+			CHECK(InnerObject->HasAllFlags(OuterObject->GetMaskedFlags(RF_PropagateToSubObjects)));
+
+			CHECK_FALSE(InnerObject->IsTemplate());
+			CHECK(InnerObject->IsDefaultSubobject());
+			CHECK_FALSE(InnerObject->IsTemplateForSubobjects());
+
+			USubobjectInstancingTestOuterObject* NewOuterObject = NewObject<USubobjectInstancingTestOuterObject>(GetTransientPackage(), InstancingTestClass);
+			CHECK(InnerObject != NewOuterObject->OptionalInnerObjectFromType.GetValue());
+		}
+
+		SECTION("Optional default subobject array initialized from CDO")
+		{
+			USubobjectInstancingTestOuterObject* OuterObject = NewObject<USubobjectInstancingTestOuterObject>(GetTransientPackage(), InstancingTestClass);
+
+			CHECK(OuterObject->OptionalInnerObjectArray.IsSet());
+
+			CHECK_FALSE(OuterObject->OptionalInnerObjectArray->IsEmpty());
+			for (USubobjectInstancingTestObject* InnerObject : OuterObject->OptionalInnerObjectArray.GetValue())
+			{
+				REQUIRE(InnerObject != nullptr);
+				CHECK(InnerObject->IsInOuter(OuterObject));
+
+				CHECK(InnerObject->TestValue == 100);
+
+				CHECK(InnerObject->HasAllFlags(RF_DefaultSubObject));
+				CHECK(InnerObject->HasAllFlags(OuterObject->GetMaskedFlags(RF_PropagateToSubObjects)));
+
+				CHECK_FALSE(InnerObject->IsTemplate());
+				CHECK(InnerObject->IsDefaultSubobject());
+				CHECK_FALSE(InnerObject->IsTemplateForSubobjects());
+			}
+		}
+
+		SECTION("Optional default subobject array initialized from template")
+		{
+			USubobjectInstancingTestOuterObject* Template = NewObject<USubobjectInstancingTestOuterObject>(GetTransientPackage(), InstancingTestClass, NAME_None, RF_ArchetypeObject);
+			for (USubobjectInstancingTestObject* InnerObject : Template->OptionalInnerObjectArray.GetValue())
+			{
+				InnerObject->TestValue = 200;
+			}
+
+			USubobjectInstancingTestOuterObject* OuterObject = NewObject<USubobjectInstancingTestOuterObject>(GetTransientPackage(), InstancingTestClass, NAME_None, RF_NoFlags, Template);
+
+			CHECK(OuterObject->OptionalInnerObjectArray.IsSet());
+
+			CHECK_FALSE(OuterObject->OptionalInnerObjectArray->IsEmpty());
+			for (USubobjectInstancingTestObject* InnerObject : OuterObject->OptionalInnerObjectArray.GetValue())
+			{
+				REQUIRE(InnerObject != nullptr);
+				CHECK(InnerObject->IsInOuter(OuterObject));
+
+				CHECK(InnerObject->TestValue == 200);
+
+				CHECK(InnerObject->HasAllFlags(RF_DefaultSubObject));
+				CHECK(InnerObject->HasAllFlags(OuterObject->GetMaskedFlags(RF_PropagateToSubObjects)));
+
+				CHECK_FALSE(InnerObject->IsTemplate());
+				CHECK(InnerObject->IsDefaultSubobject());
+				CHECK_FALSE(InnerObject->IsTemplateForSubobjects());
+			}
+		}
+	}
+
+	void RunSubobjectInstancingTests_Nested(UClass* InstancingTestClass)
+	{
+		SECTION("Nested default subobject initialized from CDO")
+		{
+			USubobjectInstancingTestOuterObject* OuterObject = NewObject<USubobjectInstancingTestOuterObject>(GetTransientPackage(), InstancingTestClass);
+
+			REQUIRE(OuterObject->InnerObjectWithDirectlyNestedObject->InnerObject != nullptr);
+			CHECK(OuterObject->InnerObjectWithDirectlyNestedObject->InnerObject->IsInOuter(OuterObject->InnerObjectWithDirectlyNestedObject));
+
+			CHECK(OuterObject->InnerObjectWithDirectlyNestedObject->InnerObject->TestValue == 100);
+
+			CHECK(OuterObject->InnerObjectWithDirectlyNestedObject->InnerObject->HasAllFlags(RF_DefaultSubObject));
+			CHECK(OuterObject->InnerObjectWithDirectlyNestedObject->InnerObject->HasAllFlags(OuterObject->InnerObjectWithDirectlyNestedObject->GetMaskedFlags(RF_PropagateToSubObjects)));
+
+			CHECK_FALSE(OuterObject->InnerObjectWithDirectlyNestedObject->InnerObject->IsTemplate());
+			CHECK(OuterObject->InnerObjectWithDirectlyNestedObject->InnerObject->IsDefaultSubobject());
+			CHECK_FALSE(OuterObject->InnerObjectWithDirectlyNestedObject->InnerObject->IsTemplateForSubobjects());
+
+			CHECK(OuterObject->InnerObjectWithDirectlyNestedObject->SelfRef == OuterObject->InnerObjectWithDirectlyNestedObject);
+			if (InstancingTestClass->IsNative())
+			{
+				CHECK(OuterObject->InnerObjectWithDirectlyNestedObject->OwnerObject == OuterObject);
+			}
+			else
+			{
+				// @todo - Not currently supported for native instanced properties on non-native outer types, but check the expected result to ensure backwards-compatibility.
+				//CHECK(OuterObject->InnerObjectWithDirectlyNestedObject->OwnerObject == OuterObject);
+				CHECK(OuterObject->InnerObjectWithDirectlyNestedObject->OwnerObject == InstancingTestClass->GetDefaultObject());
+			}
+
+			USubobjectInstancingTestOuterObject* NewOuterObject = NewObject<USubobjectInstancingTestOuterObject>(GetTransientPackage(), InstancingTestClass);
+			CHECK(OuterObject->InnerObjectWithDirectlyNestedObject->InnerObject != NewOuterObject->InnerObjectWithDirectlyNestedObject->InnerObject);
+		}
+
+		SECTION("Nested default subobject initialized from template")
+		{
+			USubobjectInstancingTestOuterObject* Template = NewObject<USubobjectInstancingTestOuterObject>(GetTransientPackage(), InstancingTestClass, NAME_None, RF_ArchetypeObject);
+			Template->InnerObjectWithDirectlyNestedObject->InnerObject->TestValue = 200;
+
+			USubobjectInstancingTestOuterObject* OuterObject = NewObject<USubobjectInstancingTestOuterObject>(GetTransientPackage(), InstancingTestClass, NAME_None, RF_NoFlags, Template);
+
+			REQUIRE(OuterObject->InnerObjectWithDirectlyNestedObject->InnerObject != nullptr);
+			CHECK(OuterObject->InnerObjectWithDirectlyNestedObject->InnerObject->IsInOuter(OuterObject->InnerObjectWithDirectlyNestedObject));
+
+			CHECK(OuterObject->InnerObjectWithDirectlyNestedObject->InnerObject->TestValue == 200);
+
+			CHECK(OuterObject->InnerObjectWithDirectlyNestedObject->InnerObject->HasAllFlags(RF_DefaultSubObject));
+			CHECK(OuterObject->InnerObjectWithDirectlyNestedObject->InnerObject->HasAllFlags(OuterObject->InnerObjectWithDirectlyNestedObject->GetMaskedFlags(RF_PropagateToSubObjects)));
+
+			CHECK_FALSE(OuterObject->InnerObjectWithDirectlyNestedObject->InnerObject->IsTemplate());
+			CHECK(OuterObject->InnerObjectWithDirectlyNestedObject->InnerObject->IsDefaultSubobject());
+			CHECK_FALSE(OuterObject->InnerObjectWithDirectlyNestedObject->InnerObject->IsTemplateForSubobjects());
+
+			CHECK(OuterObject->InnerObjectWithDirectlyNestedObject->SelfRef == OuterObject->InnerObjectWithDirectlyNestedObject);
+
+			// @todo - Not currently supported for native instanced properties on non-native outer types, but check the expected result to ensure backwards-compatibility.
+			// Note that native and non-native outer types currently result in different outputs.
+			//CHECK(OuterObject->InnerObjectWithDirectlyNestedObject->OwnerObject == OuterObject);
+			CHECK(OuterObject->InnerObjectWithDirectlyNestedObject->OwnerObject == (InstancingTestClass->IsNative() ? Template : InstancingTestClass->GetDefaultObject()));
+
+			USubobjectInstancingTestOuterObject* NewOuterObject = NewObject<USubobjectInstancingTestOuterObject>(GetTransientPackage(), InstancingTestClass, NAME_None, RF_NoFlags, Template);
+			CHECK(OuterObject->InnerObjectWithDirectlyNestedObject->InnerObject != NewOuterObject->InnerObjectWithDirectlyNestedObject->InnerObject);
+		}
+
+		SECTION("Deeply-nested default subobject initialized from CDO")
+		{
+			USubobjectInstancingTestOuterObject* OuterObject = NewObject<USubobjectInstancingTestOuterObject>(GetTransientPackage(), InstancingTestClass);
+
+			REQUIRE(OuterObject->InnerObjectWithIndirectlyNestedObject->InnerObject->InnerObject != nullptr);
+			CHECK(OuterObject->InnerObjectWithIndirectlyNestedObject->InnerObject->InnerObject->IsInOuter(OuterObject->InnerObjectWithIndirectlyNestedObject->InnerObject));
+
+			CHECK(OuterObject->InnerObjectWithIndirectlyNestedObject->InnerObject->InnerObject->TestValue == 100);
+
+			CHECK(OuterObject->InnerObjectWithIndirectlyNestedObject->InnerObject->InnerObject->HasAllFlags(RF_DefaultSubObject));
+			CHECK(OuterObject->InnerObjectWithIndirectlyNestedObject->InnerObject->InnerObject->HasAllFlags(OuterObject->InnerObjectWithIndirectlyNestedObject->InnerObject->GetMaskedFlags(RF_PropagateToSubObjects)));
+
+			CHECK_FALSE(OuterObject->InnerObjectWithIndirectlyNestedObject->InnerObject->InnerObject->IsTemplate());
+			CHECK(OuterObject->InnerObjectWithIndirectlyNestedObject->InnerObject->InnerObject->IsDefaultSubobject());
+			CHECK_FALSE(OuterObject->InnerObjectWithIndirectlyNestedObject->InnerObject->InnerObject->IsTemplateForSubobjects());
+
+			if (InstancingTestClass->IsNative())
+			{
+				CHECK(OuterObject->InnerObjectWithIndirectlyNestedObject->InnerObject->SelfRef == OuterObject->InnerObjectWithIndirectlyNestedObject->InnerObject);
+				CHECK(OuterObject->InnerObjectWithIndirectlyNestedObject->InnerObject->OwnerObject == OuterObject->InnerObjectWithIndirectlyNestedObject);
+			}
+			else
+			{
+				// @todo - Not currently supported for nested default subobjects instanced for non-native outer types, but check the expected result to ensure backwards-compatibility.
+				//CHECK(OuterObject->InnerObjectWithIndirectlyNestedObject->InnerObject->SelfRef == OuterObject->InnerObjectWithIndirectlyNestedObject->InnerObject);
+				//CHECK(OuterObject->InnerObjectWithIndirectlyNestedObject->InnerObject->OwnerObject == OuterObject->InnerObjectWithIndirectlyNestedObject);
+				USubobjectInstancingTestOuterObject* CDO = CastChecked<USubobjectInstancingTestOuterObject>(InstancingTestClass->GetDefaultObject());
+				CHECK(OuterObject->InnerObjectWithIndirectlyNestedObject->InnerObject->SelfRef == CDO->InnerObjectWithIndirectlyNestedObject->InnerObject->SelfRef);
+				CHECK(OuterObject->InnerObjectWithIndirectlyNestedObject->InnerObject->OwnerObject == CDO->InnerObjectWithIndirectlyNestedObject);
+			}
+
+			USubobjectInstancingTestOuterObject* NewOuterObject = NewObject<USubobjectInstancingTestOuterObject>(GetTransientPackage(), InstancingTestClass);
+			CHECK(OuterObject->InnerObjectWithIndirectlyNestedObject->InnerObject->InnerObject != NewOuterObject->InnerObjectWithIndirectlyNestedObject->InnerObject->InnerObject);
+		}
+
+		SECTION("Deeply-nested default subobject initialized from template")
+		{
+			USubobjectInstancingTestOuterObject* Template = NewObject<USubobjectInstancingTestOuterObject>(GetTransientPackage(), InstancingTestClass, NAME_None, RF_ArchetypeObject);
+			Template->InnerObjectWithIndirectlyNestedObject->InnerObject->InnerObject->TestValue = 200;
+
+			USubobjectInstancingTestOuterObject* OuterObject = NewObject<USubobjectInstancingTestOuterObject>(GetTransientPackage(), InstancingTestClass, NAME_None, RF_NoFlags, Template);
+
+			REQUIRE(OuterObject->InnerObjectWithIndirectlyNestedObject->InnerObject->InnerObject != nullptr);
+			CHECK(OuterObject->InnerObjectWithIndirectlyNestedObject->InnerObject->InnerObject->IsInOuter(OuterObject->InnerObjectWithIndirectlyNestedObject->InnerObject));
+
+			CHECK(OuterObject->InnerObjectWithIndirectlyNestedObject->InnerObject->InnerObject->TestValue == 200);
+
+			CHECK(OuterObject->InnerObjectWithIndirectlyNestedObject->InnerObject->InnerObject->HasAllFlags(RF_DefaultSubObject));
+			CHECK(OuterObject->InnerObjectWithIndirectlyNestedObject->InnerObject->InnerObject->HasAllFlags(OuterObject->InnerObjectWithIndirectlyNestedObject->InnerObject->GetMaskedFlags(RF_PropagateToSubObjects)));
+
+			CHECK_FALSE(OuterObject->InnerObjectWithIndirectlyNestedObject->InnerObject->InnerObject->IsTemplate());
+			CHECK(OuterObject->InnerObjectWithIndirectlyNestedObject->InnerObject->InnerObject->IsDefaultSubobject());
+			CHECK_FALSE(OuterObject->InnerObjectWithIndirectlyNestedObject->InnerObject->InnerObject->IsTemplateForSubobjects());
+
+			if (InstancingTestClass->IsNative())
+			{
+				CHECK(OuterObject->InnerObjectWithIndirectlyNestedObject->InnerObject->SelfRef == OuterObject->InnerObjectWithIndirectlyNestedObject->InnerObject);
+
+				// @todo - Not currently supported for nested default subobjects instanced for native outer types when using the 'Instanced' flag, but check the expected result
+				// to ensure backwards-compatibility.
+				//CHECK(OuterObject->InnerObjectWithIndirectlyNestedObject->InnerObject->OwnerObject == OuterObject->InnerObjectWithIndirectlyNestedObject);
+				CHECK(OuterObject->InnerObjectWithIndirectlyNestedObject->InnerObject->OwnerObject == Template->InnerObjectWithIndirectlyNestedObject);
+			}
+			else
+			{
+				// @todo - These cases are not currently supported for nested default subobjects constructed for non-native outer types when using the 'Instanced' flag, but check
+				// the expected result to ensure backwards-compatibility.
+				//CHECK(OuterObject->InnerObjectWithIndirectlyNestedObject->InnerObject->SelfRef == OuterObject->InnerObjectWithIndirectlyNestedObject->InnerObject);
+				//CHECK(OuterObject->InnerObjectWithIndirectlyNestedObject->InnerObject->OwnerObject == OuterObject->InnerObjectWithIndirectlyNestedObject);
+				USubobjectInstancingTestOuterObject* CDO = CastChecked<USubobjectInstancingTestOuterObject>(InstancingTestClass->GetDefaultObject());
+				CHECK(OuterObject->InnerObjectWithIndirectlyNestedObject->InnerObject->SelfRef == Template->InnerObjectWithIndirectlyNestedObject->InnerObject->SelfRef);
+				CHECK(OuterObject->InnerObjectWithIndirectlyNestedObject->InnerObject->OwnerObject == CDO->InnerObjectWithIndirectlyNestedObject);
+			}
+
+			USubobjectInstancingTestOuterObject* NewOuterObject = NewObject<USubobjectInstancingTestOuterObject>(GetTransientPackage(), InstancingTestClass, NAME_None, RF_NoFlags, Template);
+			CHECK(OuterObject->InnerObjectWithIndirectlyNestedObject->InnerObject->InnerObject != NewOuterObject->InnerObjectWithIndirectlyNestedObject->InnerObject->InnerObject);
+		}
+
+		SECTION("Default subobject with nested dynamically-instanced subobject")
+		{
+			if (InstancingTestClass->IsNative())
+			{
+				// Skip: This test only applies to non-native outer types.
+			}
+			else
+			{
+				// Create a dynamically-instanced test class to use for the outer default subobject that we'll simulate being instanced at edit time below. It will contain an inner
+				// object property that needs to be instanced dynamically on new instances of the owning object (which itself may not use dynamic instancing). This is a valid use case.
+				UClass* TestClass = FSubobjectInstancingTestUtils::CreateDynamicallyInstancedTestClass(USubobjectInstancingTestDirectlyNestedObject::StaticClass());
+
+				// Create a template for the outer object that may or may not contain dynamic reference properties.
+				USubobjectInstancingTestOuterObject* Template = NewObject<USubobjectInstancingTestOuterObject>(GetTransientPackage(), InstancingTestClass, NAME_None, RF_ArchetypeObject);
+
+				// Simulate instantiating a subobject archetype at edit time with a type that may contain a dynamic reference property. This should use dynamic instancing to instance it when its outer object is constructed.
+				FObjectProperty* ObjectProperty = CastFieldChecked<FObjectProperty>(InstancingTestClass->FindPropertyByName(NonNativeEditTimeInnerObjectPropertyName));
+				USubobjectInstancingTestDirectlyNestedObject* NonNativeEditTimeInnerObjectDefaultValue = NewObject<USubobjectInstancingTestDirectlyNestedObject>(Template, TestClass, TEXT("NonNativeEditTimeInnerObject"), Template->GetMaskedFlags(RF_PropagateToSubObjects));
+				ObjectProperty->SetObjectPropertyValue_InContainer(Template, NonNativeEditTimeInnerObjectDefaultValue);
+
+				// Modify non-native properties that are expected to be mirrored to the nested subobject instance.
+				NonNativeEditTimeInnerObjectDefaultValue->TestValue = 200;
+				CastFieldChecked<FObjectProperty>(TestClass->FindPropertyByName(Private::NonNativeOuterObjectPropertyName))->SetObjectPropertyValue_InContainer(NonNativeEditTimeInnerObjectDefaultValue, Template);
+				CastFieldChecked<FObjectProperty>(TestClass->FindPropertyByName(Private::NonNativeSelfReferencePropertyName))->SetObjectPropertyValue_InContainer(NonNativeEditTimeInnerObjectDefaultValue, NonNativeEditTimeInnerObjectDefaultValue);
+
+				// Construct a new instance of the outer object type using the template for initialization.
+				USubobjectInstancingTestOuterObject* OuterObject = NewObject<USubobjectInstancingTestOuterObject>(GetTransientPackage(), InstancingTestClass, NAME_None, RF_NoFlags, Template);
+				USubobjectInstancingTestDirectlyNestedObject* NonNativeEditTimeInnerObject = CastChecked<USubobjectInstancingTestDirectlyNestedObject>(ObjectProperty->GetObjectPropertyValue_InContainer(OuterObject));
+
+				CHECK(NonNativeEditTimeInnerObject->TestValue == 200);
+
+				CHECK(CastFieldChecked<FObjectProperty>(TestClass->FindPropertyByName(Private::NonNativeInnerObjectPropertyName))->GetObjectPropertyValue_InContainer(NonNativeEditTimeInnerObject) != nullptr);
+				CHECK(CastFieldChecked<FObjectProperty>(TestClass->FindPropertyByName(Private::NonNativeInnerObjectPropertyName))->GetObjectPropertyValue_InContainer(NonNativeEditTimeInnerObject)->IsInOuter(OuterObject));
+
+				CHECK(CastFieldChecked<FObjectProperty>(TestClass->FindPropertyByName(Private::NonNativeOuterObjectPropertyName))->GetObjectPropertyValue_InContainer(NonNativeEditTimeInnerObject) == OuterObject);
+				CHECK(CastFieldChecked<FObjectProperty>(TestClass->FindPropertyByName(Private::NonNativeSelfReferencePropertyName))->GetObjectPropertyValue_InContainer(NonNativeEditTimeInnerObject) == NonNativeEditTimeInnerObject);
+
+				CHECK_FALSE(NonNativeEditTimeInnerObject->HasAllFlags(RF_DefaultSubObject));
+				CHECK(NonNativeEditTimeInnerObject->HasAllFlags(OuterObject->GetMaskedFlags(RF_PropagateToSubObjects)));
+
+				CHECK_FALSE(NonNativeEditTimeInnerObject->IsTemplate());
+				CHECK_FALSE(NonNativeEditTimeInnerObject->IsDefaultSubobject());
+				CHECK_FALSE(NonNativeEditTimeInnerObject->IsTemplateForSubobjects());
+
+				USubobjectInstancingTestOuterObject* NewOuterObject = NewObject<USubobjectInstancingTestOuterObject>(GetTransientPackage(), InstancingTestClass, NAME_None, RF_NoFlags, Template);
+				CHECK(NonNativeEditTimeInnerObject != ObjectProperty->GetObjectPropertyValue_InContainer(NewOuterObject));
+			}
+		}
+	}
+
+	void RunSubobjectInstancingTests_SelfReferencing(UClass* InstancingTestClass)
+	{
+		SECTION("Reference to native self property initialized from CDO")
+		{
+			USubobjectInstancingTestOuterObject* OuterObject = NewObject<USubobjectInstancingTestOuterObject>(GetTransientPackage(), InstancingTestClass);
+
+			CHECK(OuterObject->SelfRef == OuterObject);
+		}
+
+		SECTION("Reference to native self property initialized from template")
+		{
+			USubobjectInstancingTestOuterObject* Template = NewObject<USubobjectInstancingTestOuterObject>(GetTransientPackage(), InstancingTestClass, NAME_None, RF_ArchetypeObject);
+			USubobjectInstancingTestOuterObject* OuterObject = NewObject<USubobjectInstancingTestOuterObject>(GetTransientPackage(), InstancingTestClass, NAME_None, RF_NoFlags, Template);
+
+			// @todo - Self-referencing is not currently supported for native instanced properties when instanced from a template, but check the expected result to ensure backwards-compatibility.
+			//CHECK(OuterObject->SelfRef == OuterObject);
+			CHECK(OuterObject->SelfRef == Template);
+		}
+
+		SECTION("Reference to non-native self property initialized from CDO")
+		{
+			if (InstancingTestClass->IsNative())
+			{
+				// Skip: This test only applies to non-native types.
+			}
+			else
+			{
+				USubobjectInstancingTestOuterObject* OuterObject = NewObject<USubobjectInstancingTestOuterObject>(GetTransientPackage(), InstancingTestClass);
+				FObjectProperty* ObjectProperty = CastFieldChecked<FObjectProperty>(InstancingTestClass->FindPropertyByName(NonNativeSelfReferencePropertyName));
+				UObject* NonNativeSelfRef = ObjectProperty->GetObjectPropertyValue_InContainer(OuterObject);
+
+				// This is initialized on the CDO, but in the non-native case, the property is linked into the PCL chain. So unlike the native case, InitProperties() will copy it to the new
+				// instance, and FObjectInstancingGraph::GetInstancedSubobject() is expected to resolve the source value during subobject instancing. This currently has limited support.
+				// 
+				// @todo (UE-219797) - When NOT using dynamic instancing, GetInstancedSubobject() currently requires the 'CPF_AllowSelfReference' flag to work. This flag is considered to
+				// be deprecated in favor of dynamic instancing and will eventually be removed. It can only be set in code on non-native properties at this time (there is no UHT support).
+				if (InstancingTestClass->ShouldUseDynamicSubobjectInstancing()
+					|| ObjectProperty->HasAnyPropertyFlags(CPF_AllowSelfReference))
+				{
+					CHECK(NonNativeSelfRef == OuterObject);
+				}
+				else
+				{
+					// Unsupported cases (currently types w/o dynamic reference support)
+					CHECK(NonNativeSelfRef == InstancingTestClass->GetDefaultObject());
+				}
+			}
+		}
+
+		SECTION("Reference to non-native self property initialized from template")
+		{
+			if (InstancingTestClass->IsNative())
+			{
+				// Skip: This test only applies to non-native types.
+			}
+			else
+			{
+				USubobjectInstancingTestOuterObject* Template = NewObject<USubobjectInstancingTestOuterObject>(GetTransientPackage(), InstancingTestClass, NAME_None, RF_ArchetypeObject);
+				USubobjectInstancingTestOuterObject* OuterObject = NewObject<USubobjectInstancingTestOuterObject>(GetTransientPackage(), InstancingTestClass, NAME_None, RF_NoFlags, Template);
+				FObjectProperty* ObjectProperty = CastFieldChecked<FObjectProperty>(InstancingTestClass->FindPropertyByName(Private::NonNativeSelfReferencePropertyName));
+				UObject* NonNativeSelfRef = ObjectProperty->GetObjectPropertyValue_InContainer(OuterObject);
+
+				// @todo (UE-219797) - See note in previous test above for the CDO case. The 'CPF_AllowSelfReference' flag is considered to be deprecated and will eventually be removed.
+				if (InstancingTestClass->ShouldUseDynamicSubobjectInstancing()
+					|| ObjectProperty->HasAnyPropertyFlags(CPF_AllowSelfReference))
+				{
+					CHECK(NonNativeSelfRef == OuterObject);
+				}
+				else
+				{
+					// Unsupported cases (currently types w/o dynamic reference support)
+					CHECK(NonNativeSelfRef == Template);
+				}
+			}
+		}
+
+		SECTION("Reference to inner object initialized from CDO")
+		{
+			USubobjectInstancingTestOuterObject* OuterObject = NewObject<USubobjectInstancingTestOuterObject>(GetTransientPackage(), InstancingTestClass);
+
+			CHECK(OuterObject->InternalObject == OuterObject->InnerObject);
+		}
+
+		SECTION("Reference to inner object initialized from template")
+		{
+			USubobjectInstancingTestOuterObject* Template = NewObject<USubobjectInstancingTestOuterObject>(GetTransientPackage(), InstancingTestClass, NAME_None, RF_ArchetypeObject);
+			USubobjectInstancingTestOuterObject* OuterObject = NewObject<USubobjectInstancingTestOuterObject>(GetTransientPackage(), InstancingTestClass, NAME_None, RF_NoFlags, Template);
+
+			CHECK(OuterObject->InternalObject == OuterObject->InnerObject);
+		}
+	}
+
+	void RunSubobjectInstancingTests_Deferred(UClass* InstancingTestClass)
+	{
+		SECTION("Simulate native property instanced at edit time on non-native CDO")
+		{
+			if (InstancingTestClass->IsNative())
+			{
+				// Skip: This test only applies to non-native types.
+			}
+			else
+			{
+				USubobjectInstancingTestOuterObject* CDO = Cast<USubobjectInstancingTestOuterObject>(InstancingTestClass->GetDefaultObject());
+
+				// Default value before it's edited (as will be initialized by the native super class ctor).
+				CHECK(CDO->EditTimeInnerObject == nullptr);
+
+				// Simulate a deferred edit-time instancing event along with a value override on the non-native CDO.
+				CDO->EditTimeInnerObject = NewObject<USubobjectInstancingTestObject>(CDO, GET_MEMBER_NAME_CHECKED(USubobjectInstancingTestOuterObject, EditTimeInnerObject), CDO->GetMaskedFlags(RF_PropagateToSubObjects));
+				CDO->EditTimeInnerObject->TestValue = 200;
+
+				CHECK(CDO->EditTimeInnerObject->IsTemplate());
+				// Since we instanced the subobject on the CDO, the value is "captured" in the same manner as if it were instanced at construction time, which makes it a default subobject by definition.
+				// Given the current (legacy) implementation, this will be an expected result.
+				CHECK(CDO->EditTimeInnerObject->IsDefaultSubobject());
+				CHECK(CDO->EditTimeInnerObject->IsTemplateForSubobjects());
+
+				USubobjectInstancingTestOuterObject* OuterObject = NewObject<USubobjectInstancingTestOuterObject>(GetTransientPackage(), InstancingTestClass);
+
+				// Since the subobject archetype was not instanced at construction time, it's expected to remain set to NULL. When the CDO is used as the default object for initialization, as an optimization,
+				// the native object initializer will not include properties inherited from the native super class hierarchy, because it expects they have already been initialized by the native ctor. For this
+				// reason, users cannot override a native instanced reference member's default value. This test is validating that any such override will not actually be used by subobject template instancing.
+				CHECK(OuterObject->EditTimeInnerObject == nullptr);
+			}
+		}
+
+		SECTION("Simulate native property instanced at edit time on non-native template")
+		{
+			if (InstancingTestClass->IsNative())
+			{
+				// Skip: This test only applies to non-native types.
+			}
+			else
+			{
+				USubobjectInstancingTestOuterObject* Template = NewObject<USubobjectInstancingTestOuterObject>(GetTransientPackage(), InstancingTestClass, NAME_None, RF_ArchetypeObject);
+				CHECK(Template->EditTimeInnerObject == nullptr);
+
+				// Simulate a deferred edit-time instancing event along with a value override on the non-native template.
+				Template->EditTimeInnerObject = NewObject<USubobjectInstancingTestObject>(Template, GET_MEMBER_NAME_CHECKED(USubobjectInstancingTestOuterObject, EditTimeInnerObject), Template->GetMaskedFlags(RF_PropagateToSubObjects));
+				Template->EditTimeInnerObject->TestValue = 200;
+
+				CHECK(Template->EditTimeInnerObject->IsTemplate());
+				CHECK_FALSE(Template->EditTimeInnerObject->IsDefaultSubobject());
+				// Since this is not considered a default subobject nor was it captured by a CDO, it's expected that this call should fail w/ the default flags argument on the internal IsTemplate() call.
+				CHECK_FALSE(Template->EditTimeInnerObject->IsTemplateForSubobjects());
+				// However, we can configure the API to tell us whether the subobject is also an archetype that will be used for instancing at construction time if the outer template is used for initialization.
+				CHECK(Template->EditTimeInnerObject->IsTemplateForSubobjects(RF_ArchetypeObject));
+
+				USubobjectInstancingTestOuterObject* OuterObject = NewObject<USubobjectInstancingTestOuterObject>(GetTransientPackage(), InstancingTestClass, NAME_None, RF_NoFlags, Template);
+
+				CHECK(OuterObject->EditTimeInnerObject != nullptr);
+				CHECK(OuterObject->EditTimeInnerObject->IsInOuter(OuterObject));
+
+				CHECK(OuterObject->EditTimeInnerObject->TestValue == 200);
+
+				CHECK_FALSE(OuterObject->EditTimeInnerObject->HasAllFlags(RF_DefaultSubObject));
+				CHECK(OuterObject->EditTimeInnerObject->HasAllFlags(OuterObject->GetMaskedFlags(RF_PropagateToSubObjects)));
+
+				CHECK_FALSE(OuterObject->EditTimeInnerObject->IsTemplate());
+				CHECK_FALSE(OuterObject->EditTimeInnerObject->IsDefaultSubobject());
+				CHECK_FALSE(OuterObject->EditTimeInnerObject->IsTemplateForSubobjects());
+			}
+		}
+
+		SECTION("Simulate non-native property instanced at edit time on non-native CDO")
+		{
+			if (InstancingTestClass->IsNative())
+			{
+				// Skip: This test only applies to non-native types.
+			}
+			else
+			{
+				USubobjectInstancingTestOuterObject* CDO = Cast<USubobjectInstancingTestOuterObject>(InstancingTestClass->GetDefaultObject());
+
+				FObjectProperty* ObjectProperty = CastFieldChecked<FObjectProperty>(InstancingTestClass->FindPropertyByName(Private::NonNativeEditTimeInnerObjectPropertyName));
+				USubobjectInstancingTestObject* NonNativeEditTimeInnerObjectDefaultValue = Cast<USubobjectInstancingTestObject>(ObjectProperty->GetObjectPropertyValue_InContainer(CDO));
+
+				// Default value before it's edited (as will have been initialized by non-native construction).
+				CHECK(NonNativeEditTimeInnerObjectDefaultValue == nullptr);
+
+				// Simulate a deferred edit-time instancing event along with a value override on the non-native CDO.
+				NonNativeEditTimeInnerObjectDefaultValue = NewObject<USubobjectInstancingTestObject>(CDO, TEXT("NonNativeEditTimeInnerObject"), CDO->GetMaskedFlags(RF_PropagateToSubObjects));
+				NonNativeEditTimeInnerObjectDefaultValue->TestValue = 200;
+				ObjectProperty->SetObjectPropertyValue_InContainer(CDO, NonNativeEditTimeInnerObjectDefaultValue);
+
+				CHECK(NonNativeEditTimeInnerObjectDefaultValue->IsTemplate());
+				// Since we instanced the subobject on the CDO, the value is "captured" in the same manner as if it were instanced at construction time, which makes it a default subobject by definition.
+				// Given the current (legacy) implementation, this will be an expected result.
+				CHECK(NonNativeEditTimeInnerObjectDefaultValue->IsDefaultSubobject());
+				CHECK(NonNativeEditTimeInnerObjectDefaultValue->IsTemplateForSubobjects());
+
+				USubobjectInstancingTestOuterObject* OuterObject = NewObject<USubobjectInstancingTestOuterObject>(GetTransientPackage(), InstancingTestClass);
+				USubobjectInstancingTestObject* NonNativeEditTimeInnerObject = Cast<USubobjectInstancingTestObject>(ObjectProperty->GetObjectPropertyValue_InContainer(OuterObject));
+
+				REQUIRE(NonNativeEditTimeInnerObject != nullptr);
+				CHECK(NonNativeEditTimeInnerObject->IsInOuter(OuterObject));
+
+				CHECK(NonNativeEditTimeInnerObject->TestValue == 200);
+
+				CHECK_FALSE(NonNativeEditTimeInnerObject->HasAllFlags(RF_DefaultSubObject));
+				CHECK(NonNativeEditTimeInnerObject->HasAllFlags(OuterObject->GetMaskedFlags(RF_PropagateToSubObjects)));
+
+				CHECK_FALSE(NonNativeEditTimeInnerObject->IsTemplate());
+				CHECK(NonNativeEditTimeInnerObject->IsDefaultSubobject());
+				CHECK_FALSE(NonNativeEditTimeInnerObject->IsTemplateForSubobjects());
+
+				USubobjectInstancingTestOuterObject* NewOuterObject = NewObject<USubobjectInstancingTestOuterObject>(GetTransientPackage(), InstancingTestClass);
+				CHECK(NonNativeEditTimeInnerObject != ObjectProperty->GetObjectPropertyValue_InContainer(NewOuterObject));
+			}
+		}
+
+		SECTION("Simulate non-native property instanced at edit time on non-native template")
+		{
+			if (InstancingTestClass->IsNative())
+			{
+				// Skip: This test only applies to non-native iterations - the native class won't include the non-native property.
+			}
+			else
+			{
+				USubobjectInstancingTestOuterObject* Template = NewObject<USubobjectInstancingTestOuterObject>(GetTransientPackage(), InstancingTestClass, NAME_None, RF_ArchetypeObject);
+
+				FObjectProperty* ObjectProperty = CastFieldChecked<FObjectProperty>(InstancingTestClass->FindPropertyByName(Private::NonNativeEditTimeInnerObjectPropertyName));
+				USubobjectInstancingTestObject* NonNativeEditTimeInnerObjectDefaultValue = Cast<USubobjectInstancingTestObject>(ObjectProperty->GetObjectPropertyValue_InContainer(Template));
+
+				// Default value before it's edited (as will have been initialized by non-native construction).
+				CHECK(NonNativeEditTimeInnerObjectDefaultValue == nullptr);
+
+				// Simulate a deferred edit-time instancing event along with a value override on the non-native template.
+				NonNativeEditTimeInnerObjectDefaultValue = NewObject<USubobjectInstancingTestObject>(Template, TEXT("NonNativeEditTimeInnerObject"), Template->GetMaskedFlags(RF_PropagateToSubObjects));
+				NonNativeEditTimeInnerObjectDefaultValue->TestValue = 200;
+				ObjectProperty->SetObjectPropertyValue_InContainer(Template, NonNativeEditTimeInnerObjectDefaultValue);
+
+				CHECK(NonNativeEditTimeInnerObjectDefaultValue->IsTemplate());
+				// Since we instanced the subobject on the template and not the CDO, this is not considered to be a default subobject.
+				CHECK_FALSE(NonNativeEditTimeInnerObjectDefaultValue->IsDefaultSubobject());
+				// It is also not considered to be a template for a CDO or default subobject, but rather a standalone archetype subobject.
+				CHECK_FALSE(NonNativeEditTimeInnerObjectDefaultValue->IsTemplateForSubobjects());
+				CHECK(NonNativeEditTimeInnerObjectDefaultValue->IsTemplateForSubobjects(RF_ArchetypeObject));
+
+				USubobjectInstancingTestOuterObject* OuterObject = NewObject<USubobjectInstancingTestOuterObject>(GetTransientPackage(), InstancingTestClass, NAME_None, RF_NoFlags, Template);
+				USubobjectInstancingTestObject* NonNativeEditTimeInnerObject = Cast<USubobjectInstancingTestObject>(ObjectProperty->GetObjectPropertyValue_InContainer(OuterObject));
+
+				REQUIRE(NonNativeEditTimeInnerObject != nullptr);
+				CHECK(NonNativeEditTimeInnerObject->IsInOuter(OuterObject));
+
+				CHECK(NonNativeEditTimeInnerObject->TestValue == 200);
+
+				CHECK_FALSE(NonNativeEditTimeInnerObject->HasAllFlags(RF_DefaultSubObject));
+				CHECK(NonNativeEditTimeInnerObject->HasAllFlags(OuterObject->GetMaskedFlags(RF_PropagateToSubObjects)));
+
+				CHECK_FALSE(NonNativeEditTimeInnerObject->IsTemplate());
+				CHECK_FALSE(NonNativeEditTimeInnerObject->IsDefaultSubobject());
+				CHECK_FALSE(NonNativeEditTimeInnerObject->IsTemplateForSubobjects());
+
+				USubobjectInstancingTestOuterObject* NewOuterObject = NewObject<USubobjectInstancingTestOuterObject>(GetTransientPackage(), InstancingTestClass);
+				CHECK(NonNativeEditTimeInnerObject != ObjectProperty->GetObjectPropertyValue_InContainer(NewOuterObject));
+			}
+		}
+
+		SECTION("Native property initialized from template with deferred construction using NewObject()")
+		{
+			USubobjectInstancingTestOuterObject* Template = NewObject<USubobjectInstancingTestOuterObject>(GetTransientPackage(), InstancingTestClass, NAME_None, RF_ArchetypeObject);
+			Template->InnerObjectPostInit->TestValue = 200;
+
+			USubobjectInstancingTestOuterObject* OuterObject = NewObject<USubobjectInstancingTestOuterObject>(GetTransientPackage(), InstancingTestClass, NAME_None, RF_NoFlags, Template);
+
+			REQUIRE(OuterObject->InnerObjectPostInit != nullptr);
+			CHECK(OuterObject->InnerObjectPostInit->IsInOuter(OuterObject));
+
+			CHECK(OuterObject->InnerObjectPostInit->TestValue == 200);
+
+			CHECK_FALSE(OuterObject->InnerObjectPostInit->HasAllFlags(RF_DefaultSubObject));
+			CHECK(OuterObject->InnerObjectPostInit->HasAllFlags(OuterObject->GetMaskedFlags(RF_PropagateToSubObjects)));
+
+			CHECK_FALSE(OuterObject->InnerObjectPostInit->IsTemplate());
+			CHECK(OuterObject->InnerObjectPostInit->IsDefaultSubobject());
+			CHECK_FALSE(OuterObject->InnerObjectPostInit->IsTemplateForSubobjects());
+
+			USubobjectInstancingTestOuterObject* NewOuterObject = NewObject<USubobjectInstancingTestOuterObject>(GetTransientPackage(), InstancingTestClass, NAME_None, RF_NoFlags, Template);
+			CHECK(OuterObject->InnerObjectPostInit != NewOuterObject->InnerObjectPostInit);
+		}
+	}
+
+	void RunSubobjectInstancingTests_Other(UClass* InstancingTestClass)
+	{
+		SECTION("StaticDuplicateObject() instancing")
+		{
+			USubobjectInstancingTestOuterObject* SrcObject = NewObject<USubobjectInstancingTestOuterObject>(GetTransientPackage(), InstancingTestClass);
+
+			SrcObject->InnerObject->TestValue = 200;
+			SrcObject->SharedObject->TestValue = 200;
+			SrcObject->TransientInnerObject->TestValue = 200;
+			SrcObject->LocalOnlyInnerObject->TestValue = 200;
+			SrcObject->InnerObjectFromType->TestValue = 200;
+			SrcObject->InnerObjectUsingNew->TestValue = 200;
+			SrcObject->InnerObjectPostInit->TestValue = 200;
+			SrcObject->StructWithInnerObjects.InnerObject->TestValue = 200;
+			SrcObject->OptionalInnerObject.GetValue()->TestValue = 200;
+			for (USubobjectInstancingTestObject* InnerObject : SrcObject->InnerObjectArray)
+			{
+				InnerObject->TestValue = 200;
+			}
+
+			SrcObject->InnerObjectWithDirectlyNestedObject->TestValue = 200;
+			SrcObject->InnerObjectWithDirectlyNestedObject->InnerObject->TestValue = 200;
+
+			SrcObject->InnerObjectWithIndirectlyNestedObject->TestValue = 200;
+			SrcObject->InnerObjectWithIndirectlyNestedObject->InnerObject->TestValue = 200;
+			SrcObject->InnerObjectWithIndirectlyNestedObject->InnerObject->InnerObject->TestValue = 200;
+
+			// Note: Using PIE mode for duplication as low level test runners do not configure UPS at boot time and normal mode requires it.
+			FObjectDuplicationParameters Params = InitStaticDuplicateObjectParams(SrcObject, GetTransientPackage());
+			Params.PortFlags = PPF_DuplicateForPIE;
+
+			USubobjectInstancingTestOuterObject* DstObject = CastChecked<USubobjectInstancingTestOuterObject>(StaticDuplicateObjectEx(Params));
+
+			CHECK(DstObject->SelfRef == DstObject);
+			CHECK(DstObject->NullObject == nullptr);
+
+			CHECK(DstObject->InnerObject != nullptr);
+			CHECK(DstObject->InnerObject->IsInOuter(DstObject));
+			CHECK(SrcObject->InnerObject != DstObject->InnerObject);
+			CHECK(SrcObject->InnerObject->TestValue == DstObject->InnerObject->TestValue);
+
+			CHECK(DstObject->SharedObject != nullptr);
+			CHECK(DstObject->SharedObject->IsInOuter(DstObject));
+			CHECK(SrcObject->SharedObject != DstObject->SharedObject);
+			CHECK(SrcObject->SharedObject->TestValue == DstObject->SharedObject->TestValue);
+
+			CHECK(DstObject->ExternalObject != nullptr);
+			CHECK_FALSE(DstObject->ExternalObject->IsInOuter(DstObject));
+			CHECK(SrcObject->ExternalObject == DstObject->ExternalObject);
+
+			CHECK(DstObject->TransientInnerObject != nullptr);
+			CHECK(DstObject->TransientInnerObject->IsInOuter(DstObject));
+			CHECK(DstObject->TransientInnerObject != SrcObject->TransientInnerObject);
+			CHECK(DstObject->TransientInnerObject->TestValue == SrcObject->TransientInnerObject->TestValue);
+
+			CHECK(DstObject->LocalOnlyInnerObject != nullptr);
+			CHECK(DstObject->LocalOnlyInnerObject->IsInOuter(DstObject));
+			CHECK(DstObject->LocalOnlyInnerObject != SrcObject->LocalOnlyInnerObject);
+			CHECK(DstObject->LocalOnlyInnerObject->TestValue == SrcObject->LocalOnlyInnerObject->TestValue);
+
+			CHECK(DstObject->InnerObjectFromType != nullptr);
+			CHECK(DstObject->InnerObjectFromType->IsInOuter(DstObject));
+			CHECK(DstObject->InnerObjectFromType != SrcObject->LocalOnlyInnerObject);
+			CHECK(DstObject->InnerObjectFromType->TestValue == SrcObject->InnerObjectFromType->TestValue);
+
+			CHECK(DstObject->InnerObjectUsingNew != nullptr);
+			CHECK(DstObject->InnerObjectUsingNew->IsInOuter(DstObject));
+			CHECK(DstObject->InnerObjectUsingNew != SrcObject->InnerObjectUsingNew);
+			CHECK(DstObject->InnerObjectUsingNew->TestValue == SrcObject->InnerObjectUsingNew->TestValue);
+
+			CHECK(DstObject->InnerObjectPostInit != nullptr);
+			CHECK(DstObject->InnerObjectPostInit->IsInOuter(DstObject));
+			CHECK(DstObject->InnerObjectPostInit != SrcObject->InnerObjectPostInit);
+			CHECK(DstObject->InnerObjectPostInit->TestValue == SrcObject->InnerObjectPostInit->TestValue);
+
+			CHECK(DstObject->StructWithInnerObjects.InnerObject != nullptr);
+			CHECK(DstObject->StructWithInnerObjects.InnerObject->IsInOuter(DstObject));
+			CHECK(DstObject->StructWithInnerObjects.InnerObject != SrcObject->StructWithInnerObjects.InnerObject);
+			CHECK(DstObject->StructWithInnerObjects.InnerObject->TestValue == SrcObject->StructWithInnerObjects.InnerObject->TestValue);
+
+			CHECK(DstObject->OptionalInnerObject.IsSet());
+			CHECK(DstObject->OptionalInnerObject.GetValue() != nullptr);
+			CHECK(DstObject->OptionalInnerObject.GetValue()->IsInOuter(DstObject));
+			CHECK(DstObject->OptionalInnerObject.GetValue() != SrcObject->OptionalInnerObject.GetValue());
+			CHECK(DstObject->OptionalInnerObject.GetValue()->TestValue == SrcObject->OptionalInnerObject.GetValue()->TestValue);
+
+			CHECK_FALSE(DstObject->InnerObjectArray.IsEmpty());
+			CHECK(DstObject->InnerObjectArray.Num() == SrcObject->InnerObjectArray.Num());
+			for (int32 ArrayIdx = 0; ArrayIdx < DstObject->InnerObjectArray.Num(); ++ArrayIdx)
+			{
+				CHECK(DstObject->InnerObjectArray[ArrayIdx] != nullptr);
+				CHECK(DstObject->InnerObjectArray[ArrayIdx]->IsInOuter(DstObject));
+				CHECK(DstObject->InnerObjectArray[ArrayIdx] != SrcObject->InnerObjectArray[ArrayIdx]);
+				CHECK(DstObject->InnerObjectArray[ArrayIdx]->TestValue == SrcObject->InnerObjectArray[ArrayIdx]->TestValue);
+			}
+
+			CHECK(DstObject->InnerObjectWithDirectlyNestedObject != nullptr);
+			CHECK(DstObject->InnerObjectWithDirectlyNestedObject->IsInOuter(DstObject));
+			CHECK(DstObject->InnerObjectWithDirectlyNestedObject != SrcObject->InnerObjectWithDirectlyNestedObject);
+			CHECK(DstObject->InnerObjectWithDirectlyNestedObject->TestValue == SrcObject->InnerObjectWithDirectlyNestedObject->TestValue);
+
+			CHECK(DstObject->InnerObjectWithDirectlyNestedObject->InnerObject != nullptr);
+			CHECK(DstObject->InnerObjectWithDirectlyNestedObject->InnerObject->IsInOuter(DstObject->InnerObjectWithDirectlyNestedObject));
+			CHECK(DstObject->InnerObjectWithDirectlyNestedObject->InnerObject != SrcObject->InnerObjectWithDirectlyNestedObject->InnerObject);
+			CHECK(DstObject->InnerObjectWithDirectlyNestedObject->InnerObject->TestValue == SrcObject->InnerObjectWithDirectlyNestedObject->InnerObject->TestValue);
+
+			CHECK(DstObject->InnerObjectWithIndirectlyNestedObject != nullptr);
+			CHECK(DstObject->InnerObjectWithIndirectlyNestedObject->IsInOuter(DstObject));
+			CHECK(DstObject->InnerObjectWithIndirectlyNestedObject != SrcObject->InnerObjectWithIndirectlyNestedObject);
+			CHECK(DstObject->InnerObjectWithIndirectlyNestedObject->TestValue == SrcObject->InnerObjectWithIndirectlyNestedObject->TestValue);
+
+			CHECK(DstObject->InnerObjectWithIndirectlyNestedObject->InnerObject != nullptr);
+			CHECK(DstObject->InnerObjectWithIndirectlyNestedObject->InnerObject->IsInOuter(DstObject->InnerObjectWithIndirectlyNestedObject));
+			CHECK(DstObject->InnerObjectWithIndirectlyNestedObject->InnerObject != SrcObject->InnerObjectWithIndirectlyNestedObject->InnerObject);
+			CHECK(DstObject->InnerObjectWithIndirectlyNestedObject->InnerObject->TestValue == SrcObject->InnerObjectWithIndirectlyNestedObject->InnerObject->TestValue);
+
+			CHECK(DstObject->InnerObjectWithIndirectlyNestedObject->InnerObject->InnerObject != nullptr);
+			CHECK(DstObject->InnerObjectWithIndirectlyNestedObject->InnerObject->InnerObject->IsInOuter(DstObject->InnerObjectWithIndirectlyNestedObject->InnerObject));
+			CHECK(DstObject->InnerObjectWithIndirectlyNestedObject->InnerObject->InnerObject != SrcObject->InnerObjectWithIndirectlyNestedObject->InnerObject->InnerObject);
+			CHECK(DstObject->InnerObjectWithIndirectlyNestedObject->InnerObject->InnerObject->TestValue == SrcObject->InnerObjectWithIndirectlyNestedObject->InnerObject->InnerObject->TestValue);
+		}
+	}
+
+	void RunSubobjectInstancingTests(UClass* InstancingTestClass)
+	{
+		RunSubobjectInstancingTests_Base(InstancingTestClass);
+		RunSubobjectInstancingTests_Array(InstancingTestClass);
+		RunSubobjectInstancingTests_Set(InstancingTestClass);
+		RunSubobjectInstancingTests_Map(InstancingTestClass);
+		RunSubobjectInstancingTests_Optional(InstancingTestClass);
+		RunSubobjectInstancingTests_Struct(InstancingTestClass);
+		RunSubobjectInstancingTests_Nested(InstancingTestClass);
+		RunSubobjectInstancingTests_Deferred(InstancingTestClass);
+		RunSubobjectInstancingTests_SelfReferencing(InstancingTestClass);
+		RunSubobjectInstancingTests_Other(InstancingTestClass);
+	}
+}
+
+const FName FSubobjectInstancingTestUtils::GetNonNativeInnerObjectPropertyName()
+{
+	return SubobjectInstancingTest::Private::NonNativeInnerObjectPropertyName;
+}
+
+const FName FSubobjectInstancingTestUtils::GetNonNativeOuterObjectPropertyName()
+{
+	return SubobjectInstancingTest::Private::NonNativeOuterObjectPropertyName;
+}
+
+const FName FSubobjectInstancingTestUtils::GetNonNativeSelfReferencePropertyName()
+{
+	return SubobjectInstancingTest::Private::NonNativeSelfReferencePropertyName;
+}
+
+const FName FSubobjectInstancingTestUtils::GetNonNativeEditTimeInnerObjectPropertyName()
+{
+	return SubobjectInstancingTest::Private::NonNativeEditTimeInnerObjectPropertyName;
+}
+
+UClass* FSubobjectInstancingTestUtils::CreateNonNativeInstancingTestClass(UClass* SuperClass)
+{
+	return SubobjectInstancingTest::Private::NewTestClass<UClass>(SuperClass);
+}
+
+UClass* FSubobjectInstancingTestUtils::CreateDynamicallyInstancedTestClass(UClass* SuperClass)
+{
+	return SubobjectInstancingTest::Private::NewTestClass<UDynamicSubobjectInstancingTestClass>(SuperClass);
+}
+
+TEST_CASE_NAMED(FNativeDefaultSubobjectsTest, "CoreUObject::NativeDefaultSubobjects", "[CoreUObject][EngineFilter]")
+{
+	SECTION("Class default subobjects")
+	{
+		const USubobjectInstancingTestOuterObject* CDO = GetDefault<USubobjectInstancingTestOuterObject>();
+
+		REQUIRE(CDO->InnerObject != nullptr);
+		CHECK(CDO->InnerObject->IsInOuter(CDO));
+
+		CHECK(CDO->InnerObject->HasAllFlags(RF_ArchetypeObject));
+		CHECK(CDO->InnerObject->HasAllFlags(RF_DefaultSubObject));
+		CHECK(CDO->InnerObject->HasAllFlags(CDO->GetMaskedFlags(RF_PropagateToSubObjects)));
+
+		CHECK(CDO->InnerObject->IsTemplate());
+		CHECK(CDO->InnerObject->IsDefaultSubobject());
+		CHECK(CDO->InnerObject->IsTemplateForSubobjects());
+
+		REQUIRE(CDO->InnerObjectWithDirectlyNestedObject != nullptr);
+		REQUIRE(CDO->InnerObjectWithDirectlyNestedObject->InnerObject != nullptr);
+		CHECK(CDO->InnerObjectWithDirectlyNestedObject->InnerObject->IsInOuter(CDO->InnerObjectWithDirectlyNestedObject));
+
+		CHECK(CDO->InnerObjectWithDirectlyNestedObject->InnerObject->IsTemplate());
+		CHECK(CDO->InnerObjectWithDirectlyNestedObject->InnerObject->IsDefaultSubobject());
+		CHECK(CDO->InnerObjectWithDirectlyNestedObject->InnerObject->IsTemplateForSubobjects());
+
+		REQUIRE(CDO->InnerObjectWithIndirectlyNestedObject != nullptr);
+		REQUIRE(CDO->InnerObjectWithIndirectlyNestedObject->InnerObject != nullptr);
+		REQUIRE(CDO->InnerObjectWithIndirectlyNestedObject->InnerObject->InnerObject != nullptr);
+		CHECK(CDO->InnerObjectWithIndirectlyNestedObject->InnerObject->IsInOuter(CDO->InnerObjectWithIndirectlyNestedObject));
+		CHECK(CDO->InnerObjectWithIndirectlyNestedObject->InnerObject->InnerObject->IsInOuter(CDO->InnerObjectWithIndirectlyNestedObject->InnerObject));
+
+		CHECK(CDO->InnerObjectWithIndirectlyNestedObject->InnerObject->IsTemplate());
+		CHECK(CDO->InnerObjectWithIndirectlyNestedObject->InnerObject->IsDefaultSubobject());
+		CHECK(CDO->InnerObjectWithIndirectlyNestedObject->InnerObject->IsTemplateForSubobjects());
+
+		CHECK(CDO->InnerObjectWithIndirectlyNestedObject->InnerObject->InnerObject->IsTemplate());
+		CHECK(CDO->InnerObjectWithIndirectlyNestedObject->InnerObject->InnerObject->IsDefaultSubobject());
+		CHECK(CDO->InnerObjectWithIndirectlyNestedObject->InnerObject->InnerObject->IsTemplateForSubobjects());
+	}
+
+	SECTION("Template default subobjects")
+	{
+		USubobjectInstancingTestOuterObject* Template = NewObject<USubobjectInstancingTestOuterObject>(GetTransientPackage(), NAME_None, RF_ArchetypeObject);
+
+		REQUIRE(Template->InnerObject != nullptr);
+		CHECK(Template->InnerObject->IsInOuter(Template));
+
+		CHECK(Template->InnerObject->HasAllFlags(RF_ArchetypeObject));
+		CHECK(Template->InnerObject->HasAllFlags(RF_DefaultSubObject));
+		CHECK(Template->InnerObject->HasAllFlags(Template->GetMaskedFlags(RF_PropagateToSubObjects)));
+
+		CHECK(Template->InnerObject->IsTemplate());
+		CHECK(Template->InnerObject->IsDefaultSubobject());
+		CHECK(Template->InnerObject->IsTemplateForSubobjects());
+
+		REQUIRE(Template->InnerObjectWithDirectlyNestedObject != nullptr);
+		REQUIRE(Template->InnerObjectWithDirectlyNestedObject->InnerObject != nullptr);
+		CHECK(Template->InnerObjectWithDirectlyNestedObject->InnerObject->IsInOuter(Template->InnerObjectWithDirectlyNestedObject));
+
+		CHECK(Template->InnerObjectWithDirectlyNestedObject->InnerObject->IsTemplate());
+		CHECK(Template->InnerObjectWithDirectlyNestedObject->InnerObject->IsDefaultSubobject());
+		CHECK(Template->InnerObjectWithDirectlyNestedObject->InnerObject->IsTemplateForSubobjects());
+
+		REQUIRE(Template->InnerObjectWithIndirectlyNestedObject != nullptr);
+		REQUIRE(Template->InnerObjectWithIndirectlyNestedObject->InnerObject != nullptr);
+		REQUIRE(Template->InnerObjectWithIndirectlyNestedObject->InnerObject->InnerObject != nullptr);
+		CHECK(Template->InnerObjectWithIndirectlyNestedObject->InnerObject->IsInOuter(Template->InnerObjectWithIndirectlyNestedObject));
+		CHECK(Template->InnerObjectWithIndirectlyNestedObject->InnerObject->InnerObject->IsInOuter(Template->InnerObjectWithIndirectlyNestedObject->InnerObject));
+
+		CHECK(Template->InnerObjectWithIndirectlyNestedObject->InnerObject->IsTemplate());
+		CHECK(Template->InnerObjectWithIndirectlyNestedObject->InnerObject->IsDefaultSubobject());
+		CHECK(Template->InnerObjectWithIndirectlyNestedObject->InnerObject->IsTemplateForSubobjects());
+
+		CHECK(Template->InnerObjectWithIndirectlyNestedObject->InnerObject->InnerObject->IsTemplate());
+		CHECK(Template->InnerObjectWithIndirectlyNestedObject->InnerObject->InnerObject->IsDefaultSubobject());
+		CHECK(Template->InnerObjectWithIndirectlyNestedObject->InnerObject->InnerObject->IsTemplateForSubobjects());
+	}
+
+	SECTION("Default subobject constructed using SetDefaultSubobjectClass() override")
+	{
+		USubobjectInstancingTestDerivedOuterObjectWithTypeOverride* OuterObject = NewObject<USubobjectInstancingTestDerivedOuterObjectWithTypeOverride>(GetTransientPackage());
+
+		REQUIRE(OuterObject->InnerObject != nullptr);
+		CHECK(OuterObject->InnerObject->IsA<USubobjectInstancingTestDerivedObject>());
+		CHECK(OuterObject->InnerObject->GetArchetype()->IsA<USubobjectInstancingTestObject>());
+	}
+
+	SECTION("Optional default subobject excluded from construction using DoNotCreateDefaultSubobject() override")
+	{
+		USubobjectInstancingTestDerivedOuterObjectWithDoNotCreateOverride* OuterObject = NewObject<USubobjectInstancingTestDerivedOuterObjectWithDoNotCreateOverride>(GetTransientPackage());
+
+		REQUIRE(OuterObject->OptionalInnerObject.IsSet());
+		CHECK(OuterObject->OptionalInnerObject.GetValue() == nullptr);
+	}
+}
+
+TEST_CASE_NAMED(FNativeInstancedSubobjectsTest, "CoreUObject::NativeInstancedSubobjects", "[CoreUObject][EngineFilter]")
+{
+	UClass* NativeInstancingTestClass = USubobjectInstancingTestOuterObject::StaticClass();
+
+	REQUIRE(NativeInstancingTestClass != nullptr);
+	SubobjectInstancingTest::Private::RunSubobjectInstancingTests(NativeInstancingTestClass);
+}
+
+TEST_CASE_NAMED(FNonNativeInstancedSubobjectsTest, "CoreUObject::NonNativeInstancedSubobjects", "[CoreUObject][EngineFilter]")
+{
+	UClass* NativeInstancingTestClass = USubobjectInstancingTestOuterObject::StaticClass();
+	UClass* NonNativeInstancingTestClass = FSubobjectInstancingTestUtils::CreateNonNativeInstancingTestClass(NativeInstancingTestClass);
+
+	REQUIRE(NonNativeInstancingTestClass != nullptr);
+	SubobjectInstancingTest::Private::RunSubobjectInstancingTests(NonNativeInstancingTestClass);
+}
+
+TEST_CASE_NAMED(FDynamicSubobjectInstancingTest, "CoreUObject::DynamicSubobjectInstancing", "[CoreUObject][EngineFilter]")
+{
+	UClass* NativeInstancingTestClass = USubobjectInstancingTestOuterObject::StaticClass();
+	UClass* DynamicallyInstancedTestClass = FSubobjectInstancingTestUtils::CreateDynamicallyInstancedTestClass(NativeInstancingTestClass);
+
+	REQUIRE(DynamicallyInstancedTestClass != nullptr);
+	CHECK(DynamicallyInstancedTestClass->ShouldUseDynamicSubobjectInstancing());
+	SubobjectInstancingTest::Private::RunSubobjectInstancingTests(DynamicallyInstancedTestClass);
+}
+
+}
+
+#endif // WITH_TESTS
\ No newline at end of file
diff --git a/Engine/Source/Runtime/CoreUObject/Tests/UObject/SubobjectInstancingTest.h b/Engine/Source/Runtime/CoreUObject/Tests/UObject/SubobjectInstancingTest.h
new file mode 100644
--- /dev/null
+++ b/Engine/Source/Runtime/CoreUObject/Tests/UObject/SubobjectInstancingTest.h
@@ -0,0 +1,204 @@
+// Copyright Epic Games, Inc. All Rights Reserved.
+
+#pragma once
+
+#if WITH_TESTS
+
+#include "UObject/Class.h"
+#include "UObject/Object.h"
+#include "SubobjectInstancingTest.generated.h"
+
+namespace UE
+{
+
+struct FSubobjectInstancingTestUtils
+{
+	static const FName GetNonNativeInnerObjectPropertyName();
+	static const FName GetNonNativeOuterObjectPropertyName();
+	static const FName GetNonNativeSelfReferencePropertyName();
+	static const FName GetNonNativeEditTimeInnerObjectPropertyName();
+
+	static UClass* CreateNonNativeInstancingTestClass(UClass* SuperClass = nullptr);
+	static UClass* CreateDynamicallyInstancedTestClass(UClass* SuperClass = nullptr);
+};
+
+}
+
+UCLASS()
+class USubobjectInstancingTestObject : public UObject
+{
+	GENERATED_BODY()
+
+public:
+	UPROPERTY()
+	int32 TestValue = 100;
+};
+
+UCLASS()
+class USubobjectInstancingTestDerivedObject : public USubobjectInstancingTestObject
+{
+	GENERATED_BODY()
+};
+
+UCLASS()
+class USubobjectInstancingTestDirectlyNestedObject : public USubobjectInstancingTestObject
+{
+	GENERATED_BODY()
+
+public:
+	USubobjectInstancingTestDirectlyNestedObject();
+
+	UPROPERTY(Instanced)
+	TObjectPtr<UObject> SelfRef;
+
+	UPROPERTY(Instanced)
+	TObjectPtr<UObject> OwnerObject;
+
+	UPROPERTY(Instanced)
+	TObjectPtr<USubobjectInstancingTestObject> InnerObject;
+};
+
+UCLASS()
+class USubobjectInstancingTestIndirectlyNestedObject : public USubobjectInstancingTestObject
+{
+	GENERATED_BODY()
+
+public:
+	USubobjectInstancingTestIndirectlyNestedObject();
+
+	UPROPERTY(Instanced)
+	TObjectPtr<USubobjectInstancingTestDirectlyNestedObject> InnerObject;
+};
+
+UCLASS(DefaultToInstanced)
+class USubobjectInstancingDefaultToInstancedTestObject : public USubobjectInstancingTestObject
+{
+	GENERATED_BODY()
+};
+
+USTRUCT()
+struct FSubobjectInstancingTestStructType
+{
+	GENERATED_BODY()
+
+	UPROPERTY(Instanced)
+	TObjectPtr<USubobjectInstancingTestObject> InnerObject;
+
+	UPROPERTY()
+	TObjectPtr<USubobjectInstancingDefaultToInstancedTestObject> InnerObjectFromType;
+
+	UPROPERTY(Instanced)
+	TArray<TObjectPtr<USubobjectInstancingTestObject>> InnerObjectArray;
+};
+
+UCLASS()
+class USubobjectInstancingTestOuterObject : public UObject
+{
+    GENERATED_BODY()
+
+public:
+	USubobjectInstancingTestOuterObject(const FObjectInitializer& ObjectInitializer = FObjectInitializer::Get());
+
+	virtual void PostInitProperties() override;
+
+	UPROPERTY(Instanced)
+	TObjectPtr<UObject> SelfRef;
+
+	UPROPERTY(Instanced)
+	TObjectPtr<USubobjectInstancingTestObject> NullObject;
+
+	UPROPERTY(Instanced)
+	TObjectPtr<USubobjectInstancingTestObject> InnerObject;
+
+	UPROPERTY()
+	TObjectPtr<USubobjectInstancingTestObject> SharedObject;
+
+	UPROPERTY()
+	TObjectPtr<USubobjectInstancingTestObject> ExternalObject;
+
+	UPROPERTY(Instanced)
+	TObjectPtr<USubobjectInstancingTestObject> InternalObject;
+
+	UPROPERTY(Transient, Instanced)
+	TObjectPtr<USubobjectInstancingTestObject> TransientInnerObject;
+
+	UPROPERTY(Instanced)
+	TObjectPtr<USubobjectInstancingTestObject> EditTimeInnerObject;
+
+	UPROPERTY(Instanced)
+	TObjectPtr<USubobjectInstancingTestObject> LocalOnlyInnerObject;
+
+	UPROPERTY()
+	TObjectPtr<USubobjectInstancingDefaultToInstancedTestObject> InnerObjectFromType;
+
+	UPROPERTY(Instanced)
+	TObjectPtr<USubobjectInstancingTestObject> InnerObjectUsingNew;
+
+	UPROPERTY(Instanced)
+	TObjectPtr<USubobjectInstancingTestObject> InnerObjectPostInit;
+
+	UPROPERTY(Instanced)
+	TArray<TObjectPtr<USubobjectInstancingTestObject>> InnerObjectArray;
+
+	UPROPERTY()
+	TArray<TObjectPtr<USubobjectInstancingDefaultToInstancedTestObject>> InnerObjectFromTypeArray;
+
+	UPROPERTY(Instanced)
+	TSet<TObjectPtr<USubobjectInstancingTestObject>> InnerObjectSet;
+
+	UPROPERTY()
+	TSet<TObjectPtr<USubobjectInstancingDefaultToInstancedTestObject>> InnerObjectFromTypeSet;
+
+	UPROPERTY(Instanced)
+	TMap<TObjectPtr<USubobjectInstancingTestObject>, TObjectPtr<USubobjectInstancingTestObject>> InnerObjectMap;
+
+	UPROPERTY()
+	TMap<TObjectPtr<USubobjectInstancingDefaultToInstancedTestObject>, TObjectPtr<USubobjectInstancingDefaultToInstancedTestObject>> InnerObjectFromTypeMap;
+
+	UPROPERTY()
+	FSubobjectInstancingTestStructType StructWithInnerObjects;
+
+	UPROPERTY(Instanced)
+	TOptional<TObjectPtr<USubobjectInstancingTestObject>> OptionalInnerObject;
+
+	UPROPERTY()
+	TOptional<TObjectPtr<USubobjectInstancingDefaultToInstancedTestObject>> OptionalInnerObjectFromType;
+
+	UPROPERTY(Instanced)
+	TOptional<TArray<TObjectPtr<USubobjectInstancingTestObject>>>  OptionalInnerObjectArray;
+
+	UPROPERTY(Instanced)
+	TObjectPtr<USubobjectInstancingTestDirectlyNestedObject> InnerObjectWithDirectlyNestedObject;
+
+	UPROPERTY(Instanced)
+	TObjectPtr<USubobjectInstancingTestIndirectlyNestedObject> InnerObjectWithIndirectlyNestedObject;
+};
+
+UCLASS()
+class USubobjectInstancingTestDerivedOuterObjectWithTypeOverride : public USubobjectInstancingTestOuterObject
+{
+	GENERATED_BODY()
+
+public:
+	USubobjectInstancingTestDerivedOuterObjectWithTypeOverride(const FObjectInitializer& ObjectInitializer);
+};
+
+UCLASS()
+class USubobjectInstancingTestDerivedOuterObjectWithDoNotCreateOverride : public USubobjectInstancingTestOuterObject
+{
+	GENERATED_BODY()
+
+public:
+	USubobjectInstancingTestDerivedOuterObjectWithDoNotCreateOverride(const FObjectInitializer& ObjectInitializer);
+};
+
+UCLASS()
+class UDynamicSubobjectInstancingTestClass : public UClass
+{
+	GENERATED_BODY()
+
+public:
+	UDynamicSubobjectInstancingTestClass();
+};
+
+#endif	// WITH_TESTS
