remote exec:
- reduces the number of net pumps we do to be minimal
- we now only pump the network on the second 'iteration' through the work, which only happens if transactions aborted due to waiting for dependencies
- if we queue work up and execute it all successfully, we will no longer pump the network at all internally
- added GRemoteExecutorForcedNetPumpIterationCount (defaults 0 = off) to enable an occasional unconditional net pump

diff --git a/Engine/Source/Runtime/CoreUObject/Private/UObject/RemoteExecutor.cpp b/Engine/Source/Runtime/CoreUObject/Private/UObject/RemoteExecutor.cpp
--- a/Engine/Source/Runtime/CoreUObject/Private/UObject/RemoteExecutor.cpp
+++ b/Engine/Source/Runtime/CoreUObject/Private/UObject/RemoteExecutor.cpp
@@ -6,6 +6,18 @@
 
 DEFINE_LOG_CATEGORY(LogRemoteExec);
 
+int32 GRemoteExecutorForcedNetPumpIterationCount = 0;
+static FAutoConsoleVariableRef CVarRemoteExecutorForcedNetPumpIterationCount(
+	TEXT("remoteexecutor.ForcedNetPumpIterationCount"),
+	GRemoteExecutorForcedNetPumpIterationCount,
+	TEXT("How many execution pumps can we perform before we force pump the network? 0 = disabled"));
+
+bool GRemoteExecutorNewNetPumpBehavior = true;
+static FAutoConsoleVariableRef CVarRemoteExecutorNewNetPumpBehavior(
+	TEXT("remoteexecutor.NewNetPumpBehavior"),
+	GRemoteExecutorNewNetPumpBehavior,
+	TEXT("Enable the 'new' net pump behavior"));
+
 namespace UE::RemoteExecutor
 {
 TDelegate<void()> TickNetworkDelegate;
@@ -66,6 +78,13 @@ class FRemoteExecutor
 
 	TArray<FRemoteExecutorWork> PendingWorks;
 
+	// GRemoteExecutorForcedNetPumpIterationCount controls how many iterations
+	// we're allowed to go before we unconditionally try to pump the network
+	// (to reduce latency of servicing requests of other servers)
+	// Here we keep track of how many iterations it has been since the last
+	// network pump happened.
+	uint32 IterationsSinceNetworkPump = 0;
+
 public:
 	FRemoteExecutorWork* ExecutingWork = nullptr;
 
@@ -149,28 +168,10 @@ public:
 				bPrintStall = true;
 			}
 
-			// pump the network
-			UE::RemoteExecutor::TickNetworkDelegate.ExecuteIfBound();
-
-			// fetch and enqueue all RPCs to execute from the network
-			for (;;)
+			if (!GRemoteExecutorNewNetPumpBehavior)
 			{
-				TOptional<TTuple<FName, FRemoteWorkPriority, bool, TFunction<void(void)>>> NewWorkOpt = UE::RemoteExecutor::FetchNextDeferredRPCDelegate.Execute();
-				if (NewWorkOpt.IsSet())
-				{
-					TTuple<FName, FRemoteWorkPriority, bool, TFunction<void(void)>>& NewWork = NewWorkOpt.GetValue();
-					EnqueueWork(NewWork.Get<0>(), NewWork.Get<1>(), NewWork.Get<2>(), NewWork.Get<3>());
-
-					UE_LOG(LogRemoteExec, VeryVerbose, TEXT("ExecutePendingWork[%d] enqueued deferred RPC '%s' %s %s"),
-						LocalIterationNumber,
-						*NewWork.Get<0>().ToString(),
-						*NewWork.Get<1>().ToString(),
-						NewWork.Get<2>() ? TEXT("(transactional)") :  TEXT("(non-transactional)"));
-				}
-				else
-				{
-					break;
-				}
+				// pump the network
+				UE::RemoteExecutor::TickNetworkDelegate.ExecuteIfBound();
 			}
 
 			// check if we're actively servicing a remote multi-server commit
@@ -187,8 +188,14 @@ public:
 						*ActiveRemoteMultiServerCommitRequestId.ToString());
 				}
 
-				// if we have an active remote multi-server commit that we are servicing, we pause executing any local work
-				// until that is complete
+				// if we have an active remote multi-server commit that we are servicing, we pump the network and 
+				// pause executing any local work until that is complete
+				if (GRemoteExecutorNewNetPumpBehavior)
+				{
+					IterationsSinceNetworkPump = 0;
+					UE::RemoteExecutor::TickNetworkDelegate.ExecuteIfBound();
+				}
+
 				continue;
 			}
 
@@ -207,6 +214,26 @@ public:
 				break;
 			}
 
+			if (GRemoteExecutorNewNetPumpBehavior)
+			{
+				IterationsSinceNetworkPump++;
+
+				// if GRemoteExecutorForcedNetPumpIterationCount is nonzero, we may need to unconditionally pump the network
+				bool bUnconditionallyPumpNetwork = false;
+			
+				if (GRemoteExecutorForcedNetPumpIterationCount)
+				{
+					bUnconditionallyPumpNetwork = (IterationsSinceNetworkPump % GRemoteExecutorForcedNetPumpIterationCount) == 0;
+				}
+
+				if ((LocalIterationNumber > 0) || bUnconditionallyPumpNetwork)
+				{
+					// pump the network
+					IterationsSinceNetworkPump = 0;
+					UE::RemoteExecutor::TickNetworkDelegate.ExecuteIfBound();
+				}
+			}
+			
 			for (int32 PendingWorkIndex = 0; PendingWorkIndex < PendingWorks.Num(); PendingWorkIndex++)
 			{
 				FRemoteExecutorWork& PendingWork = PendingWorks[PendingWorkIndex];
@@ -635,8 +662,8 @@ void ExecuteTransactionalInternal(FName WorkName, FRemoteWorkPriority InWorkPrio
 #if UE_WITH_REMOTE_OBJECT_HANDLE
 	check(!AutoRTFM::IsClosed());
 
-	GRemoteExecutor.EnqueueWork(WorkName, InWorkPriority, true, [&InWork]() { InWork(); });
-	GRemoteExecutor.ExecutePendingWork();
+	UE::RemoteExecutor::EnqueueWorkWithExplicitPriority(WorkName, InWorkPriority, true, [&InWork]() { InWork(); });
+	UE::RemoteExecutor::ExecutePendingWork();
 #else
 	InWork();
 #endif
@@ -918,19 +945,19 @@ FRemoteWorkPriority CreateRootWorkPriority()
 
 void ExecutePendingWork()
 {
+	static const FName Name = "ExecuteTransactional";
+	UE::RemoteExecutor::OnRegionBeginDelegate.ExecuteIfBound(Name);
 	GRemoteExecutor.ExecutePendingWork();
+	UE::RemoteExecutor::OnRegionEndDelegate.ExecuteIfBound(FString());
 }
 
 void ExecuteTransactionalWithExplicitPriority(FName WorkName, FRemoteWorkPriority WorkPriority, const TFunctionRef<void(void)>& Work)
 {
-	static const FName Name = "ExecuteTransactional";
-	UE::RemoteExecutor::OnRegionBeginDelegate.ExecuteIfBound(Name);
 #if UE_WITH_REMOTE_OBJECT_HANDLE
 	ExecuteTransactionalInternal(WorkName, WorkPriority, Work);
 #else
 	Work();
 #endif
-	UE::RemoteExecutor::OnRegionEndDelegate.ExecuteIfBound(FString());
 }
 
 void ExecuteTransactional(FName WorkName, const TFunctionRef<void(void)>& Work)
@@ -939,17 +966,22 @@ void ExecuteTransactional(FName WorkName, const TFunctionRef<void(void)>& Work)
 	ExecuteTransactionalWithExplicitPriority(WorkName, RootWorkPriority, Work);
 }
 
-void EnqueueTransactional(FName WorkName, const TFunction<void(void)>& InWork)
+void EnqueueWorkWithExplicitPriority(FName WorkName, FRemoteWorkPriority WorkPriority, bool bIsTransactional, const TFunction<void(void)>& InWork)
 {
 #if UE_WITH_REMOTE_OBJECT_HANDLE
-	check(!AutoRTFM::IsClosed());
-	FRemoteWorkPriority RootWorkPriority = FRemoteWorkPriority::CreateRootWorkPriority(FRemoteServerId::GetLocalServerId(), GRemoteExecutor.GenerateNextTransactionId());
-	GRemoteExecutor.EnqueueWork(WorkName, RootWorkPriority, true, InWork);
+	GRemoteExecutor.EnqueueWork(WorkName, WorkPriority, bIsTransactional, InWork);
 #else
 	InWork();
 #endif
 }
 
+void EnqueueWork(FName WorkName, bool bIsTransactional, const TFunction<void(void)>& InWork)
+{
+	check(!AutoRTFM::IsClosed());
+	FRemoteWorkPriority RootWorkPriority = FRemoteWorkPriority::CreateRootWorkPriority(FRemoteServerId::GetLocalServerId(), GRemoteExecutor.GenerateNextTransactionId());
+	EnqueueWorkWithExplicitPriority(WorkName, RootWorkPriority, bIsTransactional, InWork);
+}
+
 } // namespace UE::RemoteExecutor
 
 FString FRemoteWorkPriority::ToString() const
diff --git a/Engine/Source/Runtime/CoreUObject/Public/UObject/RemoteExecutor.h b/Engine/Source/Runtime/CoreUObject/Public/UObject/RemoteExecutor.h
--- a/Engine/Source/Runtime/CoreUObject/Public/UObject/RemoteExecutor.h
+++ b/Engine/Source/Runtime/CoreUObject/Public/UObject/RemoteExecutor.h
@@ -226,7 +226,8 @@ namespace UE::RemoteExecutor
 
 	COREUOBJECT_API void ExecuteTransactional(FName WorkName, const TFunctionRef<void(void)>& Work);
 	COREUOBJECT_API void ExecuteTransactionalWithExplicitPriority(FName WorkName, FRemoteWorkPriority WorkPriority, const TFunctionRef<void(void)>& Work);
-	COREUOBJECT_API void EnqueueTransactional(FName WorkName, const TFunction<void(void)>& Work);
+	COREUOBJECT_API void EnqueueWorkWithExplicitPriority(FName WorkName, FRemoteWorkPriority WorkPriority, bool bIsTransactional, const TFunction<void(void)>& Work);
+	COREUOBJECT_API void EnqueueWork(FName WorkName, bool bIsTransactional, const TFunction<void(void)>& Work);
 	COREUOBJECT_API void ExecutePendingWork();
 
 	COREUOBJECT_API void AbortTransactionRequiresDependencies(FStringView Description);
diff --git a/Engine/Source/Runtime/Engine/Private/TickTaskManager.cpp b/Engine/Source/Runtime/Engine/Private/TickTaskManager.cpp
--- a/Engine/Source/Runtime/Engine/Private/TickTaskManager.cpp
+++ b/Engine/Source/Runtime/Engine/Private/TickTaskManager.cpp
@@ -429,7 +429,7 @@ public:
 
 				if (Target->bRunTransactionally)
 				{
-					UE::RemoteExecutor::EnqueueTransactional(Target->CachedDiagnosticContext, ExecuteTickWork);
+					UE::RemoteExecutor::EnqueueWork(Target->CachedDiagnosticContext, true, ExecuteTickWork);
 				}
 				else
 				{
