Submitting on behalf of Jodon Karlik: When using Remote Objects, do not warn about any missing Garbage objects.

diff --git a/Engine/Source/Runtime/CoreUObject/Private/UObject/RemoteObjectTransfer.cpp b/Engine/Source/Runtime/CoreUObject/Private/UObject/RemoteObjectTransfer.cpp
--- a/Engine/Source/Runtime/CoreUObject/Private/UObject/RemoteObjectTransfer.cpp
+++ b/Engine/Source/Runtime/CoreUObject/Private/UObject/RemoteObjectTransfer.cpp
@@ -863,7 +863,9 @@ UObject* MigrateObjectFromObjectData(FRemoteObjectData& ObjectData, const FUObje
 		// Let's make sure we received data for all of the objects
 		const UObject* RootObject = ReceivedObjects[0];
 		TArray<UObject*> AllSubObjects;
-		GetObjectsWithOuter(RootObject, AllSubObjects);
+		constexpr bool bIncludeNestedObjects = true;
+		constexpr EObjectFlags ExclusionFlags = EObjectFlags::RF_MirroredGarbage;
+		GetObjectsWithOuter(RootObject, AllSubObjects, bIncludeNestedObjects, ExclusionFlags);
 
 		const bool bSameNumberOfObjects = (1 + AllSubObjects.Num() == ReceivedObjects.Num());
 		if (!bSameNumberOfObjects)
