Set the RF_MirroredGarbage immediately (in the open) when calling Object->MarkAsGarbage() because something may be checking that flag (through IsValid(Object)) in the same transaction.

diff --git a/Engine/Source/Runtime/CoreUObject/Public/UObject/UObjectBaseUtility.h b/Engine/Source/Runtime/CoreUObject/Public/UObject/UObjectBaseUtility.h
--- a/Engine/Source/Runtime/CoreUObject/Public/UObject/UObjectBaseUtility.h
+++ b/Engine/Source/Runtime/CoreUObject/Public/UObject/UObjectBaseUtility.h
@@ -177,12 +177,15 @@ public:
 	{
 		check(!IsRooted());
 
+		// We need to set the RF_MirroredGarbage flag before OnCommit because something in the same transaction may be checking this flag and AtomicallySetFlags() is already AutoRTFM aware.
+		// A common example used across many of the engine systems is IsValid(Object) which checks the RF_MirroredGarbage flag.
+		// However both EInternalObjectFlags::Async and EInternalObjectFlags::Garbage flags are only checked by GC so it's safe to clear/set them OnCommit.
+		AtomicallySetFlags(RF_MirroredGarbage);
+
 		// only set these flags after we have finished our transaction
 		// GC should not happen during a transaction so delaying this wont cause GC issues
 		UE_AUTORTFM_ONCOMMIT(this)
 		{
-			AtomicallySetFlags(RF_MirroredGarbage);
-
 			// If we explicitly marked the object as garbage, remove the async flag so it's visible to the GC
 			AtomicallyClearInternalFlags(EInternalObjectFlags::Async);
 
