[PropertyObject]
- Improve logic to avoid doing a full load if we can find the object already in memory

diff --git a/Engine/Source/Runtime/CoreUObject/Private/UObject/PropertyObject.cpp b/Engine/Source/Runtime/CoreUObject/Private/UObject/PropertyObject.cpp
--- a/Engine/Source/Runtime/CoreUObject/Private/UObject/PropertyObject.cpp
+++ b/Engine/Source/Runtime/CoreUObject/Private/UObject/PropertyObject.cpp
@@ -70,13 +70,20 @@ EConvertFromTypeResult FObjectProperty::ConvertFromType(const FPropertyTag& Tag,
 		// If we're async loading, only zenloader supports synchronous loads from the async loading thread.
 		// Other loaders are not safe to do a sync load because they may crash or fail to set the variable, so throw an error if it's not already in memory
 		ELoaderType LoaderType = GetLoaderType();
-		if (!IsInGameThread() && (LoaderType != ELoaderType::ZenLoader))
+		if (!IsInGameThread())
 		{
 			PreviousValueObj = PreviousValue.Get();
 
 			if (!PreviousValueObj && !PreviousValue.IsNull())
 			{
-				UE_LOG(LogClass, Error, TEXT("Failed to convert soft path %s to unloaded object as this is not supported during async loading with the currently active loader '%s'. Load and resave %s in the editor to fix!"), *PreviousValue.ToString(), LexToString(LoaderType), *Slot.GetUnderlyingArchive().GetArchiveName());
+				if (LoaderType != ELoaderType::ZenLoader)
+				{
+					UE_LOG(LogClass, Error, TEXT("Failed to convert soft path %s to unloaded object as this is not supported during async loading with the currently active loader '%s'. Load and resave %s in the editor to fix!"), *PreviousValue.ToString(), LexToString(LoaderType), *Slot.GetUnderlyingArchive().GetArchiveName());
+				}
+				else
+				{
+					PreviousValueObj = PreviousValue.LoadSynchronous();
+				}
 			}
 		}
 		else
