Struct Utils: trivial crash fix when attempting to use GetStructHash64 with User Defined Structs.
User Defined Structs have no CppStructOps set.

diff --git a/Engine/Source/Runtime/CoreUObject/Private/StructUtils/StructUtilsTypes.cpp b/Engine/Source/Runtime/CoreUObject/Private/StructUtils/StructUtilsTypes.cpp
--- a/Engine/Source/Runtime/CoreUObject/Private/StructUtils/StructUtilsTypes.cpp
+++ b/Engine/Source/Runtime/CoreUObject/Private/StructUtils/StructUtilsTypes.cpp
@@ -187,7 +187,8 @@ namespace UE::StructUtils
 	{
 		const uint64 BaseHash = FArchiveCityHash64::HashString(ScriptStruct.GetPathName());
 
-		if (ScriptStruct.GetCppStructOps()->HasGetTypeHash())
+		const UScriptStruct::ICppStructOps* CppStructOps = ScriptStruct.GetCppStructOps();
+		if (CppStructOps && CppStructOps->HasGetTypeHash())
 		{
 			const uint32 StructHash = ScriptStruct.GetStructTypeHash(StructMemory);
 			return CityHash128to64(Uint128_64(BaseHash, StructHash));
