[Bug Patch] There's existing bugs with impersonation of IDOs that are blocking users. As a workaround, this change introduces a CVar that when enabled, will only impersonate IDOs with unknown properties. This means that the bugs will still exist, but their impact will be localized to objects who's data already contains unknown properties.

diff --git a/Engine/Source/Runtime/CoreUObject/Internal/UObject/InstanceDataObjectUtils.h b/Engine/Source/Runtime/CoreUObject/Internal/UObject/InstanceDataObjectUtils.h
--- a/Engine/Source/Runtime/CoreUObject/Internal/UObject/InstanceDataObjectUtils.h
+++ b/Engine/Source/Runtime/CoreUObject/Internal/UObject/InstanceDataObjectUtils.h
@@ -31,7 +31,7 @@ UE_API bool IsInstanceDataObjectSupportEnabledForGC(const UClass* Class);
 /** Query if uninitialized alert UI is enabled **/
 UE_API bool IsUninitializedAlertUIEnabled();
 /** Query if InstanceDataObjects should be saved instead of instances. */
-UE_API bool IsInstanceDataObjectImpersonationEnabledOnSave();
+UE_API bool IsInstanceDataObjectImpersonationEnabledOnSave(const UObject* Object);
 /** Query if InstanceDataObjects should be constructed using an archetype chain. */
 UE_API bool IsInstanceDataObjectArchetypeChainEnabled();
 /** Query if InstanceDataObjects should impose placeholder objects for missing types. */
diff --git a/Engine/Source/Runtime/CoreUObject/Private/UObject/InstanceDataObjectUtils.cpp b/Engine/Source/Runtime/CoreUObject/Private/UObject/InstanceDataObjectUtils.cpp
--- a/Engine/Source/Runtime/CoreUObject/Private/UObject/InstanceDataObjectUtils.cpp
+++ b/Engine/Source/Runtime/CoreUObject/Private/UObject/InstanceDataObjectUtils.cpp
@@ -347,7 +347,6 @@ namespace UE
 		TEXT("Enables alert information for uninitalized properties. Requires IDO.Enable=true")
 	);
 
-	// TODO: re-enable this ASAP. This disables most IDO features but disabling was necessary to unblock those experiencing IDO bugs
 	bool bEnableIDOImpersonationOnSave = true;
 	FAutoConsoleVariableRef EnableIDOImpersonationOnSaveCVar(
 		TEXT("IDO.Impersonation.EnableOnSave"),
@@ -355,6 +354,13 @@ namespace UE
 		TEXT("When enabled, IDOs will be saved instead of instances. Disabling this will stop data retention on save.")
 	);
 
+	bool bSkipImpersonationWithoutUnknowns = true;
+	FAutoConsoleVariableRef SkipImpersonationWithoutUnknownsCVar(
+		TEXT("IDO.Impersonation.SkipWithoutUnknownProperties"),
+		bSkipImpersonationWithoutUnknowns,
+		TEXT("When enabled, don't impersonate if the IDO has no unknown properties.")
+	);
+
 	bool bEnableIDOsForBlueprintArchetypes = true;
 	FAutoConsoleVariableRef EnableIDOsForBlueprintArchetypesCVar(
 		TEXT("IDO.EnableBlueprintArchetypes"),
@@ -583,9 +589,19 @@ namespace UE
 		return bEnableUninitializedUI;
 	}
 
-	bool IsInstanceDataObjectImpersonationEnabledOnSave()
+	bool IsInstanceDataObjectImpersonationEnabledOnSave(const UObject* Object)
 	{
-		return bEnableIDOImpersonationOnSave;
+		if (bEnableIDOImpersonationOnSave)
+		{
+			if (bSkipImpersonationWithoutUnknowns)
+			{
+				UObject* Ido = FPropertyBagRepository::Get().FindInstanceDataObject(Object);
+				// only impersonate if the object containse unknown (loose) properties
+				return Ido && StructContainsLooseProperties(Ido->GetClass());
+			}
+			return true;
+		}
+		return false;
 	}
 
 	bool IsInstanceDataObjectArchetypeChainEnabled()
diff --git a/Engine/Source/Runtime/CoreUObject/Private/UObject/UObjectThreadContext.cpp b/Engine/Source/Runtime/CoreUObject/Private/UObject/UObjectThreadContext.cpp
--- a/Engine/Source/Runtime/CoreUObject/Private/UObject/UObjectThreadContext.cpp
+++ b/Engine/Source/Runtime/CoreUObject/Private/UObject/UObjectThreadContext.cpp
@@ -161,7 +161,7 @@ FScopedObjectSerializeContext::FScopedObjectSerializeContext(UObject* InObject,
 	const bool bSupportsIDO = !Archive.IsCooking() && IsInstanceDataObjectSupportEnabled(Object);
 
 	// determine whether impersonation should be enabled when saving
-	const bool bImpersonateOnSave = bSupportsIDO && UE::IsInstanceDataObjectImpersonationEnabledOnSave();
+	const bool bImpersonateOnSave = bSupportsIDO && UE::IsInstanceDataObjectImpersonationEnabledOnSave(InObject);
 
 	// Disable if a newer version of the class exists because there is no point in creating an IDO of an obsolete type.
 	bCreateInstanceDataObject = bIsLoading && !bIsIDO && bSupportsIDO && !Object->GetClass()->HasAnyClassFlags(CLASS_NewerVersionExists);
