Call TObjectPtr::IsRemote() instead of IsResolved() when used in the context of remote objects

diff --git a/Engine/Source/Runtime/CoreUObject/Private/UObject/PropertyBaseObject.cpp b/Engine/Source/Runtime/CoreUObject/Private/UObject/PropertyBaseObject.cpp
--- a/Engine/Source/Runtime/CoreUObject/Private/UObject/PropertyBaseObject.cpp
+++ b/Engine/Source/Runtime/CoreUObject/Private/UObject/PropertyBaseObject.cpp
@@ -877,7 +877,7 @@ void FObjectPropertyBase::CheckValidObject(void* ValueAddress, TObjectPtr<UObjec
 #if UE_WITH_REMOTE_OBJECT_HANDLE
 	// Verifying the object class will result in resolving remote objects and we don't want that.
 	// However, since remote objects are guarnteed to have compatible binaries it's not required to validate objects.
-	if (!Object.IsResolved())
+	if (Object.IsRemote())
 	{
 		return;
 	}
diff --git a/Engine/Source/Runtime/CoreUObject/Private/UObject/RemoteObjectSerialization.cpp b/Engine/Source/Runtime/CoreUObject/Private/UObject/RemoteObjectSerialization.cpp
--- a/Engine/Source/Runtime/CoreUObject/Private/UObject/RemoteObjectSerialization.cpp
+++ b/Engine/Source/Runtime/CoreUObject/Private/UObject/RemoteObjectSerialization.cpp
@@ -814,7 +814,7 @@ public:
 	{
 		FObjectPtr Value;
 		ReadObjectReference(Value);
-		check(Value.IsResolved());
+		check(!Value.IsRemote());
 		Obj = Value.Get();
 		return *this;
 	}
@@ -841,7 +841,7 @@ public:
 		FObjectPtr Ptr;
 		ReadObjectReference(Ptr);
 
-		if (Ptr.IsResolved())
+		if (!Ptr.IsRemote())
 		{
 			Value = Ptr.Get();
 		}
@@ -895,7 +895,7 @@ public:
 
 	virtual FArchive& operator<<(FObjectPtr& Value) override
 	{
-		if (!Value.IsResolved())
+		if (Value.IsRemote())
 		{
 			PatchObjectReference(Value);
 		}
@@ -961,7 +961,7 @@ UObject* ConstructRemoteObject(const FRemoteObjectHeader& Header, FRemoteObjectC
 	UClass* Class = CastChecked<UClass>(Header.Class.Get());
 	UObject* Outer = nullptr;
 #if UE_WITH_REMOTE_OBJECT_HANDLE
-	if (Header.Outer.IsResolved())
+	if (!Header.Outer.IsRemote())
 	{
 		Outer = Header.Outer.Get();
 	}
diff --git a/Engine/Source/Runtime/Net/Core/Private/Net/Core/PropertyConditions/PropertyConditions.cpp b/Engine/Source/Runtime/Net/Core/Private/Net/Core/PropertyConditions/PropertyConditions.cpp
--- a/Engine/Source/Runtime/Net/Core/Private/Net/Core/PropertyConditions/PropertyConditions.cpp
+++ b/Engine/Source/Runtime/Net/Core/Private/Net/Core/PropertyConditions/PropertyConditions.cpp
@@ -126,7 +126,7 @@ void FNetPropertyConditionManager::PostGarbageCollect()
 	{
 		if (
 #if UE_WITH_REMOTE_OBJECT_HANDLE
-			!FObjectPtr(It.Key().GetRemoteId()).IsResolved() ||
+			FObjectPtr(It.Key().GetRemoteId()).IsRemote() ||
 #endif
 			!It.Key().ResolveObjectPtr())
 		{
