[CoreUObject]
* Fixed crash in reflection initialization if callstack was inside loading default properties and a new dll was loaded with new classes. The new types end up in pending registrations but were not explicitly initialized. So by random you could end up with a new type being initialized and having a TSubclassOf<OtherNewType> where OtherNewType was not initialized. Note, this scenario should be really rare, only happens once when loading LyraEditor.

diff --git a/Engine/Source/Runtime/CoreUObject/Private/UObject/UObjectBase.cpp b/Engine/Source/Runtime/CoreUObject/Private/UObject/UObjectBase.cpp
--- a/Engine/Source/Runtime/CoreUObject/Private/UObject/UObjectBase.cpp
+++ b/Engine/Source/Runtime/CoreUObject/Private/UObject/UObjectBase.cpp
@@ -1045,6 +1045,15 @@ void ProcessNewlyLoadedUObjects(FName InModuleName, bool bCanProcessNewlyLoadedO
 		StructRegistry.HasPendingRegistrations() ||
 		EnumRegistry.HasPendingRegistrations())
 	{
+		if (bNewUObjects)
+		{
+			// We are past first iteration and there are new pending registrations. We need to call this to make sure new classes are explicitly initialized
+			// since handling default properties might have references to the new classes.
+			// Example: A TSubclassOf<Type> in a new type where Type also is new. In this loop the default property of that TSubclassOf will collect
+			// the metaclass but the metaclass might not have been initialized without this call
+			UClassRegisterAllCompiledInClasses();
+		}
+
 		bNewUObjects = true;
 		UObjectProcessRegistrants();
 		UObjectLoadAllCompiledInStructs();
