Fixed composite EPackageFlags not being shown in the debugger - caused by non-power-of-two enumerators.

diff --git a/Engine/Source/Runtime/CoreUObject/Public/UObject/ObjectMacros.h b/Engine/Source/Runtime/CoreUObject/Public/UObject/ObjectMacros.h
--- a/Engine/Source/Runtime/CoreUObject/Public/UObject/ObjectMacros.h
+++ b/Engine/Source/Runtime/CoreUObject/Public/UObject/ObjectMacros.h
@@ -148,11 +148,11 @@ enum EPackageFlags
 	PKG_RuntimeGenerated			= 0x20000000,	///< This package contains elements that are runtime generated, and may not follow standard loading order rules
 	PKG_ReloadingForCooker			= 0x40000000,   ///< This package is reloading in the cooker, try to avoid getting data we will never need. We won't save this package.
 	PKG_FilterEditorOnly			= 0x80000000,	///< Package has editor-only data filtered out
-
-	// Transient Flags are cleared when serializing to or from PackageFileSummary
-	PKG_TransientFlags				= PKG_NewlyCreated | PKG_IsSaving | PKG_ReloadingForCooker,
 };
 
+// Transient Flags are cleared when serializing to or from PackageFileSummary
+#define PKG_TransientFlags	(PKG_NewlyCreated | PKG_IsSaving | PKG_ReloadingForCooker)
+
 /** Flag mask that indicates if this package is a package that exists in memory only. */
 #define PKG_InMemoryOnly	(EPackageFlags)(PKG_CompiledIn | PKG_NewlyCreated)
 
