[PackageLocalizationManager]
- Fix thread-safety issue when performing lazy initialization

diff --git a/Engine/Source/Runtime/CoreUObject/Private/Internationalization/PackageLocalizationManager.cpp b/Engine/Source/Runtime/CoreUObject/Private/Internationalization/PackageLocalizationManager.cpp
--- a/Engine/Source/Runtime/CoreUObject/Private/Internationalization/PackageLocalizationManager.cpp
+++ b/Engine/Source/Runtime/CoreUObject/Private/Internationalization/PackageLocalizationManager.cpp
@@ -79,8 +79,14 @@ void FPackageLocalizationManager::PerformLazyInitialization()
 		return;
 	}
 
-	if (!ActiveCache.IsValid() && LazyInitFunc)
+	if (!bIsInitialized.load(std::memory_order_acquire) && LazyInitFunc)
 	{
+		UE::TUniqueLock ScopeLock(InitializationMutex);
+		if (ActiveCache.IsValid())
+		{
+			return;
+		}
+
 		LazyInitFunc(*this);
 
 		if (!ActiveCache.IsValid())
@@ -94,6 +100,7 @@ void FPackageLocalizationManager::InitializeFromLazyCallback(FLazyInitFunc InLaz
 {
 	LazyInitFunc = MoveTemp(InLazyInitFunc);
 	ActiveCache.Reset();
+	bIsInitialized = false;
 }
 
 void FPackageLocalizationManager::InitializeFromCache(const TSharedRef<IPackageLocalizationCache>& InCache)
@@ -107,6 +114,8 @@ void FPackageLocalizationManager::InitializeFromCache(const TSharedRef<IPackageL
 	{
 		ActiveCache->ConditionalUpdateCache();
 	}
+
+	bIsInitialized.store(true, std::memory_order_release);
 }
 
 void FPackageLocalizationManager::InitializeFromDefaultCache()
diff --git a/Engine/Source/Runtime/CoreUObject/Public/Internationalization/PackageLocalizationManager.h b/Engine/Source/Runtime/CoreUObject/Public/Internationalization/PackageLocalizationManager.h
--- a/Engine/Source/Runtime/CoreUObject/Public/Internationalization/PackageLocalizationManager.h
+++ b/Engine/Source/Runtime/CoreUObject/Public/Internationalization/PackageLocalizationManager.h
@@ -2,6 +2,7 @@
 
 #pragma once
 
+#include "Async/Mutex.h"
 #include "Containers/UnrealString.h"
 #include "CoreMinimal.h"
 #include "Misc/CoreDelegates.h"
@@ -9,6 +10,8 @@
 #include "Templates/SharedPointer.h"
 #include "UObject/NameTypes.h"
 
+#include <atomic>
+
 class IPackageLocalizationCache;
 
 /** Singleton class that manages localized package data. */
@@ -94,4 +97,10 @@ private:
 
 	/** Pointer to our currently active cache. Only valid after Initialize has been called. */
 	TSharedPtr<IPackageLocalizationCache> ActiveCache;
+
+	/** The ActiveCache pointer is not atomic, so use a real atomic test for initialization */
+	std::atomic<bool> bIsInitialized = false;
+
+	/** To make sure initialization is thread-safe. */
+	UE::FMutex InitializationMutex;
 };
