- Changed UnversionPropertySerialization to use FConfigCacheIni::ForPlatform() instead of LoadlocalIniFile (so it will get the HF'd values if applied)

diff --git a/Engine/Source/Runtime/CoreUObject/Private/Serialization/UnversionedPropertySerialization.cpp b/Engine/Source/Runtime/CoreUObject/Private/Serialization/UnversionedPropertySerialization.cpp
--- a/Engine/Source/Runtime/CoreUObject/Private/Serialization/UnversionedPropertySerialization.cpp
+++ b/Engine/Source/Runtime/CoreUObject/Private/Serialization/UnversionedPropertySerialization.cpp
@@ -890,13 +890,13 @@ static bool CanUseUnversionedPropertySerialization()
 	return bAllow;
 }
 
-static bool CanUseUnversionedPropertySerialization(FConfigFile& TargetIni)
+static bool CanUseUnversionedPropertySerialization(const FConfigFile& TargetIni)
 {
 	bool bAllow;
 	return TargetIni.GetBool(TEXT("Core.System"), TEXT("CanUseUnversionedPropertySerialization"), /* out */ bAllow) && bAllow;
 }
 
-static bool CanUseUnversionedPropertySerializationForServerOnly(FConfigFile& TargetIni)
+static bool CanUseUnversionedPropertySerializationForServerOnly(const FConfigFile& TargetIni)
 {
 	bool bAllow;
 	return TargetIni.GetBool(TEXT("Core.System"), TEXT("CanUseUnversionedPropertySerializationForServerOnly"), /* out */ bAllow) && bAllow;
@@ -930,8 +930,7 @@ bool CanUseUnversionedPropertySerialization(const ITargetPlatform* Target)
 		return *CachedValue;
 	}
 
-	FConfigFile TargetIni;
-	FConfigCacheIni::LoadLocalIniFile(TargetIni, TEXT("Engine"), /* base INI */ true, *Target->IniPlatformName());
+	const FConfigFile& TargetIni = *FConfigCacheIni::ForPlatform(*Target->IniPlatformName())->FindConfigFile(GEngineIni);
 	bool bTargetValue = CanUseUnversionedPropertySerialization(TargetIni);
 	if (bIsServerOnly)
 	{
