Validate LinkerSave unmapped object against the package harvester excluded object list

diff --git a/Engine/Source/Runtime/CoreUObject/Private/UObject/LinkerSave.cpp b/Engine/Source/Runtime/CoreUObject/Private/UObject/LinkerSave.cpp
--- a/Engine/Source/Runtime/CoreUObject/Private/UObject/LinkerSave.cpp
+++ b/Engine/Source/Runtime/CoreUObject/Private/UObject/LinkerSave.cpp
@@ -15,6 +15,7 @@
 #include "UObject/LazyObjectPtr.h"
 #include "UObject/SoftObjectPtr.h"
 #include "Internationalization/TextPackageNamespaceUtil.h"
+#include "UObject/SavePackage/SaveContext.h"
 #include "UObject/ObjectSaveContext.h"
 #include "UObject/UObjectThreadContext.h"
 #include "UObject/UnrealType.h"
@@ -170,7 +171,7 @@ int32 FLinkerSave::MapSoftObjectPath(const FSoftObjectPath& SoftObjectPath) cons
 }
 
 
-FPackageIndex FLinkerSave::MapObject(TObjectPtr<const UObject> Object) const
+FPackageIndex FLinkerSave::MapObject(TObjectPtr<const UObject> Object, bool bValidateExcluded) const
 {
 	if (Object)
 	{
@@ -227,6 +228,15 @@ FPackageIndex FLinkerSave::MapObject(TObjectPtr<const UObject> Object) const
 
 			return *Found;
 		}
+		else if (bValidateExcluded && SaveContext)
+		{
+			UE_CLOG(SaveContext->FindCachedObjectStatus(Object) == nullptr, LogLinker, Error,
+					TEXT("Invalid serialization during save of Export '%s' in package '%s'. Object '%s' encountered during serialization without being previously harvested. This will cause the reference to be null at runtime. "),
+					(CurrentlySavingExport.IsExport() ? *Exp(CurrentlySavingExport).ObjectName.ToString() : TEXT("Unknown")),
+					*LinkerRoot->GetPathName(),
+					*Object->GetPathName()
+				);
+		}
 	}
 	return FPackageIndex();
 }
diff --git a/Engine/Source/Runtime/CoreUObject/Private/UObject/SavePackage/SaveContext.h b/Engine/Source/Runtime/CoreUObject/Private/UObject/SavePackage/SaveContext.h
--- a/Engine/Source/Runtime/CoreUObject/Private/UObject/SavePackage/SaveContext.h
+++ b/Engine/Source/Runtime/CoreUObject/Private/UObject/SavePackage/SaveContext.h
@@ -939,6 +939,16 @@ public:
 		return *ObjectStatus;
 	}
 
+	UE::SavePackageUtilities::FObjectStatus* FindCachedObjectStatus(TObjectPtr<const UObject> InObject) const
+	{
+		using namespace UE::SavePackageUtilities;
+		if (const TUniquePtr<FObjectStatus>* ObjectStatus = ObjectStatusCache.Find(InObject))
+		{
+			return ObjectStatus->Get();
+		}
+		return nullptr;
+	}
+
 	void RecordIllegalReference(UObject* InFrom, UObject* InTo, EIllegalRefReason InReason, FString&& InOptionalReasonText = FString())
 	{
 		HarvestedIllegalReferences.Add({ InFrom, InTo, InReason, MoveTemp(InOptionalReasonText) });
diff --git a/Engine/Source/Runtime/CoreUObject/Private/UObject/SavePackage2.cpp b/Engine/Source/Runtime/CoreUObject/Private/UObject/SavePackage2.cpp
--- a/Engine/Source/Runtime/CoreUObject/Private/UObject/SavePackage2.cpp
+++ b/Engine/Source/Runtime/CoreUObject/Private/UObject/SavePackage2.cpp
@@ -1096,6 +1096,7 @@ ESavePackageResult CreateLinker(FSaveContext& SaveContext)
 	{
 		LinkerSave->SetOutputDevice(SaveContext.GetError());
 	}
+	LinkerSave->AssignSaveContext(&SaveContext);
 	LinkerSave->SetTransientPropertyOverrides(SaveContext.GetTransientPropertyOverrides());
 	LinkerSave->bUpdatingLoadedPath = SaveContext.IsUpdatingLoadedPath();
 	LinkerSave->bProceduralSave = SaveContext.IsProceduralSave();
@@ -1738,7 +1739,7 @@ void SavePreloadDependencies(FStructuredArchive::FRecord& StructuredArchiveRoot,
 			{
 				return; // We assume nothing in coreuobject ever loads assets in a constructor
 			}
-			FPackageIndex Index = Linker->MapObject(ToTest);
+			FPackageIndex Index = Linker->MapObject(ToTest, false);
 			if (Index.IsNull() && bOnlyIfInLinkerTable)
 			{
 				return;
diff --git a/Engine/Source/Runtime/CoreUObject/Public/UObject/LinkerSave.h b/Engine/Source/Runtime/CoreUObject/Public/UObject/LinkerSave.h
--- a/Engine/Source/Runtime/CoreUObject/Public/UObject/LinkerSave.h
+++ b/Engine/Source/Runtime/CoreUObject/Public/UObject/LinkerSave.h
@@ -31,6 +31,7 @@ class FOutputDevice;
 class FPackagePath;
 class UObject;
 class UPackage;
+class FSaveContext;
 namespace UE { class FDerivedData; }
 namespace UE { class FPackageTrailerBuilder; }
 struct FLazyObjectPtr;
@@ -78,9 +79,6 @@ public:
 	/** Flag that indicate if we are currently serializing the package header soft object paths. Used to disable mapping while serializing the header itself. */
 	bool bIsWritingHeaderSoftObjectPaths = false;
 
-	/** Save context associated with this linker */
-	TRefCountPtr<FUObjectSerializeContext> SaveContext;
-
 	TArray<FFileRegion> FileRegions;
 
 	/**
@@ -180,6 +178,14 @@ public:
 	 */
 	bool TryAssignFileSaver(FStringView InFilename, bool bForceByteSwapping, bool bInSaveUnversioned = false);
 
+	/**
+	 * Assign the temporary save context with the linker
+	 */
+	void AssignSaveContext(FSaveContext* InContext)
+	{
+		SaveContext = InContext;
+	}
+
 	/** Returns the appropriate name index for the source name, or INDEX_NONE if not found in NameIndices */
 	int32 MapName( FNameEntryId Name) const;
 
@@ -187,7 +193,7 @@ public:
 	int32 MapSoftObjectPath(const FSoftObjectPath& SoftObjectPath) const;
 
 	/** Returns the appropriate package index for the source object, or default value if not found in ObjectIndicesMap */
-	FPackageIndex MapObject(TObjectPtr<const UObject> Object) const;
+	FPackageIndex MapObject(TObjectPtr<const UObject> Object, bool bValidateExcluded = true) const;
 
 	/**
 	 * Called when an object begins serializing property data using script serialization.
@@ -360,6 +366,8 @@ protected:
 	void AssignSaverInternal(FArchive* InSaver, bool bForceByteSwapping, bool bInSaveUnversioned);
 
 private:
+	/** temporary save context used to generate this linker. */
+	FSaveContext* SaveContext = nullptr;
 	/** Optional log output to bubble errors back up. */
 	FOutputDevice* LogOutput = nullptr;
 
