[Enhanced Input] Do not generate the generic BP interface " (Message)" versions of the EI subsystem interface because they did not actually work. If you do actually make a custom object which implements the enhanced input interface, then it is already a subsystem, and would thus generate its own "K2Node_GetSubsystem" node actions anyway.

To do this, add a piece of metadata, CannotGenerateMessageNodes. If you add :
CannotGenerateMessageNodes=true

to your UInterface class metadata, then the message node spawners will not be added.

[RN] The "Message" versions of the Enhanced Input subsystem will no longer be displayed in the blueprint context menu.

diff --git a/Engine/Plugins/EnhancedInput/Source/EnhancedInput/Public/EnhancedInputSubsystemInterface.h b/Engine/Plugins/EnhancedInput/Source/EnhancedInput/Public/EnhancedInputSubsystemInterface.h
--- a/Engine/Plugins/EnhancedInput/Source/EnhancedInput/Public/EnhancedInputSubsystemInterface.h
+++ b/Engine/Plugins/EnhancedInput/Source/EnhancedInput/Public/EnhancedInputSubsystemInterface.h
@@ -25,7 +25,7 @@ class UEnhancedInputUserSettings;
 class UEnhancedPlayerMappableKeyProfile;
 
 // Subsystem interface
-UINTERFACE(MinimalAPI, meta=(CannotImplementInterfaceInBlueprint))
+UINTERFACE(MinimalAPI, meta=(CannotImplementInterfaceInBlueprint, CannotGenerateMessageNodes=true))
 class UEnhancedInputSubsystemInterface : public UInterface
 {
 	GENERATED_BODY()
diff --git a/Engine/Source/Editor/BlueprintGraph/Private/BlueprintActionDatabase.cpp b/Engine/Source/Editor/BlueprintGraph/Private/BlueprintActionDatabase.cpp
--- a/Engine/Source/Editor/BlueprintGraph/Private/BlueprintActionDatabase.cpp
+++ b/Engine/Source/Editor/BlueprintGraph/Private/BlueprintActionDatabase.cpp
@@ -708,7 +708,9 @@ static void BlueprintActionDatabaseImpl::AddClassFunctionActions(UClass const* c
 				ActionListOut.Add(FuncSpawner);
 			}
 
-			if (FKismetEditorUtilities::IsClassABlueprintInterface(Class))
+			// If an interface class has specified "CannotGenerateMessageNodes=true" in its metadata, then don't create any message nodes for
+			// this interface class.
+			if (FKismetEditorUtilities::IsClassABlueprintInterface(Class) && !Class->GetBoolMetaData(TEXT("CannotGenerateMessageNodes")))
 			{
 				// Use the default function name
 				ActionListOut.Add(MakeMessageNodeSpawner(Function));
diff --git a/Engine/Source/Runtime/CoreUObject/Public/UObject/ObjectMacros.h b/Engine/Source/Runtime/CoreUObject/Public/UObject/ObjectMacros.h
--- a/Engine/Source/Runtime/CoreUObject/Public/UObject/ObjectMacros.h
+++ b/Engine/Source/Runtime/CoreUObject/Public/UObject/ObjectMacros.h
@@ -1769,6 +1769,19 @@ namespace UM
 	{
 		/// [InterfaceMetadata] This interface cannot be implemented by a blueprint (e.g., it has only non-exposed C++ member methods), equivalent to NotBlueprintable
 		CannotImplementInterfaceInBlueprint,
+
+		/// [InterfaceMetadata] If this metadata is set to true, then the interface cannot generate generic "Message"
+		/// nodes which are used for type-agnostic message passing.
+		/// 
+		/// This is useful if you want to clean up the blueprint action context menu and force users to use a different method of calling
+		/// into your interface.
+		///
+		/// Example:
+		///		UINTERFACE(MinimalAPI, meta=(CannotGenerateMessageNodes=true))
+		///		class UFooBarSubsystemInterface : public UInterface { GENERATED_BODY() };
+		///
+		///	This metadata is read in the BlueprintActionDatabaseImpl::AddClassFunctionActions function
+		CannotGenerateMessageNodes
 	};
 }
 
