Reduce log severity for new Save serialization mapping error until issue reporting is investigated

diff --git a/Engine/Source/Runtime/CoreUObject/Private/UObject/LinkerSave.cpp b/Engine/Source/Runtime/CoreUObject/Private/UObject/LinkerSave.cpp
--- a/Engine/Source/Runtime/CoreUObject/Private/UObject/LinkerSave.cpp
+++ b/Engine/Source/Runtime/CoreUObject/Private/UObject/LinkerSave.cpp
@@ -230,7 +230,7 @@ FPackageIndex FLinkerSave::MapObject(TObjectPtr<const UObject> Object, bool bVal
 		}
 		else if (bValidateExcluded && SaveContext)
 		{
-			UE_CLOG(SaveContext->FindCachedObjectStatus(Object) == nullptr, LogLinker, Error,
+			UE_CLOG(SaveContext->FindCachedObjectStatus(Object) == nullptr, LogLinker, Display,
 					TEXT("Invalid serialization during save of Export '%s' in package '%s'. Object '%s' encountered during serialization without being previously harvested. This will cause the reference to be null at runtime. "),
 					(CurrentlySavingExport.IsExport() ? *Exp(CurrentlySavingExport).ObjectName.ToString() : TEXT("Unknown")),
 					*LinkerRoot->GetPathName(),
