Unshelved from pending changelist '44100303':

Optimized TCppStructOps vtables to only include functions that will be called.

diff --git a/Engine/Plugins/Runtime/RigVM/Source/RigVMDeveloper/Public/RigVMModel/RigVMControllerActions.h b/Engine/Plugins/Runtime/RigVM/Source/RigVMDeveloper/Public/RigVMModel/RigVMControllerActions.h
--- a/Engine/Plugins/Runtime/RigVM/Source/RigVMDeveloper/Public/RigVMModel/RigVMControllerActions.h
+++ b/Engine/Plugins/Runtime/RigVM/Source/RigVMDeveloper/Public/RigVMModel/RigVMControllerActions.h
@@ -190,8 +190,6 @@ public:
 
 #endif
 
-protected:
-
 	// Empty constructor
 	FRigVMBaseAction()
 		: WeakController(nullptr)
@@ -199,6 +197,8 @@ protected:
 	{
 	}
 
+protected:
+
 	RIGVMDEVELOPER_API void EnsureControllerValidity() const;
 
 	mutable TWeakObjectPtr<URigVMController> WeakController;
@@ -217,7 +217,6 @@ protected:
 
 	friend struct FRigVMActionWrapper;
 	friend class URigVMActionStack;
-	friend class UScriptStruct;
 };
 
 /**
diff --git a/Engine/Source/Runtime/Core/Public/Misc/EnumClassFlags.h b/Engine/Source/Runtime/Core/Public/Misc/EnumClassFlags.h
--- a/Engine/Source/Runtime/Core/Public/Misc/EnumClassFlags.h
+++ b/Engine/Source/Runtime/Core/Public/Misc/EnumClassFlags.h
@@ -50,4 +50,40 @@ constexpr void EnumRemoveFlags(Enum& Flags, Enum FlagsToRemove)
 {
 	using UnderlyingType = __underlying_type(Enum);
 	Flags = (Enum)((UnderlyingType)Flags & ~(UnderlyingType)FlagsToRemove);
-}
\ No newline at end of file
+}
+
+template <typename Enum>
+constexpr Enum EnumLowestSetFlag(Enum Flags)
+{
+	using UnderlyingType = __underlying_type(Enum);
+#ifdef _MSC_VER
+	#pragma warning(push)
+	#pragma warning(disable : 4146) // unary minus operator applied to unsigned type, result still unsigned
+#endif
+	return Flags & (Enum)(-(UnderlyingType)Flags);
+#ifdef _MSC_VER
+	#pragma warning(pop)
+#endif
+}
+
+template <typename Enum>
+constexpr Enum EnumRemoveLowestSetFlag(Enum Flags)
+{
+	using UnderlyingType = __underlying_type(Enum);
+	return (Enum)((UnderlyingType)Flags & ((UnderlyingType)Flags - 1));
+}
+
+template <typename Enum>
+constexpr int EnumNumSetFlags(Enum Flags)
+{
+	using UnderlyingType = __underlying_type(Enum);
+
+	int Result = 0;
+	UnderlyingType Int = (UnderlyingType)Flags;
+	while (Int != 0)
+	{
+		++Result;
+		Int &= Int - 1;
+	}
+	return Result;
+}
diff --git a/Engine/Source/Runtime/CoreUObject/Public/UObject/Class.h b/Engine/Source/Runtime/CoreUObject/Public/UObject/Class.h
--- a/Engine/Source/Runtime/CoreUObject/Public/UObject/Class.h
+++ b/Engine/Source/Runtime/CoreUObject/Public/UObject/Class.h
@@ -49,6 +49,7 @@
 #include "Serialization/StructuredArchiveAdapters.h"
 #include "Templates/AlignmentTemplates.h"
 #include "Templates/EnableIf.h"
+#include "Templates/Identity.h"
 #include "Templates/IsAbstract.h"
 #include "Templates/IsEnum.h"
 #include "Templates/IsPODType.h"
@@ -91,8 +92,11 @@ class FEditPropertyChain;
 class FField;
 class FOutputDevice;
 class FProperty;
+class FReferenceCollector;
 class FStructProperty;
 class FStructuredArchiveRecord;
+class UObject;
+class UPackageMap;
 class UPropertyWrapper;
 struct CGetTypeHashable;
 struct FArchiveFieldName;
@@ -1014,6 +1018,757 @@ enum EStructFlags
 		}
 	}
 
+namespace UE::CoreUObject::Private
+{
+	/** Filled by implementation classes to report their capabilities */
+	struct FCapabilities
+	{
+		EPropertyFlags ComputedPropertyFlags;
+		EPropertyObjectReferenceType HasSerializerObjectReferences;
+		bool HasNoopConstructor : 1;
+		bool HasZeroConstructor : 1;
+		bool HasDestructor : 1;
+		bool HasSerializer : 1;
+		bool HasStructuredSerializer : 1;
+		bool HasPostSerialize : 1;
+		bool HasNetSerializer : 1;
+		bool HasNetSharedSerialization : 1;
+		bool HasNetDeltaSerializer : 1;
+		bool HasPostScriptConstruct : 1;
+		bool IsPlainOldData : 1;
+		bool IsUECoreType : 1;
+		bool IsUECoreVariant : 1;
+		bool HasCopy : 1;
+		bool HasIdentical : 1;
+		bool HasExportTextItem : 1;
+		bool HasImportTextItem : 1;
+		bool HasAddStructReferencedObjects : 1;
+		bool HasSerializeFromMismatchedTag : 1;
+		bool HasStructuredSerializeFromMismatchedTag : 1;
+		bool HasGetTypeHash : 1;
+		bool HasIntrusiveUnsetOptionalState : 1;
+		bool IsAbstract : 1;
+		bool HasFindInnerPropertyInstance : 1;
+		bool ClearOnFinishDestroy : 1;
+#if WITH_EDITOR
+		bool HasCanEditChange : 1;
+#endif
+		bool HasVisitor : 1;
+	};
+
+	template <typename CppStruct>
+	constexpr FCapabilities GetCapabilities()
+	{
+		using TraitsType = TStructOpsTypeTraits<CppStruct>;
+
+		constexpr FCapabilities Capabilities {
+			(TIsPODType<CppStruct>::Value ? CPF_IsPlainOldData : CPF_None)
+			| (std::is_trivially_destructible_v<CppStruct> ? CPF_NoDestructor : CPF_None)
+			| (TIsZeroConstructType<CppStruct>::Value ? CPF_ZeroConstructor : CPF_None)
+			| (TModels_V<CGetTypeHashable, CppStruct> ? CPF_HasGetValueTypeHash : CPF_None),
+			TraitsType::WithSerializerObjectReferences,
+			TraitsType::WithNoInitConstructor,
+			TraitsType::WithZeroConstructor,
+			!(TraitsType::WithNoDestructor || TIsPODType<CppStruct>::Value),
+			TraitsType::WithSerializer,
+			TraitsType::WithStructuredSerializer,
+			TraitsType::WithPostSerialize,
+			TraitsType::WithNetSerializer,
+			TraitsType::WithNetSharedSerialization,
+			TraitsType::WithNetDeltaSerializer,
+			TraitsType::WithPostScriptConstruct,
+			TIsPODType<CppStruct>::Value,
+			TIsUECoreType<CppStruct>::Value,
+			TIsUECoreVariant<CppStruct>::Value,
+			TraitsType::WithCopy,
+			TraitsType::WithIdentical || TraitsType::WithIdenticalViaEquality,
+			TraitsType::WithExportTextItem,
+			TraitsType::WithImportTextItem,
+			TraitsType::WithAddStructReferencedObjects,
+			TraitsType::WithSerializeFromMismatchedTag,
+			TraitsType::WithStructuredSerializeFromMismatchedTag,
+			TModels_V<CGetTypeHashable, CppStruct>,
+			::HasIntrusiveUnsetOptionalState<CppStruct>(),
+			TIsAbstract<CppStruct>::Value,
+			TraitsType::WithFindInnerPropertyInstance,
+			TraitsType::WithClearOnFinishDestroy,
+#if WITH_EDITOR
+			TraitsType::WithCanEditChange,
+#endif
+			TraitsType::WithVisitor,
+		};
+		return Capabilities;
+	}
+
+	template <typename CppStruct>
+	void Construct(void* Dest)
+	{
+		using TraitsType = TStructOpsTypeTraits<CppStruct>;
+
+		check(!TraitsType::WithZeroConstructor); // don't call this if we have indicated it is not necessary
+		// that could have been an if statement, but we might as well force optimization above the virtual call
+		// could also not attempt to call the constructor for types where this is not possible, but I didn't do that here
+#if CHECK_PUREVIRTUALS
+		if constexpr (!TraitsType::WithPureVirtual)
+#endif
+		{
+			if constexpr (TraitsType::WithNoInitConstructor)
+			{
+				::new (Dest) CppStruct(ForceInit);
+			}
+			else
+			{
+				::new (Dest) CppStruct();
+			}
+		}
+	}
+
+	template <typename CppStruct>
+	void ConstructForTests(void* Dest)
+	{
+		using TraitsType = TStructOpsTypeTraits<CppStruct>;
+
+		check(!TraitsType::WithZeroConstructor); // don't call this if we have indicated it is not necessary
+		// that could have been an if statement, but we might as well force optimization above the virtual call
+		// could also not attempt to call the constructor for types where this is not possible, but I didn't do that here
+#if CHECK_PUREVIRTUALS
+		if constexpr (!TraitsType::WithPureVirtual)
+#endif
+		{
+			if constexpr (TraitsType::WithNoInitConstructor)
+			{
+				::new (Dest) CppStruct(ForceInit);
+			}
+			else
+			{
+				::new (Dest) CppStruct;
+			}
+		}
+	}
+
+	template <typename CppStruct>
+	void Destruct(void* Dest)
+	{
+		using TraitsType = TStructOpsTypeTraits<CppStruct>;
+
+		check(!(TraitsType::WithNoDestructor || TIsPODType<CppStruct>::Value)); // don't call this if we have indicated it is not necessary
+		// that could have been an if statement, but we might as well force optimization above the virtual call
+		// could also not attempt to call the destructor for types where this is not possible, but I didn't do that here
+		DestructItem((CppStruct*)Dest);
+	}
+
+	template <typename CppStruct>
+	bool SerializeArchive(FArchive& Ar, void* Data, UStruct* DefaultsStruct, const void* Defaults)
+	{
+		using TraitsType = TStructOpsTypeTraits<CppStruct>;
+
+		check(TraitsType::WithSerializer); // don't call this if we have indicated it is not necessary
+		if constexpr (TraitsType::WithSerializer)
+		{
+			if constexpr (TModels_V<CStructSerializableWithDefaults, CppStruct>)
+			{
+				return ((CppStruct*)Data)->Serialize(Ar, DefaultsStruct, Defaults);
+			}
+			else
+			{
+				return ((CppStruct*)Data)->Serialize(Ar);
+			}
+		}
+		else
+		{
+			return false;
+		}
+	}
+
+	template <typename CppStruct>
+	bool SerializeSlot(FStructuredArchive::FSlot Slot, void* Data, UStruct* DefaultsStruct, const void* Defaults)
+	{
+		using TraitsType = TStructOpsTypeTraits<CppStruct>;
+
+		check(TraitsType::WithStructuredSerializer); // don't call this if we have indicated it is not necessary
+		if constexpr (TraitsType::WithStructuredSerializer)
+		{
+			if constexpr (TModels_V<CStructStructuredSerializableWithDefaults, CppStruct>)
+			{
+				return ((CppStruct*)Data)->Serialize(Slot, DefaultsStruct, Defaults);
+			}
+			else
+			{
+				return ((CppStruct*)Data)->Serialize(Slot);
+			}
+		}
+		else
+		{
+			return false;
+		}
+	}
+
+	template <typename CppStruct>
+	void PostSerialize(const FArchive& Ar, void *Data)
+	{
+		using TraitsType = TStructOpsTypeTraits<CppStruct>;
+
+		check(TraitsType::WithPostSerialize); // don't call this if we have indicated it is not necessary
+		if constexpr (TraitsType::WithPostSerialize)
+		{
+			((CppStruct*)Data)->PostSerialize(Ar);
+		}
+	}
+
+	template <typename CppStruct>
+	bool NetSerialize(FArchive& Ar, UPackageMap* Map, bool& bOutSuccess, void* Data)
+	{
+		using TraitsType = TStructOpsTypeTraits<CppStruct>;
+
+		if constexpr (TraitsType::WithNetSerializer)
+		{
+			return ((CppStruct*)Data)->NetSerialize(Ar, Map, bOutSuccess);
+		}
+		else
+		{
+			return false;
+		}
+	}
+
+	template <typename CppStruct>
+	bool NetDeltaSerialize(FNetDeltaSerializeInfo& DeltaParms, void* Data)
+	{
+		using TraitsType = TStructOpsTypeTraits<CppStruct>;
+
+		if constexpr (TraitsType::WithNetDeltaSerializer)
+		{
+			return ((CppStruct*)Data)->NetDeltaSerialize(DeltaParms);
+		}
+		else
+		{
+			return false;
+		}
+	}
+
+	template <typename CppStruct>
+	void PostScriptConstruct(void* Data)
+	{
+		using TraitsType = TStructOpsTypeTraits<CppStruct>;
+
+		check(TraitsType::WithPostScriptConstruct); // don't call this if we have indicated it is not necessary
+		if constexpr (TraitsType::WithPostScriptConstruct)
+		{
+			((CppStruct*)Data)->PostScriptConstruct();
+		}
+	}
+
+	template <typename CppStruct>
+	void GetPreloadDependencies(void* Data, TArray<UObject*>& OutDeps)
+	{
+		using TraitsType = TStructOpsTypeTraits<CppStruct>;
+
+		if constexpr (TraitsType::WithGetPreloadDependencies)
+		{
+			((CppStruct*)Data)->GetPreloadDependencies(OutDeps);
+		}
+	}
+
+	template <typename CppStruct>
+	bool Copy(void* Dest, void const* Src, int32 ArrayDim)
+	{
+		using TraitsType = TStructOpsTypeTraits<CppStruct>;
+
+		if constexpr (TraitsType::WithCopy)
+		{
+			static_assert((!TIsPODType<CppStruct>::Value), "You probably don't want custom copy for a POD type.");
+
+			CppStruct* TypedDest = (CppStruct*)Dest;
+			const CppStruct* TypedSrc  = (const CppStruct*)Src;
+
+			for (; ArrayDim; --ArrayDim)
+			{
+				*TypedDest++ = *TypedSrc++;
+			}
+			return true;
+		}
+		else
+		{
+			return false;
+		}
+	}
+
+	template <typename CppStruct>
+	bool Identical(const void* A, const void* B, uint32 PortFlags, bool& bOutResult)
+	{
+		using TraitsType = TStructOpsTypeTraits<CppStruct>;
+
+		check(TraitsType::WithIdentical || TraitsType::WithIdenticalViaEquality); // don't call this if we have indicated it is not necessary
+		if constexpr (TraitsType::WithIdentical)
+		{
+			static_assert(!TraitsType::WithIdenticalViaEquality, "Should not have both WithIdenticalViaEquality and WithIdentical.");
+
+			bOutResult = ((const CppStruct*)A)->Identical((const CppStruct*)B, PortFlags);
+			return true;
+		}
+		else if constexpr (TraitsType::WithIdenticalViaEquality)
+		{
+			bOutResult = (*(const CppStruct*)A == *(const CppStruct*)B);
+			return true;
+		}
+		else
+		{
+			bOutResult = false;
+			return false;
+		}
+	}
+
+	template <typename CppStruct>
+	bool ExportTextItem(FString& ValueStr, const void* PropertyValue, const void* DefaultValue, UObject* Parent, int32 PortFlags, UObject* ExportRootScope)
+	{
+		using TraitsType = TStructOpsTypeTraits<CppStruct>;
+
+		check(TraitsType::WithExportTextItem); // don't call this if we have indicated it is not necessary
+		if constexpr (TraitsType::WithExportTextItem)
+		{
+			if (DefaultValue)
+			{
+				return ((const CppStruct*)PropertyValue)->ExportTextItem(ValueStr, *(const CppStruct*)DefaultValue, Parent, PortFlags, ExportRootScope);
+			}
+			else
+			{
+				TTypeCompatibleBytes<CppStruct> TmpDefaultValue;
+				FMemory::Memzero(TmpDefaultValue.GetTypedPtr(), sizeof(CppStruct));
+				if (!GetCapabilities<CppStruct>().HasZeroConstructor)
+				{
+					Construct<CppStruct>(TmpDefaultValue.GetTypedPtr());
+				}
+
+				const bool bResult = ((const CppStruct*)PropertyValue)->ExportTextItem(ValueStr, *(const CppStruct*)TmpDefaultValue.GetTypedPtr(), Parent, PortFlags, ExportRootScope);
+
+				if (GetCapabilities<CppStruct>().HasDestructor)
+				{
+					Destruct<CppStruct>(TmpDefaultValue.GetTypedPtr());
+				}
+
+				return bResult;
+			}
+		}
+		else
+		{
+			return false;
+		}
+	}
+
+	template <typename CppStruct>
+	bool ImportTextItem(const TCHAR*& Buffer, void* Data, int32 PortFlags, UObject* OwnerObject, FOutputDevice* ErrorText)
+	{
+		using TraitsType = TStructOpsTypeTraits<CppStruct>;
+
+		check(TraitsType::WithImportTextItem); // don't call this if we have indicated it is not necessary
+		if constexpr (TraitsType::WithImportTextItem)
+		{
+			return ((CppStruct*)Data)->ImportTextItem(Buffer, PortFlags, OwnerObject, ErrorText);
+		}
+		else
+		{
+			return false;
+		}
+	}
+
+	template <typename CppStruct>
+	bool FindInnerPropertyInstance(FName PropertyName, const void* Data, const FProperty*& OutProp, const void*& OutData)
+	{
+		using TraitsType = TStructOpsTypeTraits<CppStruct>;
+
+		check(TraitsType::WithFindInnerPropertyInstance); // don't call this if we have indicated it is not necessary
+		if constexpr (TraitsType::WithFindInnerPropertyInstance)
+		{
+			return ((CppStruct*)Data)->FindInnerPropertyInstance(PropertyName, OutProp, OutData);
+		}
+		else
+		{
+			return false;
+		}
+	}
+
+	template <typename CppStruct>
+	void (*AddStructReferencedObjects())(void* A, FReferenceCollector& Collector)
+	{
+		using TraitsType = TStructOpsTypeTraits<CppStruct>;
+
+		check(TraitsType::WithAddStructReferencedObjects); // don't call this if we have indicated it is not necessary
+		return &AddStructReferencedObjectsOrNot<CppStruct>;
+	}
+
+	template <typename CppStruct>
+	bool SerializeFromMismatchedTag(FPropertyTag const& Tag, FArchive& Ar, void* Data)
+	{
+		using TraitsType = TStructOpsTypeTraits<CppStruct>;
+
+		check(TraitsType::WithSerializeFromMismatchedTag); // don't call this if we have indicated it is not allowed
+		if constexpr (TraitsType::WithSerializeFromMismatchedTag)
+		{
+			if constexpr (TIsUECoreType<CppStruct>::Value)
+			{
+				// Custom version of SerializeFromMismatchedTag for core types, which don't have access to FPropertyTag.
+				FName StructName;
+				if (Tag.GetType().GetName() == NAME_StructProperty && Tag.GetType().GetParameterCount() >= 1)
+				{
+					StructName = Tag.GetType().GetParameterName(0);
+				}
+				return ((CppStruct*)Data)->SerializeFromMismatchedTag(StructName, Ar);
+			}
+			else
+			{
+				return ((CppStruct*)Data)->SerializeFromMismatchedTag(Tag, Ar);
+			}
+		}
+		else
+		{
+			return false;
+		}
+	}
+
+	template <typename CppStruct>
+	bool StructuredSerializeFromMismatchedTag(FPropertyTag const& Tag, FStructuredArchive::FSlot Slot, void* Data)
+	{
+		using TraitsType = TStructOpsTypeTraits<CppStruct>;
+
+		check(TraitsType::WithStructuredSerializeFromMismatchedTag); // don't call this if we have indicated it is not allowed
+		if constexpr (TraitsType::WithStructuredSerializeFromMismatchedTag)
+		{
+			if constexpr (TIsUECoreType<CppStruct>::Value)
+			{
+				// Custom version of SerializeFromMismatchedTag for core types, which don't understand FPropertyTag.
+				FName StructName;
+				if (Tag.GetType().GetName() == NAME_StructProperty && Tag.GetType().GetParameterCount() >= 1)
+				{
+					StructName = Tag.GetType().GetParameterName(0);
+				}
+				return ((CppStruct*)Data)->SerializeFromMismatchedTag(StructName, Slot);
+			}
+			else
+			{
+				return ((CppStruct*)Data)->SerializeFromMismatchedTag(Tag, Slot);
+			}
+		}
+		else
+		{
+			return false;
+		}
+	}
+
+	template <typename CppStruct>
+	uint32 GetStructTypeHash(const void* Src)
+	{
+		ensure(GetCapabilities<CppStruct>().HasGetTypeHash);
+
+		if constexpr (TModels_V<CGetTypeHashable, CppStruct>)
+		{
+			return GetTypeHashHelper(*(const CppStruct*)Src);
+		}
+		else
+		{
+			return 0;
+		}
+	}
+
+	template <typename CppStruct>
+	void InitializeIntrusiveUnsetOptionalValue(void* Data)
+	{
+		if constexpr(::HasIntrusiveUnsetOptionalState<CppStruct>())
+		{
+			new (Data) TOptional<CppStruct>();
+		}
+		else
+		{
+			checkf(false, TEXT("InitializeIntrusiveUnsetOptionalValue should not be called for types which do not have an intrusive unset optional state"));
+		}
+	}
+
+	template <typename CppStruct>
+	bool IsIntrusiveOptionalValueSet(const void* Data)
+	{
+		if constexpr(::HasIntrusiveUnsetOptionalState<CppStruct>())
+		{
+			return reinterpret_cast<const TOptional<CppStruct>*>(Data)->IsSet();
+		}
+		else 
+		{
+			checkf(false, TEXT("IsIntrusiveOptionalValueSet should not be called for types which do not have an intrusive unset optional state"));
+			return false;
+		}
+	}
+
+	template <typename CppStruct>
+	void ClearIntrusiveOptionalValue(void* Data)
+	{
+		if constexpr(::HasIntrusiveUnsetOptionalState<CppStruct>())
+		{
+			reinterpret_cast<TOptional<CppStruct>*>(Data)->Reset();
+		}
+		else
+		{
+			checkf(false, TEXT("ClearIntrusiveOptionalValue should not be called for types which do not have an intrusive unset optional state"));
+		}
+	}
+
+	template <typename CppStruct>
+	bool IsIntrusiveOptionalSafeForGC()
+	{
+		using TraitsType = TStructOpsTypeTraits<CppStruct>;
+
+		return TraitsType::WithIntrusiveOptionalSafeForGC;
+	}
+
+#if WITH_EDITOR
+	template <typename CppStruct>
+	bool CanEditChange(const FEditPropertyChain& PropertyChain, const void* Data)
+	{
+		using TraitsType = TStructOpsTypeTraits<CppStruct>;
+
+		if constexpr (TraitsType::WithCanEditChange)
+		{
+			return ((const CppStruct*)Data)->CanEditChange(PropertyChain);
+		}
+		else
+		{
+			return false;
+		}
+	}
+	#define UE_FOREACH_FAKE_VTABLE_FUNC_EDITOR_ONLY(...) __VA_ARGS__
+#else
+	#define UE_FOREACH_FAKE_VTABLE_FUNC_EDITOR_ONLY(...)
+#endif // WITH_EDITOR
+
+	template <typename CppStruct>
+	EPropertyVisitorControlFlow Visit(FPropertyVisitorContext& Context, const TFunctionRef<EPropertyVisitorControlFlow(const FPropertyVisitorContext& Context)> InFunc)
+	{
+		using TraitsType = TStructOpsTypeTraits<CppStruct>;
+
+		if constexpr (TraitsType::WithVisitor)
+		{
+			CppStruct* Struct = (CppStruct*)Context.Data.PropertyData;
+			return Struct->Visit(Context, InFunc);
+		}
+		else
+		{
+			return EPropertyVisitorControlFlow::StepOver;
+		}
+	}
+
+	template <typename CppStruct>
+	void* ResolveVisitedPathInfo(void* Data, const FPropertyVisitorInfo& Info)
+	{
+		using TraitsType = TStructOpsTypeTraits<CppStruct>;
+
+		if constexpr (TraitsType::WithVisitor)
+		{
+			return ((CppStruct*)Data)->ResolveVisitedPathInfo(Info);
+		}
+		else
+		{
+			return nullptr;
+		}
+	}
+
+	#define UE_FOREACH_FAKE_VTABLE_FUNC(IndexNameFunctypeCondition) \
+		IndexNameFunctypeCondition(0,	Construct,								void(void* Address),																																			!TStructOpsTypeTraits<CppStruct>::WithZeroConstructor) \
+		IndexNameFunctypeCondition(1,	ConstructForTests,						void(void* Address),																																			!TStructOpsTypeTraits<CppStruct>::WithZeroConstructor) \
+		IndexNameFunctypeCondition(2,	Destruct,								void(void* Address),																																			!TStructOpsTypeTraits<CppStruct>::WithNoDestructor && !TIsPODType<CppStruct>::Value) \
+		IndexNameFunctypeCondition(3,	SerializeArchive,						bool(FArchive& Ar, void* Data, UStruct* DefaultsStruct, const void* Defaults),																					TStructOpsTypeTraits<CppStruct>::WithSerializer) \
+		IndexNameFunctypeCondition(4,	SerializeSlot,							bool(FStructuredArchive::FSlot Slot, void* Data, UStruct* DefaultsStruct, const void* Defaults),																TStructOpsTypeTraits<CppStruct>::WithStructuredSerializer) \
+		IndexNameFunctypeCondition(5,	PostSerialize,							void(const FArchive& Ar, void* Data),																															TStructOpsTypeTraits<CppStruct>::WithPostSerialize) \
+		IndexNameFunctypeCondition(6,	NetSerialize,							bool(FArchive& Ar, UPackageMap* Map, bool& bOutSuccess, void* Data),																							true) \
+		IndexNameFunctypeCondition(7,	NetDeltaSerialize,						bool(FNetDeltaSerializeInfo& DeltaParms, void* Data),																											true) \
+		IndexNameFunctypeCondition(8,	PostScriptConstruct,					void(void* Data),																																				TStructOpsTypeTraits<CppStruct>::WithPostScriptConstruct) \
+		IndexNameFunctypeCondition(9,	GetPreloadDependencies,					void(void* Data, TArray<UObject*>& OutDeps),																													true) \
+		IndexNameFunctypeCondition(10,	Copy,									bool(void* Dest, void const* Src, int32 ArrayDim),																												true) \
+		IndexNameFunctypeCondition(11,	Identical,								bool(const void* A, const void* B, uint32 PortFlags, bool& bOutResult),																							TStructOpsTypeTraits<CppStruct>::WithIdentical || TStructOpsTypeTraits<CppStruct>::WithIdenticalViaEquality) \
+		IndexNameFunctypeCondition(12,	ExportTextItem,							bool(FString& ValueStr, const void* PropertyValue, const void* DefaultValue, UObject* Parent, int32 PortFlags, UObject* ExportRootScope),						TStructOpsTypeTraits<CppStruct>::WithExportTextItem) \
+		IndexNameFunctypeCondition(13,	ImportTextItem,							bool(const TCHAR*& Buffer, void* Data, int32 PortFlags, UObject* OwnerObject, FOutputDevice* ErrorText),														TStructOpsTypeTraits<CppStruct>::WithImportTextItem) \
+		IndexNameFunctypeCondition(14,	FindInnerPropertyInstance,				bool(FName PropertyName, const void* Data, const FProperty*& OutProp, const void*& OutData),																	TStructOpsTypeTraits<CppStruct>::WithFindInnerPropertyInstance) \
+		IndexNameFunctypeCondition(15,	AddStructReferencedObjects,				void(*())(void* A, FReferenceCollector& Collector),						 																						TStructOpsTypeTraits<CppStruct>::WithAddStructReferencedObjects) \
+		IndexNameFunctypeCondition(16,	SerializeFromMismatchedTag,				bool(FPropertyTag const& Tag, FArchive& Ar, void* Data),																										TStructOpsTypeTraits<CppStruct>::WithSerializeFromMismatchedTag) \
+		IndexNameFunctypeCondition(17,	StructuredSerializeFromMismatchedTag,	bool(FPropertyTag const& Tag, FStructuredArchive::FSlot Slot, void* Data),																						TStructOpsTypeTraits<CppStruct>::WithStructuredSerializeFromMismatchedTag) \
+		IndexNameFunctypeCondition(18,	GetStructTypeHash,						uint32(const void* Src),																																		true) \
+		IndexNameFunctypeCondition(19,	InitializeIntrusiveUnsetOptionalValue,	void(void* Data),																																				::HasIntrusiveUnsetOptionalState<CppStruct>()) \
+		IndexNameFunctypeCondition(20,	IsIntrusiveOptionalValueSet,			bool(const void* Data),																																			::HasIntrusiveUnsetOptionalState<CppStruct>()) \
+		IndexNameFunctypeCondition(21,	ClearIntrusiveOptionalValue,			void(void* Data),																																				::HasIntrusiveUnsetOptionalState<CppStruct>()) \
+		IndexNameFunctypeCondition(22,	IsIntrusiveOptionalSafeForGC,			bool(),																																							true) \
+		UE_FOREACH_FAKE_VTABLE_FUNC_EDITOR_ONLY( \
+		IndexNameFunctypeCondition(23,	CanEditChange,							bool(const FEditPropertyChain& PropertyChain, const void* Data),																								true) \
+		) \
+		IndexNameFunctypeCondition(24,	Visit,									EPropertyVisitorControlFlow(FPropertyVisitorContext& Context, const TFunctionRef<EPropertyVisitorControlFlow(const FPropertyVisitorContext& Context)> InFunc),	true) \
+		IndexNameFunctypeCondition(25,	ResolveVisitedPathInfo,					void*(void* Data, const FPropertyVisitorInfo& Info),																											true)
+
+	// These represent one function per optional function in the fake vtable.
+	enum class EStructOpsFakeVTableFlags : uint32
+	{
+		None = 0,
+
+		#define UE_VTABLE_ENUM(Index, Name, Functype, Condition) Name = 1u << Index,
+		UE_FOREACH_FAKE_VTABLE_FUNC(UE_VTABLE_ENUM)
+		#undef UE_VTABLE_ENUM
+	};
+	ENUM_CLASS_FLAGS(EStructOpsFakeVTableFlags)
+
+	template <EStructOpsFakeVTableFlags Flag>
+	struct THandlingFunctionType;
+	template <typename CppStruct, EStructOpsFakeVTableFlags Flag>
+	struct THandlingFunctionPtr;
+
+	#define UE_VTABLE_FUNC(Index, Name, Functype, Condition) \
+		template <> \
+		struct THandlingFunctionType<EStructOpsFakeVTableFlags::Name> \
+		{ \
+			using Type = Functype; \
+		}; \
+		template <typename CppStruct> \
+		struct THandlingFunctionPtr<CppStruct, EStructOpsFakeVTableFlags::Name> \
+		{ \
+			inline static constexpr TIdentity_T<Functype>* Value = &Name<CppStruct>; \
+		};
+	UE_FOREACH_FAKE_VTABLE_FUNC(UE_VTABLE_FUNC)
+	#undef UE_VTABLE_FUNC
+
+	template <typename CppStruct>
+	constexpr EStructOpsFakeVTableFlags GetStructOpsFakeVTableFlags()
+	{
+		using TraitsType = TStructOpsTypeTraits<CppStruct>;
+
+		EStructOpsFakeVTableFlags Result = EStructOpsFakeVTableFlags::None;
+
+		#define UE_VTABLE_SETFLAG(Index, Name, Functype, Condition) \
+			if constexpr (Condition) \
+			{ \
+				Result |= EStructOpsFakeVTableFlags::Name; \
+			}
+		UE_FOREACH_FAKE_VTABLE_FUNC(UE_VTABLE_SETFLAG)
+		#undef UE_VTABLE_SETFLAG
+
+		return Result;
+	}
+
+	struct FStructOpsFakeVTable
+	{
+		// These flags will indicate which functions are present in the HandlingFunctions array.
+		EStructOpsFakeVTableFlags Flags;
+
+		/** returns struct capabilities */
+		FCapabilities (*GetCapabilities)();
+
+// MSVC doesn't like zero-sized arrays in the middle of structs
+#ifndef _MSC_VER
+		// One function pointer will exist in this array for every corresponding flag in Flags.
+		// The function pointer types will be different, depending on the required signature, and
+		// will need to be cast to the right type before being called.
+		void (*HandlingFunctions[0])();
+#endif
+
+		// Retrieves a handling func associated with a particular flag.
+		template <EStructOpsFakeVTableFlags FuncFlag>
+		UE_REWRITE auto GetHandlingFunc() const
+		{
+			// Ensure a single flag is passed
+			static_assert(EnumNumSetFlags(FuncFlag) == 1);
+
+#ifdef _MSC_VER
+	#pragma warning(push)
+	#pragma warning(disable : 4191) // unsafe conversion from 'function pointer type A' to 'function pointer type B'
+#endif
+			return (typename THandlingFunctionType<FuncFlag>::Type*)this->GetHandlingFuncImpl(FuncFlag);
+#ifdef _MSC_VER
+	#pragma warning(pop)
+#endif
+		}
+
+	private:
+		void (*GetHandlingFuncImpl(EStructOpsFakeVTableFlags FuncFlag) const)()
+		{
+			// Ensure that the fake vtable actually contains an entry for the requested function
+			check(!!(this->Flags & FuncFlag));
+
+			// Mask off all higher flags
+			uint32 StructFlagsInt = (uint32)this->Flags;
+			StructFlagsInt &= (uint32)FuncFlag - 1;
+
+			// The number of remaining flags will be the index of the function into the function table.
+			// Hopefully this turns into a popcnt.
+			uint32 Index = 0;
+			while (StructFlagsInt != 0)
+			{
+				StructFlagsInt &= StructFlagsInt - 1;
+				++Index;
+			}
+
+#ifdef _MSC_VER
+			// Ensure that the end of the struct is where the start of the function pointer array in TStructOpsFakeVTable will be.
+			static_assert(alignof(FStructOpsFakeVTable) == alignof(void(*)()));
+			return ((void(**)())(this + 1))[Index];
+#else
+			return this->HandlingFunctions[Index];
+#endif
+		}
+	};
+
+	template <uint32 NumFlags>
+	struct TStructOpsFakeVTable
+	{
+		FStructOpsFakeVTable Header;
+		void (*HandlingFunctions[NumFlags])();
+	};
+
+	template <>
+	struct TStructOpsFakeVTable<0>
+	{
+		FStructOpsFakeVTable Header;
+	};
+
+	// This represents an array of function pointers to handling functions and a header describing the capabilities
+	// of the struct. Only one pointer exists per handling function. The index of the handling function is the index
+	// of the set bit in the capabilities.
+	//
+	// Example: a struct S containing HasDestructor and HasIdentical would be laid out like this in memory:
+	//   ComputedPropertyFlags
+    //   HasSerializerObjectReferences
+	//   ...:                         false
+	//   HasDestructor:               true
+	//   ...:                         false
+	//   HasIdentical:                true
+	//   ...:                         false
+	//   void(*HandlingFunctions[2]): { &Destruct<S>, &Identical<S> }
+	//
+	// We can't use constexpr because it uses reinterpret casts.
+	template <typename CppStruct, auto... Funcs>
+	inline const TStructOpsFakeVTable<sizeof...(Funcs)> TStructOpsFakeVTableInstance = 
+	{
+		{
+			GetStructOpsFakeVTableFlags<CppStruct>(),
+			&GetCapabilities<CppStruct>
+		},
+		{
+#ifdef _MSC_VER
+	#pragma warning(push)
+	#pragma warning(disable : 4191) // unsafe conversion from 'function pointer type A' to 'function pointer type B'
+#endif
+			(void(*)())Funcs...
+#ifdef _MSC_VER
+	#pragma warning(pop)
+#endif
+		}
+	};
+
+	template <typename CppStruct, EStructOpsFakeVTableFlags Flags, auto... Funcs>
+	struct TStructOpsFakeVTableForTypeImpl : TStructOpsFakeVTableForTypeImpl<CppStruct, EnumRemoveLowestSetFlag(Flags), Funcs..., THandlingFunctionPtr<CppStruct, EnumLowestSetFlag(Flags)>::Value>
+	{
+	};
+
+	template <typename CppStruct, auto... Funcs>
+	struct TStructOpsFakeVTableForTypeImpl<CppStruct, EStructOpsFakeVTableFlags::None, Funcs...>
+	{
+		inline static const FStructOpsFakeVTable* Value = &TStructOpsFakeVTableInstance<CppStruct, Funcs...>.Header;
+	};
+
+	template <typename CppStruct>
+	using TStructOpsFakeVTableForType = TStructOpsFakeVTableForTypeImpl<CppStruct, GetStructOpsFakeVTableFlags<CppStruct>()>;
+}
+
 /**
  * Reflection data for a standalone structure declared in a header or as a user defined struct
  */
@@ -1025,41 +1780,9 @@ public:
 	/** Interface to template to manage dynamic access to C++ struct construction and destruction **/
 	struct ICppStructOps
 	{
-		/** Filled by implementation classes to report their capabilities */
-		struct FCapabilities
-		{
-			EPropertyFlags ComputedPropertyFlags;
-			EPropertyObjectReferenceType HasSerializerObjectReferences;
-			bool HasNoopConstructor : 1;
-			bool HasZeroConstructor : 1;
-			bool HasDestructor : 1;
-			bool HasSerializer : 1;
-			bool HasStructuredSerializer : 1;
-			bool HasPostSerialize : 1;
-			bool HasNetSerializer : 1;
-			bool HasNetSharedSerialization : 1;
-			bool HasNetDeltaSerializer : 1;
-			bool HasPostScriptConstruct : 1;
-			bool IsPlainOldData : 1;
-			bool IsUECoreType : 1;
-			bool IsUECoreVariant : 1;
-			bool HasCopy : 1;
-			bool HasIdentical : 1;
-			bool HasExportTextItem : 1;
-			bool HasImportTextItem : 1;
-			bool HasAddStructReferencedObjects : 1;
-			bool HasSerializeFromMismatchedTag : 1;
-			bool HasStructuredSerializeFromMismatchedTag : 1;
-			bool HasGetTypeHash : 1;
-			bool HasIntrusiveUnsetOptionalState : 1;
-			bool IsAbstract : 1;
-			bool HasFindInnerPropertyInstance : 1;
-			bool ClearOnFinishDestroy : 1;
-#if WITH_EDITOR
-			bool HasCanEditChange : 1;
-#endif
-			bool HasVisitor : 1;
-		};
+		using FCapabilities = UE::CoreUObject::Private::FCapabilities;
+
+		const UE::CoreUObject::Private::FStructOpsFakeVTable* FakeVPtr;
 
 		/**
 		 * Constructor
@@ -1070,37 +1793,55 @@ public:
 			, Alignment(InAlignment)
 		{
 		}
-		virtual ~ICppStructOps() {}
 
 		/** returns struct capabilities */
-		virtual FCapabilities GetCapabilities() const = 0;
+		FCapabilities GetCapabilities() const
+		{
+			return FakeVPtr->GetCapabilities();
+		}
 
 		/** return true if this class has a no-op constructor and takes EForceInit to init **/
 		bool HasNoopConstructor() const
 		{
 			return GetCapabilities().HasNoopConstructor;
 		}
+
 		/** return true if memset can be used instead of the constructor **/
 		bool HasZeroConstructor() const
 		{
 			return GetCapabilities().HasZeroConstructor;
 		}
+
 		/** Call the C++ constructor **/
-		virtual void Construct(void *Dest) = 0;
+		void Construct(void *Dest)
+		{
+			FakeVPtr->GetHandlingFunc<UE::CoreUObject::Private::EStructOpsFakeVTableFlags::Construct>()(Dest);
+		}
+
 		/** Call the C++ constructor without value-init (new T instead of new T()) **/
-		virtual void ConstructForTests(void* Dest) = 0;
+		void ConstructForTests(void* Dest)
+		{
+			FakeVPtr->GetHandlingFunc<UE::CoreUObject::Private::EStructOpsFakeVTableFlags::ConstructForTests>()(Dest);
+		}
+
 		/** return false if this destructor can be skipped **/
 		bool HasDestructor() const
 		{
 			return GetCapabilities().HasDestructor;
 		}
+
 		/** Call the C++ destructor **/
-		virtual void Destruct(void *Dest) = 0;
+		void Destruct(void *Dest)
+		{
+			FakeVPtr->GetHandlingFunc<UE::CoreUObject::Private::EStructOpsFakeVTableFlags::Destruct>()(Dest);
+		}
+
 		/** return the sizeof() of this structure **/
 		UE_FORCEINLINE_HINT int32 GetSize() const
 		{
 			return Size;
 		}
+
 		/** return the alignof() of this structure **/
 		UE_FORCEINLINE_HINT int32 GetAlignment() const
 		{
@@ -1112,6 +1853,7 @@ public:
 		{
 			return GetCapabilities().HasSerializer;
 		}
+
 		/** return true if this class can serialize to a structured archive**/
 		bool HasStructuredSerializer() const
 		{
@@ -1119,13 +1861,13 @@ public:
 		}
 
 		UE_DEPRECATED(5.6, "Use the overload of Serialize with defaults.")
-		virtual bool Serialize(FArchive& Ar, void* Data)
+		bool Serialize(FArchive& Ar, void* Data)
 		{
 			return Serialize(Ar, Data, nullptr, nullptr);
 		}
 
 		UE_DEPRECATED(5.6, "Use the overload of Serialize with defaults.")
-		virtual bool Serialize(FStructuredArchive::FSlot Slot, void* Data)
+		bool Serialize(FStructuredArchive::FSlot Slot, void* Data)
 		{
 			return Serialize(Slot, Data, nullptr, nullptr);
 		}
@@ -1134,16 +1876,26 @@ public:
 		 * Serialize this structure.
 		 * @return true if serialization was handled, false to fall back to standard script struct serialization.
 		 */
-		virtual bool Serialize(FArchive& Ar, void* Data, UStruct* DefaultsStruct, const void* Defaults) = 0;
-		virtual bool Serialize(FStructuredArchive::FSlot Slot, void* Data, UStruct* DefaultsStruct, const void* Defaults) = 0;
+		bool Serialize(FArchive& Ar, void* Data, UStruct* DefaultsStruct, const void* Defaults)
+		{
+			return FakeVPtr->GetHandlingFunc<UE::CoreUObject::Private::EStructOpsFakeVTableFlags::SerializeArchive>()(Ar, Data, DefaultsStruct, Defaults);
+		}
+		bool Serialize(FStructuredArchive::FSlot Slot, void* Data, UStruct* DefaultsStruct, const void* Defaults)
+		{
+			return FakeVPtr->GetHandlingFunc<UE::CoreUObject::Private::EStructOpsFakeVTableFlags::SerializeSlot>()(Slot, Data, DefaultsStruct, Defaults);
+		}
 
 		/** return true if this class implements a post serialize call **/
 		bool HasPostSerialize() const
 		{
 			return GetCapabilities().HasPostSerialize;
 		}
+
 		/** Call PostSerialize on this structure */
-		virtual void PostSerialize(const FArchive& Ar, void *Data) = 0;
+		void PostSerialize(const FArchive& Ar, void *Data)
+		{
+			return FakeVPtr->GetHandlingFunc<UE::CoreUObject::Private::EStructOpsFakeVTableFlags::PostSerialize>()(Ar, Data);
+		}
 
 		/** return true if this struct can net serialize **/
 		bool HasNetSerializer() const
@@ -1156,33 +1908,48 @@ public:
 		{
 			return GetCapabilities().HasNetSharedSerialization;
 		}
+
 		/** 
 		 * Net serialize this structure 
 		 * @return true if the struct was serialized, otherwise it will fall back to ordinary script struct net serialization
 		 */
-		virtual bool NetSerialize(FArchive& Ar, class UPackageMap* Map, bool& bOutSuccess, void *Data) = 0;
+		bool NetSerialize(FArchive& Ar, UPackageMap* Map, bool& bOutSuccess, void *Data)
+		{
+			return FakeVPtr->GetHandlingFunc<UE::CoreUObject::Private::EStructOpsFakeVTableFlags::NetSerialize>()(Ar, Map, bOutSuccess, Data);
+		}
 
 		/** return true if this struct can net delta serialize delta (serialize a network delta from a base state) **/
 		bool HasNetDeltaSerializer() const
 		{
 			return GetCapabilities().HasNetDeltaSerializer;
 		}
+
 		/** 
 		 * Net serialize delta this structure. Serialize a network delta from a base state
 		 * @return true if the struct was serialized, otherwise it will fall back to ordinary script struct net delta serialization
 		 */
-		virtual bool NetDeltaSerialize(FNetDeltaSerializeInfo & DeltaParms, void *Data) = 0;
+		bool NetDeltaSerialize(FNetDeltaSerializeInfo & DeltaParms, void *Data)
+		{
+			return FakeVPtr->GetHandlingFunc<UE::CoreUObject::Private::EStructOpsFakeVTableFlags::NetDeltaSerialize>()(DeltaParms, Data);
+		}
 
 		/** return true if this class implements a post script construct call **/
 		bool HasPostScriptConstruct() const
 		{
 			return GetCapabilities().HasPostScriptConstruct;
 		}
+
 		/** Call PostScriptConstruct on this structure */
-		virtual void PostScriptConstruct(void *Data) = 0;
+		void PostScriptConstruct(void *Data)
+		{
+			FakeVPtr->GetHandlingFunc<UE::CoreUObject::Private::EStructOpsFakeVTableFlags::PostScriptConstruct>()(Data);
+		}
 
 		/** Call PreloadDependencies on this structure */
-		virtual void GetPreloadDependencies(void* Data, TArray<UObject*>& OutDeps) = 0;
+		void GetPreloadDependencies(void* Data, TArray<UObject*>& OutDeps)
+		{
+			FakeVPtr->GetHandlingFunc<UE::CoreUObject::Private::EStructOpsFakeVTableFlags::GetPreloadDependencies>()(Data, OutDeps);
+		}
 
 		/** return true if this struct should be memcopied **/
 		bool IsPlainOldData() const
@@ -1207,73 +1974,101 @@ public:
 		{
 			return GetCapabilities().HasCopy;
 		}
+
 		/** 
 		 * Copy this structure 
 		 * @return true if the copy was handled, otherwise it will fall back to CopySingleValue
 		 */
-		virtual bool Copy(void* Dest, void const* Src, int32 ArrayDim) = 0;
+		bool Copy(void* Dest, void const* Src, int32 ArrayDim)
+		{
+			return FakeVPtr->GetHandlingFunc<UE::CoreUObject::Private::EStructOpsFakeVTableFlags::Copy>()(Dest, Src, ArrayDim);
+		}
 
 		/** return true if this struct can compare **/
 		bool HasIdentical() const
 		{
 			return GetCapabilities().HasIdentical;
 		}
+
 		/** 
 		 * Compare this structure
 		 * @return true if the comparison was performed. bOutResult contains the result of the comparison.
 		 */
-		virtual bool Identical(const void* A, const void* B, uint32 PortFlags, bool& bOutResult) = 0;
+		bool Identical(const void* A, const void* B, uint32 PortFlags, bool& bOutResult)
+		{
+			return FakeVPtr->GetHandlingFunc<UE::CoreUObject::Private::EStructOpsFakeVTableFlags::Identical>()(A, B, PortFlags, bOutResult);
+		}
 
 		/** return true if this struct can export **/
 		bool HasExportTextItem() const
 		{
 			return GetCapabilities().HasExportTextItem;
 		}
+
 		/** 
 		 * export this structure 
 		 * @return true if the copy was exported, otherwise it will fall back to FStructProperty::ExportTextItem
 		 */
-		virtual bool ExportTextItem(FString& ValueStr, const void* PropertyValue, const void* DefaultValue, class UObject* Parent, int32 PortFlags, class UObject* ExportRootScope) = 0;
+		bool ExportTextItem(FString& ValueStr, const void* PropertyValue, const void* DefaultValue, UObject* Parent, int32 PortFlags, UObject* ExportRootScope)
+		{
+			return FakeVPtr->GetHandlingFunc<UE::CoreUObject::Private::EStructOpsFakeVTableFlags::ExportTextItem>()(ValueStr, PropertyValue, DefaultValue, Parent, PortFlags, ExportRootScope);
+		}
 
 		/** return true if this struct can import **/
 		bool HasImportTextItem() const
 		{
 			return GetCapabilities().HasImportTextItem;
 		}
+
 		/** 
 		 * import this structure 
 		 * @return true if the copy was imported, otherwise it will fall back to FStructProperty::ImportText
 		 */
-		virtual bool ImportTextItem(const TCHAR*& Buffer, void* Data, int32 PortFlags, class UObject* OwnerObject, FOutputDevice* ErrorText) = 0;
+		bool ImportTextItem(const TCHAR*& Buffer, void* Data, int32 PortFlags, UObject* OwnerObject, FOutputDevice* ErrorText)
+		{
+			return FakeVPtr->GetHandlingFunc<UE::CoreUObject::Private::EStructOpsFakeVTableFlags::ImportTextItem>()(Buffer, Data, PortFlags, OwnerObject, ErrorText);
+		}
 
 		bool HasFindInnerPropertyInstance() const
 		{
 			return GetCapabilities().HasFindInnerPropertyInstance;
 		}
-		virtual bool FindInnerPropertyInstance(FName PropertyName, const void* Data, const FProperty*& OutProp, const void*& OutData) const = 0;
+
+		bool FindInnerPropertyInstance(FName PropertyName, const void* Data, const FProperty*& OutProp, const void*& OutData) const
+		{
+			return FakeVPtr->GetHandlingFunc<UE::CoreUObject::Private::EStructOpsFakeVTableFlags::FindInnerPropertyInstance>()(PropertyName, Data, OutProp, OutData);
+		}
 
 		/** return true if this struct has custom GC code **/
 		bool HasAddStructReferencedObjects() const
 		{
 			return GetCapabilities().HasAddStructReferencedObjects;
 		}
+
 		/** returns true if the native serialize functions may serialize object references of the given type **/
 		bool HasSerializerObjectReferences(EPropertyObjectReferenceType Type) const
 		{
 			return EnumHasAnyFlags(GetCapabilities().HasSerializerObjectReferences, Type);
 		}
+
 		/** 
 		 * return a pointer to a function that can add referenced objects
 		 * @return true if the copy was imported, otherwise it will fall back to FStructProperty::ImportText
 		 */
-		typedef void (*TPointerToAddStructReferencedObjects)(void* A, class FReferenceCollector& Collector);
-		virtual TPointerToAddStructReferencedObjects AddStructReferencedObjects() = 0;
+		using PointerToAddStructReferencedObjectsType = void (*)(void* A, FReferenceCollector& Collector);
+		using TPointerToAddStructReferencedObjects UE_DEPRECATED(5.7, "TPointerToAddStructReferencedObjects has been deprecated, please use PointerToAddStructReferencedObjectsType instead.") = PointerToAddStructReferencedObjectsType;
+
+		PointerToAddStructReferencedObjectsType AddStructReferencedObjects()
+		{
+			return FakeVPtr->GetHandlingFunc<UE::CoreUObject::Private::EStructOpsFakeVTableFlags::AddStructReferencedObjects>()();
+		}
 
 		/** return true if this class wants to serialize from some other tag (usually for conversion purposes) **/
 		bool HasSerializeFromMismatchedTag() const
 		{
 			return GetCapabilities().HasSerializeFromMismatchedTag;
 		}
+
 		bool HasStructuredSerializeFromMismatchedTag() const
 		{
 			return GetCapabilities().HasStructuredSerializeFromMismatchedTag;
@@ -1283,8 +2078,15 @@ public:
 		 * Serialize this structure, from some other tag
 		 * @return true if this succeeded, false will trigger a warning and not serialize at all
 		 */
-		virtual bool SerializeFromMismatchedTag(struct FPropertyTag const& Tag, FArchive& Ar, void *Data) = 0;
-		virtual bool StructuredSerializeFromMismatchedTag(struct FPropertyTag const& Tag, FStructuredArchive::FSlot Slot, void *Data) = 0;
+		bool SerializeFromMismatchedTag(FPropertyTag const& Tag, FArchive& Ar, void *Data)
+		{
+			return FakeVPtr->GetHandlingFunc<UE::CoreUObject::Private::EStructOpsFakeVTableFlags::SerializeFromMismatchedTag>()(Tag, Ar, Data);
+		}
+
+		bool StructuredSerializeFromMismatchedTag(FPropertyTag const& Tag, FStructuredArchive::FSlot Slot, void *Data)
+		{
+			return FakeVPtr->GetHandlingFunc<UE::CoreUObject::Private::EStructOpsFakeVTableFlags::StructuredSerializeFromMismatchedTag>()(Tag, Slot, Data);
+		}
 
 		/** return true if this struct has a GetTypeHash */
 		bool HasGetTypeHash() const
@@ -1293,7 +2095,10 @@ public:
 		}
 
 		/** Calls GetTypeHash if enabled */
-		virtual uint32 GetStructTypeHash(const void* Src) = 0;
+		uint32 GetStructTypeHash(const void* Src)
+		{
+			return FakeVPtr->GetHandlingFunc<UE::CoreUObject::Private::EStructOpsFakeVTableFlags::GetStructTypeHash>()(Src);
+		}
 
 		/** Returns property flag values that can be computed at compile time */
 		EPropertyFlags GetComputedPropertyFlags() const
@@ -1318,17 +2123,33 @@ public:
 		{
 			return GetCapabilities().HasIntrusiveUnsetOptionalState;
 		}
+
 		/** Construct an unset optional value */
-		virtual void InitializeIntrusiveUnsetOptionalValue(void* Data) const = 0;
+		void InitializeIntrusiveUnsetOptionalValue(void* Data) const
+		{
+			FakeVPtr->GetHandlingFunc<UE::CoreUObject::Private::EStructOpsFakeVTableFlags::InitializeIntrusiveUnsetOptionalValue>()(Data);
+		}
+
 		/** Return true if the optional value at Data is in an unset state */
-		virtual bool IsIntrusiveOptionalValueSet(const void* Data) const = 0;
+		bool IsIntrusiveOptionalValueSet(const void* Data) const
+		{
+			return FakeVPtr->GetHandlingFunc<UE::CoreUObject::Private::EStructOpsFakeVTableFlags::IsIntrusiveOptionalValueSet>()(Data);
+		}
+
 		/** Reset an optional value to its unset state */
-		virtual void ClearIntrusiveOptionalValue(void* Data) const = 0;
+		void ClearIntrusiveOptionalValue(void* Data) const
+		{
+			FakeVPtr->GetHandlingFunc<UE::CoreUObject::Private::EStructOpsFakeVTableFlags::ClearIntrusiveOptionalValue>()(Data);
+		}
+
 		/** 
 		 * Used for assertions only: confirms that this type has certified that its object reference fields are safe 
 		 * for the GC to visit while the struct is in its intrusive unset optional state.
 		*/
-		virtual bool IsIntrusiveOptionalSafeForGC() const = 0;
+		bool IsIntrusiveOptionalSafeForGC() const
+		{
+			return FakeVPtr->GetHandlingFunc<UE::CoreUObject::Private::EStructOpsFakeVTableFlags::IsIntrusiveOptionalSafeForGC>()();
+		}
 
 #if WITH_EDITOR
 		/** Returns true if this struct wants to indicate whether a property can be edited in the details panel */
@@ -1336,8 +2157,12 @@ public:
 		{
 			return GetCapabilities().HasCanEditChange;
 		}
+
 		/** Returns true if this struct would allow the given property to be edited in the details panel. */
-		virtual bool CanEditChange(const FEditPropertyChain& PropertyChain, const void* Data) const = 0;
+		bool CanEditChange(const FEditPropertyChain& PropertyChain, const void* Data) const
+		{
+			return FakeVPtr->GetHandlingFunc<UE::CoreUObject::Private::EStructOpsFakeVTableFlags::CanEditChange>()(PropertyChain, Data);
+		}
 #endif
 
 		/** Returns true if this struct wants to indicate whether it has a custom impl for visiting properties */
@@ -1345,9 +2170,17 @@ public:
 		{
 			return GetCapabilities().HasVisitor;
 		}
+
 		/** Structs property visitor signature */
-		virtual EPropertyVisitorControlFlow Visit(FPropertyVisitorContext& Context, const TFunctionRef<EPropertyVisitorControlFlow(const FPropertyVisitorContext& /*Context*/)> InFunc) const = 0;
-		virtual void* ResolveVisitedPathInfo(void* Data, const FPropertyVisitorInfo& Info) const = 0;
+		EPropertyVisitorControlFlow Visit(FPropertyVisitorContext& Context, const TFunctionRef<EPropertyVisitorControlFlow(const FPropertyVisitorContext& Context)> InFunc) const
+		{
+			return FakeVPtr->GetHandlingFunc<UE::CoreUObject::Private::EStructOpsFakeVTableFlags::Visit>()(Context, InFunc);
+		}
+
+		void* ResolveVisitedPathInfo(void* Data, const FPropertyVisitorInfo& Info) const
+		{
+			return FakeVPtr->GetHandlingFunc<UE::CoreUObject::Private::EStructOpsFakeVTableFlags::ResolveVisitedPathInfo>()(Data, Info);
+		}
 
 	private:
 		/** sizeof() of the structure **/
@@ -1356,7 +2189,6 @@ public:
 		const int32 Alignment;
 	};
 
-
 	/** Template to manage dynamic access to C++ struct construction and destruction **/
 	template<class CPPSTRUCT>
 	struct TCppStructOps final : public ICppStructOps
@@ -1366,436 +2198,16 @@ public:
 			: ICppStructOps(sizeof(CPPSTRUCT), alignof(CPPSTRUCT))
 		{
 			static_assert(!(TTraits::WithAddStructReferencedObjects && alignof(CPPSTRUCT) < alignof(UObject*)), "AddStructReferencedObjects() requires pointer alignment");
-		}
-
-		virtual FCapabilities GetCapabilities() const override
-		{
-			constexpr FCapabilities Capabilities {
-				(TIsPODType<CPPSTRUCT>::Value ? CPF_IsPlainOldData : CPF_None)
-				| (std::is_trivially_destructible_v<CPPSTRUCT> ? CPF_NoDestructor : CPF_None)
-				| (TIsZeroConstructType<CPPSTRUCT>::Value ? CPF_ZeroConstructor : CPF_None)
-				| (TModels_V<CGetTypeHashable, CPPSTRUCT> ? CPF_HasGetValueTypeHash : CPF_None),
-				TTraits::WithSerializerObjectReferences,
-				TTraits::WithNoInitConstructor,
-				TTraits::WithZeroConstructor,
-				!(TTraits::WithNoDestructor || TIsPODType<CPPSTRUCT>::Value),
-				TTraits::WithSerializer,
-				TTraits::WithStructuredSerializer,
-				TTraits::WithPostSerialize,
-				TTraits::WithNetSerializer,
-				TTraits::WithNetSharedSerialization,
-				TTraits::WithNetDeltaSerializer,
-				TTraits::WithPostScriptConstruct,
-				TIsPODType<CPPSTRUCT>::Value,
-				TIsUECoreType<CPPSTRUCT>::Value,
-				TIsUECoreVariant<CPPSTRUCT>::Value,
-				TTraits::WithCopy,
-				TTraits::WithIdentical || TTraits::WithIdenticalViaEquality,
-				TTraits::WithExportTextItem,
-				TTraits::WithImportTextItem,
-				TTraits::WithAddStructReferencedObjects,
-				TTraits::WithSerializeFromMismatchedTag,
-				TTraits::WithStructuredSerializeFromMismatchedTag,
-				TModels_V<CGetTypeHashable, CPPSTRUCT>,
-				::HasIntrusiveUnsetOptionalState<CPPSTRUCT>(),
-				TIsAbstract<CPPSTRUCT>::Value,
-				TTraits::WithFindInnerPropertyInstance,
-				TTraits::WithClearOnFinishDestroy,
-#if WITH_EDITOR
-				TTraits::WithCanEditChange,
-#endif
-				TTraits::WithVisitor,
-			};
-			return Capabilities;
-		}
-		virtual void Construct(void* Dest) override
-		{
-			check(!TTraits::WithZeroConstructor); // don't call this if we have indicated it is not necessary
-			// that could have been an if statement, but we might as well force optimization above the virtual call
-			// could also not attempt to call the constructor for types where this is not possible, but I didn't do that here
-#if CHECK_PUREVIRTUALS
-			if constexpr (!TStructOpsTypeTraits<CPPSTRUCT>::WithPureVirtual)
-#endif
-			{
-				if constexpr (TStructOpsTypeTraits<CPPSTRUCT>::WithNoInitConstructor)
-				{
-					::new (Dest) CPPSTRUCT(ForceInit);
-				}
-				else
-				{
-					::new (Dest) CPPSTRUCT();
-				}
-			}
-		}
-		virtual void ConstructForTests(void* Dest) override
-		{
-			check(!TTraits::WithZeroConstructor); // don't call this if we have indicated it is not necessary
-			// that could have been an if statement, but we might as well force optimization above the virtual call
-			// could also not attempt to call the constructor for types where this is not possible, but I didn't do that here
-#if CHECK_PUREVIRTUALS
-			if constexpr (!TStructOpsTypeTraits<CPPSTRUCT>::WithPureVirtual)
-#endif
-			{
-				if constexpr (TStructOpsTypeTraits<CPPSTRUCT>::WithNoInitConstructor)
-				{
-					::new (Dest) CPPSTRUCT(ForceInit);
-				}
-				else
-				{
-					::new (Dest) CPPSTRUCT;
-				}
-			}
-		}
-		virtual void Destruct(void *Dest) override
-		{
-			check(!(TTraits::WithNoDestructor || TIsPODType<CPPSTRUCT>::Value)); // don't call this if we have indicated it is not necessary
-			// that could have been an if statement, but we might as well force optimization above the virtual call
-			// could also not attempt to call the destructor for types where this is not possible, but I didn't do that here
-			DestructItem((CPPSTRUCT*)Dest);
-		}
-		virtual bool Serialize(FArchive& Ar, void* Data, UStruct* DefaultsStruct, const void* Defaults) override
-		{
-			check(TTraits::WithSerializer); // don't call this if we have indicated it is not necessary
-			if constexpr (TStructOpsTypeTraits<CPPSTRUCT>::WithSerializer)
-			{
-				if constexpr (TModels_V<CStructSerializableWithDefaults, CPPSTRUCT>)
-				{
-					return ((CPPSTRUCT*)Data)->Serialize(Ar, DefaultsStruct, Defaults);
-				}
-				else
-				{
-					return ((CPPSTRUCT*)Data)->Serialize(Ar);
-				}
-			}
-			else
-			{
-				return false;
-			}
-		}
-		virtual bool Serialize(FStructuredArchive::FSlot Slot, void* Data, UStruct* DefaultsStruct, const void* Defaults) override
-		{
-			check(TTraits::WithStructuredSerializer); // don't call this if we have indicated it is not necessary
-			if constexpr (TStructOpsTypeTraits<CPPSTRUCT>::WithStructuredSerializer)
-			{
-				if constexpr (TModels_V<CStructStructuredSerializableWithDefaults, CPPSTRUCT>)
-				{
-					return ((CPPSTRUCT*)Data)->Serialize(Slot, DefaultsStruct, Defaults);
-				}
-				else
-				{
-					return ((CPPSTRUCT*)Data)->Serialize(Slot);
-				}
-			}
-			else
-			{
-				return false;
-			}
-		}
-		virtual void PostSerialize(const FArchive& Ar, void *Data) override
-		{
-			check(TTraits::WithPostSerialize); // don't call this if we have indicated it is not necessary
-			if constexpr (TStructOpsTypeTraits<CPPSTRUCT>::WithPostSerialize)
-			{
-				((CPPSTRUCT*)Data)->PostSerialize(Ar);
-			}
-		}
-		virtual bool NetSerialize(FArchive& Ar, class UPackageMap* Map, bool& bOutSuccess, void *Data) override
-		{
-			if constexpr (TStructOpsTypeTraits<CPPSTRUCT>::WithNetSerializer)
-			{
-				return ((CPPSTRUCT*)Data)->NetSerialize(Ar, Map, bOutSuccess);
-			}
-			else
-			{
-				return false;
-			}
-		}
-		virtual bool NetDeltaSerialize(FNetDeltaSerializeInfo & DeltaParms, void *Data) override
-		{
-			if constexpr (TStructOpsTypeTraits<CPPSTRUCT>::WithNetDeltaSerializer)
-			{
-				return ((CPPSTRUCT*)Data)->NetDeltaSerialize(DeltaParms);
-			}
-			else
-			{
-				return false;
-			}
-		}
-		virtual void PostScriptConstruct(void *Data) override
-		{
-			check(TTraits::WithPostScriptConstruct); // don't call this if we have indicated it is not necessary
-			if constexpr (TStructOpsTypeTraits<CPPSTRUCT>::WithPostScriptConstruct)
-			{
-				((CPPSTRUCT*)Data)->PostScriptConstruct();
-			}
-		}
-		virtual void GetPreloadDependencies(void* Data, TArray<UObject*>& OutDeps) override
-		{
-			if constexpr (TStructOpsTypeTraits<CPPSTRUCT>::WithGetPreloadDependencies)
-			{
-				((CPPSTRUCT*)Data)->GetPreloadDependencies(OutDeps);
-			}
-		}
-		virtual bool Copy(void* Dest, void const* Src, int32 ArrayDim) override
-		{
-			if constexpr (TStructOpsTypeTraits<CPPSTRUCT>::WithCopy)
-			{
-				static_assert((!TIsPODType<CPPSTRUCT>::Value), "You probably don't want custom copy for a POD type.");
-
-				CPPSTRUCT* TypedDest = (CPPSTRUCT*)Dest;
-				const CPPSTRUCT* TypedSrc  = (const CPPSTRUCT*)Src;
-
-				for (; ArrayDim; --ArrayDim)
-				{
-					*TypedDest++ = *TypedSrc++;
-				}
-				return true;
-			}
-			else
-			{
-				return false;
-			}
-		}
-		virtual bool Identical(const void* A, const void* B, uint32 PortFlags, bool& bOutResult) override
-		{
-			check((TTraits::WithIdentical || TTraits::WithIdenticalViaEquality)); // don't call this if we have indicated it is not necessary
-			if constexpr (TStructOpsTypeTraits<CPPSTRUCT>::WithIdentical)
-			{
-				static_assert(!TStructOpsTypeTraits<CPPSTRUCT>::WithIdenticalViaEquality, "Should not have both WithIdenticalViaEquality and WithIdentical.");
-
-				bOutResult = ((const CPPSTRUCT*)A)->Identical((const CPPSTRUCT*)B, PortFlags);
-				return true;
-			}
-			else if constexpr (TStructOpsTypeTraits<CPPSTRUCT>::WithIdenticalViaEquality)
-			{
-				bOutResult = (*(const CPPSTRUCT*)A == *(const CPPSTRUCT*)B);
-				return true;
-			}
-			else
-			{
-				bOutResult = false;
-				return false;
-			}
-		}
-		virtual bool ExportTextItem(FString& ValueStr, const void* PropertyValue, const void* DefaultValue, class UObject* Parent, int32 PortFlags, class UObject* ExportRootScope) override
-		{
-			check(TTraits::WithExportTextItem); // don't call this if we have indicated it is not necessary
-			if constexpr (TStructOpsTypeTraits<CPPSTRUCT>::WithExportTextItem)
-			{
-				if (DefaultValue)
-				{
-					return ((const CPPSTRUCT*)PropertyValue)->ExportTextItem(ValueStr, *(const CPPSTRUCT*)DefaultValue, Parent, PortFlags, ExportRootScope);
-				}
-				else
-				{
-					TTypeCompatibleBytes<CPPSTRUCT> TmpDefaultValue;
-					FMemory::Memzero(TmpDefaultValue.GetTypedPtr(), sizeof(CPPSTRUCT));
-					if (!HasZeroConstructor())
-					{
-						Construct(TmpDefaultValue.GetTypedPtr());
-					}
-
-					const bool bResult = ((const CPPSTRUCT*)PropertyValue)->ExportTextItem(ValueStr, *(const CPPSTRUCT*)TmpDefaultValue.GetTypedPtr(), Parent, PortFlags, ExportRootScope);
-
-					if (HasDestructor())
-					{
-						Destruct(TmpDefaultValue.GetTypedPtr());
-					}
 
-					return bResult;
-				}
-			}
-			else
-			{
-				return false;
-			}
-		}
-		virtual bool ImportTextItem(const TCHAR*& Buffer, void* Data, int32 PortFlags, class UObject* OwnerObject, FOutputDevice* ErrorText) override
-		{
-			check(TTraits::WithImportTextItem); // don't call this if we have indicated it is not necessary
-			if constexpr (TStructOpsTypeTraits<CPPSTRUCT>::WithImportTextItem)
-			{
-				return ((CPPSTRUCT*)Data)->ImportTextItem(Buffer, PortFlags, OwnerObject, ErrorText);
-			}
-			else
-			{
-				return false;
-			}
-		}
-		
-		virtual bool FindInnerPropertyInstance(FName PropertyName, const void* Data, const FProperty*& OutProp, const void*& OutData) const override
-		{
-			check(TTraits::WithFindInnerPropertyInstance); // don't call this if we have indicated it is not necessary
-			if constexpr (TStructOpsTypeTraits<CPPSTRUCT>::WithFindInnerPropertyInstance)
-            {
-            	return ((CPPSTRUCT*)Data)->FindInnerPropertyInstance(PropertyName, OutProp, OutData);
-            }
-            else
-            {
-            	return false;
-            }
-		}
+			FakeVPtr = UE::CoreUObject::Private::TStructOpsFakeVTableForType<CPPSTRUCT>::Value;
 
-		virtual TPointerToAddStructReferencedObjects AddStructReferencedObjects() override
-		{
-			check(TTraits::WithAddStructReferencedObjects); // don't call this if we have indicated it is not necessary
-			return &AddStructReferencedObjectsOrNot<CPPSTRUCT>;
-		}
-		virtual bool SerializeFromMismatchedTag(struct FPropertyTag const& Tag, FArchive& Ar, void *Data) override
-		{
-			check(TTraits::WithSerializeFromMismatchedTag); // don't call this if we have indicated it is not allowed
-			if constexpr (TStructOpsTypeTraits<CPPSTRUCT>::WithSerializeFromMismatchedTag)
-			{
-				if constexpr (TIsUECoreType<CPPSTRUCT>::Value)
-				{
-					// Custom version of SerializeFromMismatchedTag for core types, which don't have access to FPropertyTag.
-					FName StructName;
-					if (Tag.GetType().GetName() == NAME_StructProperty && Tag.GetType().GetParameterCount() >= 1)
-					{
-						StructName = Tag.GetType().GetParameterName(0);
-					}
-					return ((CPPSTRUCT*)Data)->SerializeFromMismatchedTag(StructName, Ar);
-				}
-				else
-				{
-					return ((CPPSTRUCT*)Data)->SerializeFromMismatchedTag(Tag, Ar);
-				}
-			}
-			else
-			{
-				return false;
-			}
-		}
-		virtual bool StructuredSerializeFromMismatchedTag(struct FPropertyTag const& Tag, FStructuredArchive::FSlot Slot, void *Data) override
-		{
-			check(TTraits::WithStructuredSerializeFromMismatchedTag); // don't call this if we have indicated it is not allowed
-			if constexpr (TStructOpsTypeTraits<CPPSTRUCT>::WithStructuredSerializeFromMismatchedTag)
-			{
-				if constexpr (TIsUECoreType<CPPSTRUCT>::Value)
-				{
-					// Custom version of SerializeFromMismatchedTag for core types, which don't understand FPropertyTag.
-					FName StructName;
-					if (Tag.GetType().GetName() == NAME_StructProperty && Tag.GetType().GetParameterCount() >= 1)
-					{
-						StructName = Tag.GetType().GetParameterName(0);
-					}
-					return ((CPPSTRUCT*)Data)->SerializeFromMismatchedTag(StructName, Slot);
-				}
-				else
-				{
-					return ((CPPSTRUCT*)Data)->SerializeFromMismatchedTag(Tag, Slot);
-				}
-			}
-			else
-			{
-				return false;
-			}
+			// We need some guarantees that TCppStructOps can be deleted via a ICppStructOps, until we fix up uses of TCppStructOps
+			static_assert(sizeof(TCppStructOps) == sizeof(ICppStructOps));
+			static_assert(std::is_trivially_destructible_v<TCppStructOps>);
+			static_assert(!std::is_polymorphic_v<TCppStructOps>);
 		}
 
 		static_assert(!(TTraits::WithSerializeFromMismatchedTag && TTraits::WithStructuredSerializeFromMismatchedTag), "Structs cannot have both WithSerializeFromMismatchedTag and WithStructuredSerializeFromMismatchedTag set");
-
-		uint32 GetStructTypeHash(const void* Src) override
-		{
-			ensure(HasGetTypeHash());
-
-			if constexpr (TModels_V<CGetTypeHashable, CPPSTRUCT>)
-			{
-				return GetTypeHashHelper(*(const CPPSTRUCT*)Src);
-			}
-			else
-			{
-				return 0;
-			}
-		}
-
-		/** Construct an unset optional value */
-		virtual void InitializeIntrusiveUnsetOptionalValue(void* Data) const override
-		{
-			if constexpr(::HasIntrusiveUnsetOptionalState<CPPSTRUCT>())
-			{
-				new (Data) TOptional<CPPSTRUCT>();
-			}
-			else
-			{
-				checkf(false, TEXT("InitializeIntrusiveUnsetOptionalValue should not be called for types which do not have an intrusive unset optional state"));
-			}
-		}
-
-		/** Return true if the optional value at Data is in an unset state */
-		virtual bool IsIntrusiveOptionalValueSet(const void* Data) const override
-		{
-			if constexpr(::HasIntrusiveUnsetOptionalState<CPPSTRUCT>())
-			{
-				return reinterpret_cast<const TOptional<CPPSTRUCT>*>(Data)->IsSet();
-			}
-			else 
-			{
-				checkf(false, TEXT("IsIntrusiveOptionalValueSet should not be called for types which do not have an intrusive unset optional state"));
-				return false;
-			}
-		}
-
-		/** Reset an optional value to its unset state */
-		virtual void ClearIntrusiveOptionalValue(void* Data) const override
-		{
-			if constexpr(::HasIntrusiveUnsetOptionalState<CPPSTRUCT>())
-			{
-				reinterpret_cast<TOptional<CPPSTRUCT>*>(Data)->Reset();
-			}
-			else
-			{
-				checkf(false, TEXT("ClearIntrusiveOptionalValue should not be called for types which do not have an intrusive unset optional state"));
-			}
-		}
-
-		/** 
-		 * Used for assertions only: confirms that this type has certified that its object reference fields are safe 
-		 * for the GC to visit while the struct is in its intrusive unset optional state.
-		*/
-		virtual bool IsIntrusiveOptionalSafeForGC() const
-		{
-			return TTraits::WithIntrusiveOptionalSafeForGC;
-		}
-
-#if WITH_EDITOR
-
-		virtual bool CanEditChange(const FEditPropertyChain& PropertyChain, const void* Data) const override
-		{
-			if constexpr (TStructOpsTypeTraits<CPPSTRUCT>::WithCanEditChange)
-			{
-				return ((const CPPSTRUCT*)Data)->CanEditChange(PropertyChain);
-			}
-			else
-			{
-				return false;
-			}
-		}
-#endif // WITH_EDITOR
-
-		virtual EPropertyVisitorControlFlow Visit(FPropertyVisitorContext& Context, const TFunctionRef<EPropertyVisitorControlFlow(const FPropertyVisitorContext& Context)> InFunc) const override
-		{
-			if constexpr (TStructOpsTypeTraits<CPPSTRUCT>::WithVisitor)
-			{
-				CPPSTRUCT* Struct = (CPPSTRUCT*)Context.Data.PropertyData;
-				return Struct->Visit(Context, InFunc);
-			}
-			else
-			{
-				return EPropertyVisitorControlFlow::StepOver;
-			}
-		}
-
-		virtual void* ResolveVisitedPathInfo(void* Data, const FPropertyVisitorInfo& Info) const override
-		{
-			if constexpr (TStructOpsTypeTraits<CPPSTRUCT>::WithVisitor)
-			{
-				return ((CPPSTRUCT*)Data)->ResolveVisitedPathInfo(Info);
-			}
-			else
-			{
-				return nullptr;
-			}
-		}
 	};
 
 	/** Template for noexport classes to autoregister before main starts **/
diff --git a/Engine/Source/Runtime/MovieScene/Public/Variants/MovieSceneNumericVariant.h b/Engine/Source/Runtime/MovieScene/Public/Variants/MovieSceneNumericVariant.h
--- a/Engine/Source/Runtime/MovieScene/Public/Variants/MovieSceneNumericVariant.h
+++ b/Engine/Source/Runtime/MovieScene/Public/Variants/MovieSceneNumericVariant.h
@@ -256,14 +256,13 @@ private:
 	 */
 	bool HasCustomWeakPtrFlag() const;
 
-private:
+public:
 
 	/** Functionally non-copyable except to the reflection layer - use Deep/Shallow copy instead */
 	FMovieSceneNumericVariant(const FMovieSceneNumericVariant& Other) = default;
 	FMovieSceneNumericVariant& operator=(const FMovieSceneNumericVariant& Other) = default;
 
-
-	friend UScriptStruct::TCppStructOps<FMovieSceneNumericVariant>;
+private:
 
 	alignas(8) uint8 Data[8];
 
diff --git a/Engine/Source/Runtime/MovieScene/Public/Variants/MovieSceneTimeWarpVariant.h b/Engine/Source/Runtime/MovieScene/Public/Variants/MovieSceneTimeWarpVariant.h
--- a/Engine/Source/Runtime/MovieScene/Public/Variants/MovieSceneTimeWarpVariant.h
+++ b/Engine/Source/Runtime/MovieScene/Public/Variants/MovieSceneTimeWarpVariant.h
@@ -350,8 +350,6 @@ public:
 	bool ExportTextItem(FString& ValueStr, const FMovieSceneTimeWarpVariant& DefaultValue, UObject* Parent, int32 PortFlags, UObject* ExportRootScope) const;
 	bool ImportTextItem(const TCHAR*& Buffer, int32 PortFlags, UObject* Parent, FOutputDevice* ErrorText);
 
-private:
-
 	/**
 	 * 
 	 */
@@ -360,7 +358,7 @@ private:
 	{
 	}
 
-	friend UScriptStruct::TCppStructOps<FMovieSceneTimeWarpVariant>;
+private:
 
 	static_assert(PLATFORM_LITTLE_ENDIAN, "This class has not been written with big-endian support.");
 
