Added utility function to FPackageName to convert a package path to a Verse path.

diff --git a/Engine/Source/Developer/AssetTools/Private/AssetTools.cpp b/Engine/Source/Developer/AssetTools/Private/AssetTools.cpp
--- a/Engine/Source/Developer/AssetTools/Private/AssetTools.cpp
+++ b/Engine/Source/Developer/AssetTools/Private/AssetTools.cpp
@@ -6715,26 +6715,10 @@ FString UAssetToolsImpl::GetUserFacingLongPackagePath(FStringView LongPackagePat
 {
 	if (!LongPackagePath.IsEmpty() && ShowingContentVersePath())
 	{
-		bool bHadClassesPrefix = false;
-		const FStringView MountPointName = FPathViews::GetMountPointNameFromPath(LongPackagePath, &bHadClassesPrefix);
-		if (!MountPointName.IsEmpty() && !bHadClassesPrefix)
+		UE::Core::FVersePath VersePath = FPackageName::LongPackagePathToVersePath(LongPackagePath);
+		if (VersePath.IsValid())
 		{
-			// If the object isn't mounted under a plugin, the object doesn't have a Verse path.
-			if (TSharedPtr<IPlugin> Plugin = IPluginManager::Get().FindPlugin(MountPointName))
-			{
-				// If the plugin doesn't have a root vpath, the object doesn't have a Verse path.
-				const FString& PluginVersePath = Plugin->GetVersePath();
-				if (!PluginVersePath.IsEmpty())
-				{
-					// Don't FPaths::Combine if getting the Verse path for the root plugin path, FPaths::Combine will add a trailing '/' making the Verse path invalid.
-					const int32 CharCount = MountPointName.Len() + 1;
-					UE::Core::FVersePath VersePath;
-					if (VersePath.TryMake(VersePath, LongPackagePath.Len() == CharCount ? PluginVersePath : FPaths::Combine(PluginVersePath, LongPackagePath.RightChop(CharCount))))
-					{
-						return MoveTemp(VersePath).ToString();
-					}
-				}
-			}
+			return MoveTemp(VersePath).ToString();
 		}
 	}
 
diff --git a/Engine/Source/Runtime/CoreUObject/Private/Misc/PackageName.cpp b/Engine/Source/Runtime/CoreUObject/Private/Misc/PackageName.cpp
--- a/Engine/Source/Runtime/CoreUObject/Private/Misc/PackageName.cpp
+++ b/Engine/Source/Runtime/CoreUObject/Private/Misc/PackageName.cpp
@@ -2483,6 +2483,37 @@ UE::Core::FVersePath FPackageName::GetVersePath(const FTopLevelAssetPath& AssetP
 	return Result;
 }
 
+UE::Core::FVersePath FPackageName::LongPackagePathToVersePath(FStringView LongPackagePath)
+{
+	// If the mount point is invalid, we can't create a Verse path from it.
+	bool bHadClassesPrefix = false;
+	const FStringView MountPointName = FPathViews::GetMountPointNameFromPath(LongPackagePath, &bHadClassesPrefix);
+	if (MountPointName.IsEmpty() || bHadClassesPrefix)
+	{
+		return {};
+	}
+
+	// If the package path isn't mounted under a plugin, the package path doesn't have a Verse path.
+	TSharedPtr<IPlugin> Plugin = IPluginManager::Get().FindPlugin(MountPointName);
+	if (!Plugin)
+	{
+		return {};
+	}
+
+	// If the plugin doesn't have a root vpath, the object doesn't have a Verse path.
+	const FString& PluginVersePath = Plugin->GetVersePath();
+	if (PluginVersePath.IsEmpty())
+	{
+		return {};
+	}
+
+	// Don't FPaths::Combine if getting the Verse path for the root plugin path, FPaths::Combine will add a trailing '/' making the Verse path invalid.
+	const int32 CharCount = MountPointName.Len() + 1;
+	UE::Core::FVersePath Result;
+	UE::Core::FVersePath::TryMake(Result, LongPackagePath.Len() == CharCount ? PluginVersePath : FPaths::Combine(PluginVersePath, LongPackagePath.RightChop(CharCount)));
+	return Result;
+}
+
 const FString& FPackageName::GetAssetPackageExtension()
 {
 	static FString AssetPackageExtension(LexToString(EPackageExtension::Asset));
diff --git a/Engine/Source/Runtime/CoreUObject/Public/Misc/PackageName.h b/Engine/Source/Runtime/CoreUObject/Public/Misc/PackageName.h
--- a/Engine/Source/Runtime/CoreUObject/Public/Misc/PackageName.h
+++ b/Engine/Source/Runtime/CoreUObject/Public/Misc/PackageName.h
@@ -501,6 +501,13 @@ public:
 	 */
 	UE_INTERNAL static COREUOBJECT_API UE::Core::FVersePath GetVersePath(const FTopLevelAssetPath& AssetPath);
 
+	/**
+	 * Gets the Verse path of a long package path.  Should not be used with a long package name.
+	 *
+	 * @return The Verse path of the long package path
+	 */
+	UE_INTERNAL static COREUOBJECT_API UE::Core::FVersePath LongPackagePathToVersePath(FStringView LongPackagePath);
+
 	/** 
 	 * Returns the file extension for packages containing assets.
 	 *
