Refactored FInstancedPropertyBag's CopyMatchingValuesByID function to allow the caller to provide a matching function.

Added CopyMatchingValuesByName that matches properties by name instead of ID.

diff --git a/Engine/Source/Runtime/CoreUObject/Private/StructUtils/PropertyBag.cpp b/Engine/Source/Runtime/CoreUObject/Private/StructUtils/PropertyBag.cpp
--- a/Engine/Source/Runtime/CoreUObject/Private/StructUtils/PropertyBag.cpp
+++ b/Engine/Source/Runtime/CoreUObject/Private/StructUtils/PropertyBag.cpp
@@ -4,6 +4,7 @@
 
 #include "HAL/IConsoleManager.h"
 #include "Hash/xxhash.h"
+#include "Misc/NotNull.h"
 #include "Misc/ScopeLock.h"
 #include "Misc/TransactionallySafeCriticalSection.h"
 #include "UObject/EnumProperty.h"
@@ -1367,12 +1368,15 @@ namespace UE::StructUtils::Private
 	}
 
 	/**
-	 * Copies properties from Source to Target property bag. The bag layouts does not need to match.
-	 * Properties are matched based on the ID in the property bag descs.
-	 * If bHasOverrides == true, then only the matching properties, whose ID is found in Overrides are copied (if Overrides is empty, nothing is copied).
-	 * If bHasOverrides == false, all matching properties are copied.
+	 * Copies properties from Source to Target property bag. The bag layouts do not need to match.
+	 * 
+	 * The provided FindTargetPropertyFunction should return the matching property in the Target 
+	 * property bag, or nullptr if there isn't one.
 	 */
-	void CopyMatchingValuesByID(const FConstStructView Source, FStructView Target, bool bHasOverrides, TConstArrayView<FGuid> Overrides)
+	void CopyMatchingValues(
+		const FConstStructView Source, 
+		FStructView Target, 
+		TFunctionRef<const FPropertyBagPropertyDesc*(TNotNull<const UPropertyBag*> /* TargetBagStruct */, const FPropertyBagPropertyDesc& /* SourceDesc */)> FindTargetPropertyFunction)
 	{
 		if (!Source.IsValid() || !Target.IsValid())
 		{
@@ -1390,13 +1394,7 @@ namespace UE::StructUtils::Private
 		// Iterate over source and copy to target if possible. Source is expected to usually have less items.
 		for (const FPropertyBagPropertyDesc& SourceDesc : SourceBagStruct->GetPropertyDescs())
 		{
-			const bool bShouldCopy = !bHasOverrides || Overrides.Contains(SourceDesc.ID);
-			if (!bShouldCopy)
-			{
-				continue;
-			}
-			
-			const FPropertyBagPropertyDesc* PotentialTargetDesc = TargetBagStruct->FindPropertyDescByID(SourceDesc.ID);
+			const FPropertyBagPropertyDesc* PotentialTargetDesc = FindTargetPropertyFunction(TargetBagStruct, SourceDesc);
 			if (PotentialTargetDesc == nullptr
 				|| PotentialTargetDesc->CachedProperty == nullptr
 				|| SourceDesc.CachedProperty == nullptr)
@@ -1465,12 +1463,29 @@ namespace UE::StructUtils::Private
 
 	void CopyMatchingValuesByID(const FConstStructView Source, FStructView Target)
 	{
-		CopyMatchingValuesByID(Source, Target, /*bHasOverrides*/false, {});
+		CopyMatchingValues(
+			Source, 
+			Target, 
+			[](TNotNull<const UPropertyBag*> TargetBagStruct, const FPropertyBagPropertyDesc& SourceDesc)
+			{
+				return TargetBagStruct->FindPropertyDescByID(SourceDesc.ID);
+			});
 	}
 
 	void CopyMatchingValuesByIDWithOverrides(const FConstStructView Source, FStructView Target, TConstArrayView<FGuid> Overrides)
 	{
-		CopyMatchingValuesByID(Source, Target, /*bHasOverrides*/true, Overrides);
+		CopyMatchingValues(
+			Source, 
+			Target, 
+			[&Overrides](TNotNull<const UPropertyBag*> TargetBagStruct, const FPropertyBagPropertyDesc& SourceDesc) -> const FPropertyBagPropertyDesc*
+			{
+				if (Overrides.Contains(SourceDesc.ID))
+				{
+					return TargetBagStruct->FindPropertyDescByID(SourceDesc.ID);
+				}
+
+				return nullptr;
+			});
 	}
 
 	void RemovePropertyByName(TArray<FPropertyBagPropertyDesc>& Descs, const FName PropertyName, const int32 StartIndex = 0)
@@ -1876,6 +1891,17 @@ void FInstancedPropertyBag::CopyMatchingValuesByID(const FInstancedPropertyBag&
 	UE::StructUtils::Private::CopyMatchingValuesByID(Other.Value, Value);
 }
 
+void FInstancedPropertyBag::CopyMatchingValuesByName(const FInstancedPropertyBag& Other)
+{
+	UE::StructUtils::Private::CopyMatchingValues(
+		Other.Value,
+		Value,
+		[](TNotNull<const UPropertyBag*> TargetBagStruct, const FPropertyBagPropertyDesc& SourceDesc)
+		{
+			return TargetBagStruct->FindPropertyDescByName(SourceDesc.Name);
+		});
+}
+
 int32 FInstancedPropertyBag::GetNumPropertiesInBag() const
 {
 	if (const UPropertyBag* BagStruct = GetPropertyBagStruct())
diff --git a/Engine/Source/Runtime/CoreUObject/Public/StructUtils/PropertyBag.h b/Engine/Source/Runtime/CoreUObject/Public/StructUtils/PropertyBag.h
--- a/Engine/Source/Runtime/CoreUObject/Public/StructUtils/PropertyBag.h
+++ b/Engine/Source/Runtime/CoreUObject/Public/StructUtils/PropertyBag.h
@@ -433,6 +433,13 @@ struct FInstancedPropertyBag
 	 */
 	UE_API void CopyMatchingValuesByID(const FInstancedPropertyBag& NewDescs);
 
+	/**
+	 * Copies matching property values from another bag of potentially mismatching layout.
+	 * The properties are matched between the bags based on the property name.
+	 * @param Other Reference to the bag to copy the values from
+	 */
+	UE_API void CopyMatchingValuesByName(const FInstancedPropertyBag& Other);
+
 	/** Returns number of the Properties in this Property Bag */
 	UE_API int32 GetNumPropertiesInBag() const;
 
