Fix ClassGeneratedBy member in UClass not having a harvested state so we don't report about it
Bump Linker mapping validation back to error now that false positive should be fixed

diff --git a/Engine/Source/Runtime/CoreUObject/Private/UObject/LinkerSave.cpp b/Engine/Source/Runtime/CoreUObject/Private/UObject/LinkerSave.cpp
--- a/Engine/Source/Runtime/CoreUObject/Private/UObject/LinkerSave.cpp
+++ b/Engine/Source/Runtime/CoreUObject/Private/UObject/LinkerSave.cpp
@@ -230,7 +230,7 @@ FPackageIndex FLinkerSave::MapObject(TObjectPtr<const UObject> Object, bool bVal
 		}
 		else if (bValidateExcluded && SaveContext)
 		{
-			UE_CLOG(SaveContext->FindCachedObjectStatus(Object) == nullptr, LogLinker, Display,
+			UE_CLOG(SaveContext->FindCachedObjectStatus(Object) == nullptr, LogLinker, Error,
 					TEXT("Invalid serialization during save of Export '%s' in package '%s'. Object '%s' encountered during serialization without being previously harvested. This will cause the reference to be null at runtime. "),
 					(CurrentlySavingExport.IsExport() ? *Exp(CurrentlySavingExport).ObjectName.ToString() : TEXT("Unknown")),
 					*LinkerRoot->GetPathName(),
diff --git a/Engine/Source/Runtime/CoreUObject/Private/UObject/SavePackage/PackageHarvester.cpp b/Engine/Source/Runtime/CoreUObject/Private/UObject/SavePackage/PackageHarvester.cpp
--- a/Engine/Source/Runtime/CoreUObject/Private/UObject/SavePackage/PackageHarvester.cpp
+++ b/Engine/Source/Runtime/CoreUObject/Private/UObject/SavePackage/PackageHarvester.cpp
@@ -613,6 +613,14 @@ void FPackageHarvester::ProcessExport(const FExportWithContext& InProcessContext
 		*this << Template;
 	}
 
+	// Generate ObjectStatus workaround for `ClassGeneratedBy` property if set as a validation workaround for export archive that don't set `ArIgnoreClassGeneratedByRef`
+	if (UClass* ClassExport = Cast<UClass>(Export))
+	{
+#if WITH_EDITORONLY_DATA
+		SaveContext.GetCachedObjectStatus(ClassExport->ClassGeneratedBy);
+#endif
+	}
+
 	// Serialize the object or CDO
 	if (Export->HasAnyFlags(RF_ClassDefaultObject))
 	{
