IncrementalCook: ClassSchema: Include the ClassWithin field of a UClass in its ClassSchema, because it affects which errors are given when cooking a package with instances of the class.
#rnx

diff --git a/Engine/Source/Runtime/CoreUObject/Private/Serialization/UnversionedPropertySerialization.cpp b/Engine/Source/Runtime/CoreUObject/Private/Serialization/UnversionedPropertySerialization.cpp
--- a/Engine/Source/Runtime/CoreUObject/Private/Serialization/UnversionedPropertySerialization.cpp
+++ b/Engine/Source/Runtime/CoreUObject/Private/Serialization/UnversionedPropertySerialization.cpp
@@ -399,6 +399,14 @@ struct FUnversionedStructSchema
 	}
 
 #if WITH_EDITORONLY_DATA
+	static void AppendObjectPath(const UObject* Object, FBlake3& HashBuilder)
+	{
+		for (const UObject* Outer = Object; Outer; Outer = Outer->GetOuter())
+		{
+			AppendHash(HashBuilder, Outer->GetFName());
+		}
+	}
+
 	static void AppendClassMetaData(const UStruct* Struct, FBlake3& HashBuilder)
 	{
 		// Append the full name of the struct. To improve performance, append the FNames of its outers individually
@@ -409,6 +417,12 @@ struct FUnversionedStructSchema
 			AppendHash(HashBuilder, Outer->GetFName());
 			Outermost = Outer;
 		}
+		// Append some properties of the struct that can impact instances' serialization or errors.
+		const UClass* StructAsClass = Cast<const UClass>(Struct);
+		if (StructAsClass)
+		{
+			AppendObjectPath(StructAsClass->ClassWithin, HashBuilder);
+		}
 		// Append metadata properties of the struct since they influence builds.
 		const TMap<FName, FString>* MetaData = nullptr;
 		if (Struct != UPackage::StaticClass())
