Reduced implementation for fake compressed CoreUObject vtables, using constraints to describe if a vtable entry needs to be generated, also removing the need to check that the functions aren't called by accident.

diff --git a/Engine/Source/Runtime/CoreUObject/Public/UObject/Class.h b/Engine/Source/Runtime/CoreUObject/Public/UObject/Class.h
--- a/Engine/Source/Runtime/CoreUObject/Public/UObject/Class.h
+++ b/Engine/Source/Runtime/CoreUObject/Public/UObject/Class.h
@@ -1101,13 +1101,11 @@ namespace UE::CoreUObject::Private
 	}
 
 	template <typename CppStruct>
+		requires (!TStructOpsTypeTraits<CppStruct>::WithZeroConstructor)
 	void Construct(void* Dest)
 	{
 		using TraitsType = TStructOpsTypeTraits<CppStruct>;
 
-		check(!TraitsType::WithZeroConstructor); // don't call this if we have indicated it is not necessary
-		// that could have been an if statement, but we might as well force optimization above the virtual call
-		// could also not attempt to call the constructor for types where this is not possible, but I didn't do that here
 #if CHECK_PUREVIRTUALS
 		if constexpr (!TraitsType::WithPureVirtual)
 #endif
@@ -1124,13 +1122,11 @@ namespace UE::CoreUObject::Private
 	}
 
 	template <typename CppStruct>
+		requires (!TStructOpsTypeTraits<CppStruct>::WithZeroConstructor)
 	void ConstructForTests(void* Dest)
 	{
 		using TraitsType = TStructOpsTypeTraits<CppStruct>;
 
-		check(!TraitsType::WithZeroConstructor); // don't call this if we have indicated it is not necessary
-		// that could have been an if statement, but we might as well force optimization above the virtual call
-		// could also not attempt to call the constructor for types where this is not possible, but I didn't do that here
 #if CHECK_PUREVIRTUALS
 		if constexpr (!TraitsType::WithPureVirtual)
 #endif
@@ -1147,72 +1143,45 @@ namespace UE::CoreUObject::Private
 	}
 
 	template <typename CppStruct>
+		requires (!TStructOpsTypeTraits<CppStruct>::WithNoDestructor && !TIsPODType<CppStruct>::Value)
 	void Destruct(void* Dest)
 	{
-		using TraitsType = TStructOpsTypeTraits<CppStruct>;
-
-		check(!(TraitsType::WithNoDestructor || TIsPODType<CppStruct>::Value)); // don't call this if we have indicated it is not necessary
-		// that could have been an if statement, but we might as well force optimization above the virtual call
-		// could also not attempt to call the destructor for types where this is not possible, but I didn't do that here
 		DestructItem((CppStruct*)Dest);
 	}
 
 	template <typename CppStruct>
+		requires (!!TStructOpsTypeTraits<CppStruct>::WithSerializer)
 	bool SerializeArchive(FArchive& Ar, void* Data, UStruct* DefaultsStruct, const void* Defaults)
 	{
-		using TraitsType = TStructOpsTypeTraits<CppStruct>;
-
-		check(TraitsType::WithSerializer); // don't call this if we have indicated it is not necessary
-		if constexpr (TraitsType::WithSerializer)
+		if constexpr (TModels_V<CStructSerializableWithDefaults, CppStruct>)
 		{
-			if constexpr (TModels_V<CStructSerializableWithDefaults, CppStruct>)
-			{
-				return ((CppStruct*)Data)->Serialize(Ar, DefaultsStruct, Defaults);
-			}
-			else
-			{
-				return ((CppStruct*)Data)->Serialize(Ar);
-			}
+			return ((CppStruct*)Data)->Serialize(Ar, DefaultsStruct, Defaults);
 		}
 		else
 		{
-			return false;
+			return ((CppStruct*)Data)->Serialize(Ar);
 		}
 	}
 
 	template <typename CppStruct>
+		requires (!!TStructOpsTypeTraits<CppStruct>::WithStructuredSerializer)
 	bool SerializeSlot(FStructuredArchive::FSlot Slot, void* Data, UStruct* DefaultsStruct, const void* Defaults)
 	{
-		using TraitsType = TStructOpsTypeTraits<CppStruct>;
-
-		check(TraitsType::WithStructuredSerializer); // don't call this if we have indicated it is not necessary
-		if constexpr (TraitsType::WithStructuredSerializer)
+		if constexpr (TModels_V<CStructStructuredSerializableWithDefaults, CppStruct>)
 		{
-			if constexpr (TModels_V<CStructStructuredSerializableWithDefaults, CppStruct>)
-			{
-				return ((CppStruct*)Data)->Serialize(Slot, DefaultsStruct, Defaults);
-			}
-			else
-			{
-				return ((CppStruct*)Data)->Serialize(Slot);
-			}
+			return ((CppStruct*)Data)->Serialize(Slot, DefaultsStruct, Defaults);
 		}
 		else
 		{
-			return false;
+			return ((CppStruct*)Data)->Serialize(Slot);
 		}
 	}
 
 	template <typename CppStruct>
+		requires (!!TStructOpsTypeTraits<CppStruct>::WithPostSerialize)
 	void PostSerialize(const FArchive& Ar, void *Data)
 	{
-		using TraitsType = TStructOpsTypeTraits<CppStruct>;
-
-		check(TraitsType::WithPostSerialize); // don't call this if we have indicated it is not necessary
-		if constexpr (TraitsType::WithPostSerialize)
-		{
-			((CppStruct*)Data)->PostSerialize(Ar);
-		}
+		((CppStruct*)Data)->PostSerialize(Ar);
 	}
 
 	template <typename CppStruct>
@@ -1246,15 +1215,10 @@ namespace UE::CoreUObject::Private
 	}
 
 	template <typename CppStruct>
+		requires (!!TStructOpsTypeTraits<CppStruct>::WithPostScriptConstruct)
 	void PostScriptConstruct(void* Data)
 	{
-		using TraitsType = TStructOpsTypeTraits<CppStruct>;
-
-		check(TraitsType::WithPostScriptConstruct); // don't call this if we have indicated it is not necessary
-		if constexpr (TraitsType::WithPostScriptConstruct)
-		{
-			((CppStruct*)Data)->PostScriptConstruct();
-		}
+		((CppStruct*)Data)->PostScriptConstruct();
 	}
 
 	template <typename CppStruct>
@@ -1292,164 +1256,109 @@ namespace UE::CoreUObject::Private
 		}
 	}
 
+	// If you get an overload resolution error here, it's likely you've specified both WithIdentical and WithIdenticalViaEquality
+	// in your TStructOpsTypeTraits, which is not supported.
 	template <typename CppStruct>
+		requires (!!TStructOpsTypeTraits<CppStruct>::WithIdentical)
 	bool Identical(const void* A, const void* B, uint32 PortFlags, bool& bOutResult)
 	{
-		using TraitsType = TStructOpsTypeTraits<CppStruct>;
-
-		check(TraitsType::WithIdentical || TraitsType::WithIdenticalViaEquality); // don't call this if we have indicated it is not necessary
-		if constexpr (TraitsType::WithIdentical)
-		{
-			static_assert(!TraitsType::WithIdenticalViaEquality, "Should not have both WithIdenticalViaEquality and WithIdentical.");
-
-			bOutResult = ((const CppStruct*)A)->Identical((const CppStruct*)B, PortFlags);
-			return true;
-		}
-		else if constexpr (TraitsType::WithIdenticalViaEquality)
-		{
-			bOutResult = (*(const CppStruct*)A == *(const CppStruct*)B);
-			return true;
-		}
-		else
-		{
-			bOutResult = false;
-			return false;
-		}
+		bOutResult = ((const CppStruct*)A)->Identical((const CppStruct*)B, PortFlags);
+		return true;
+	}
+	template <typename CppStruct>
+		requires (!!TStructOpsTypeTraits<CppStruct>::WithIdenticalViaEquality)
+	bool Identical(const void* A, const void* B, uint32 PortFlags, bool& bOutResult)
+	{
+		bOutResult = (*(const CppStruct*)A == *(const CppStruct*)B);
+		return true;
 	}
 
 	template <typename CppStruct>
+		requires (!!TStructOpsTypeTraits<CppStruct>::WithExportTextItem)
 	bool ExportTextItem(FString& ValueStr, const void* PropertyValue, const void* DefaultValue, UObject* Parent, int32 PortFlags, UObject* ExportRootScope)
 	{
-		using TraitsType = TStructOpsTypeTraits<CppStruct>;
-
-		check(TraitsType::WithExportTextItem); // don't call this if we have indicated it is not necessary
-		if constexpr (TraitsType::WithExportTextItem)
+		if (DefaultValue)
+		{
+			return ((const CppStruct*)PropertyValue)->ExportTextItem(ValueStr, *(const CppStruct*)DefaultValue, Parent, PortFlags, ExportRootScope);
+		}
+		else
 		{
-			if (DefaultValue)
+			TTypeCompatibleBytes<CppStruct> TmpDefaultValue;
+			FMemory::Memzero(TmpDefaultValue.GetTypedPtr(), sizeof(CppStruct));
+			if constexpr (!GetCapabilities<CppStruct>().HasZeroConstructor)
 			{
-				return ((const CppStruct*)PropertyValue)->ExportTextItem(ValueStr, *(const CppStruct*)DefaultValue, Parent, PortFlags, ExportRootScope);
+				Construct<CppStruct>(TmpDefaultValue.GetTypedPtr());
 			}
-			else
-			{
-				TTypeCompatibleBytes<CppStruct> TmpDefaultValue;
-				FMemory::Memzero(TmpDefaultValue.GetTypedPtr(), sizeof(CppStruct));
-				if (!GetCapabilities<CppStruct>().HasZeroConstructor)
-				{
-					Construct<CppStruct>(TmpDefaultValue.GetTypedPtr());
-				}
 
-				const bool bResult = ((const CppStruct*)PropertyValue)->ExportTextItem(ValueStr, *(const CppStruct*)TmpDefaultValue.GetTypedPtr(), Parent, PortFlags, ExportRootScope);
+			const bool bResult = ((const CppStruct*)PropertyValue)->ExportTextItem(ValueStr, *(const CppStruct*)TmpDefaultValue.GetTypedPtr(), Parent, PortFlags, ExportRootScope);
 
-				if (GetCapabilities<CppStruct>().HasDestructor)
-				{
-					Destruct<CppStruct>(TmpDefaultValue.GetTypedPtr());
-				}
-
-				return bResult;
+			if constexpr (GetCapabilities<CppStruct>().HasDestructor)
+			{
+				Destruct<CppStruct>(TmpDefaultValue.GetTypedPtr());
 			}
-		}
-		else
-		{
-			return false;
+
+			return bResult;
 		}
 	}
 
 	template <typename CppStruct>
+		requires (!!TStructOpsTypeTraits<CppStruct>::WithImportTextItem)
 	bool ImportTextItem(const TCHAR*& Buffer, void* Data, int32 PortFlags, UObject* OwnerObject, FOutputDevice* ErrorText)
 	{
-		using TraitsType = TStructOpsTypeTraits<CppStruct>;
-
-		check(TraitsType::WithImportTextItem); // don't call this if we have indicated it is not necessary
-		if constexpr (TraitsType::WithImportTextItem)
-		{
-			return ((CppStruct*)Data)->ImportTextItem(Buffer, PortFlags, OwnerObject, ErrorText);
-		}
-		else
-		{
-			return false;
-		}
+		return ((CppStruct*)Data)->ImportTextItem(Buffer, PortFlags, OwnerObject, ErrorText);
 	}
 
 	template <typename CppStruct>
+		requires (!!TStructOpsTypeTraits<CppStruct>::WithFindInnerPropertyInstance)
 	bool FindInnerPropertyInstance(FName PropertyName, const void* Data, const FProperty*& OutProp, const void*& OutData)
 	{
-		using TraitsType = TStructOpsTypeTraits<CppStruct>;
-
-		check(TraitsType::WithFindInnerPropertyInstance); // don't call this if we have indicated it is not necessary
-		if constexpr (TraitsType::WithFindInnerPropertyInstance)
-		{
-			return ((CppStruct*)Data)->FindInnerPropertyInstance(PropertyName, OutProp, OutData);
-		}
-		else
-		{
-			return false;
-		}
+		return ((CppStruct*)Data)->FindInnerPropertyInstance(PropertyName, OutProp, OutData);
 	}
 
 	template <typename CppStruct>
+		requires (!!TStructOpsTypeTraits<CppStruct>::WithAddStructReferencedObjects)
 	void (*AddStructReferencedObjects())(void* A, FReferenceCollector& Collector)
 	{
-		using TraitsType = TStructOpsTypeTraits<CppStruct>;
-
-		check(TraitsType::WithAddStructReferencedObjects); // don't call this if we have indicated it is not necessary
 		return &AddStructReferencedObjectsOrNot<CppStruct>;
 	}
 
 	template <typename CppStruct>
+		requires (!!TStructOpsTypeTraits<CppStruct>::WithSerializeFromMismatchedTag)
 	bool SerializeFromMismatchedTag(FPropertyTag const& Tag, FArchive& Ar, void* Data)
 	{
-		using TraitsType = TStructOpsTypeTraits<CppStruct>;
-
-		check(TraitsType::WithSerializeFromMismatchedTag); // don't call this if we have indicated it is not allowed
-		if constexpr (TraitsType::WithSerializeFromMismatchedTag)
+		if constexpr (TIsUECoreType<CppStruct>::Value)
 		{
-			if constexpr (TIsUECoreType<CppStruct>::Value)
-			{
-				// Custom version of SerializeFromMismatchedTag for core types, which don't have access to FPropertyTag.
-				FName StructName;
-				if (Tag.GetType().GetName() == NAME_StructProperty && Tag.GetType().GetParameterCount() >= 1)
-				{
-					StructName = Tag.GetType().GetParameterName(0);
-				}
-				return ((CppStruct*)Data)->SerializeFromMismatchedTag(StructName, Ar);
-			}
-			else
+			// Custom version of SerializeFromMismatchedTag for core types, which don't have access to FPropertyTag.
+			FName StructName;
+			if (Tag.GetType().GetName() == NAME_StructProperty && Tag.GetType().GetParameterCount() >= 1)
 			{
-				return ((CppStruct*)Data)->SerializeFromMismatchedTag(Tag, Ar);
+				StructName = Tag.GetType().GetParameterName(0);
 			}
+			return ((CppStruct*)Data)->SerializeFromMismatchedTag(StructName, Ar);
 		}
 		else
 		{
-			return false;
+			return ((CppStruct*)Data)->SerializeFromMismatchedTag(Tag, Ar);
 		}
 	}
 
 	template <typename CppStruct>
+		requires (!!TStructOpsTypeTraits<CppStruct>::WithStructuredSerializeFromMismatchedTag)
 	bool StructuredSerializeFromMismatchedTag(FPropertyTag const& Tag, FStructuredArchive::FSlot Slot, void* Data)
 	{
-		using TraitsType = TStructOpsTypeTraits<CppStruct>;
-
-		check(TraitsType::WithStructuredSerializeFromMismatchedTag); // don't call this if we have indicated it is not allowed
-		if constexpr (TraitsType::WithStructuredSerializeFromMismatchedTag)
+		if constexpr (TIsUECoreType<CppStruct>::Value)
 		{
-			if constexpr (TIsUECoreType<CppStruct>::Value)
+			// Custom version of SerializeFromMismatchedTag for core types, which don't understand FPropertyTag.
+			FName StructName;
+			if (Tag.GetType().GetName() == NAME_StructProperty && Tag.GetType().GetParameterCount() >= 1)
 			{
-				// Custom version of SerializeFromMismatchedTag for core types, which don't understand FPropertyTag.
-				FName StructName;
-				if (Tag.GetType().GetName() == NAME_StructProperty && Tag.GetType().GetParameterCount() >= 1)
-				{
-					StructName = Tag.GetType().GetParameterName(0);
-				}
-				return ((CppStruct*)Data)->SerializeFromMismatchedTag(StructName, Slot);
-			}
-			else
-			{
-				return ((CppStruct*)Data)->SerializeFromMismatchedTag(Tag, Slot);
+				StructName = Tag.GetType().GetParameterName(0);
 			}
+			return ((CppStruct*)Data)->SerializeFromMismatchedTag(StructName, Slot);
 		}
 		else
 		{
-			return false;
+			return ((CppStruct*)Data)->SerializeFromMismatchedTag(Tag, Slot);
 		}
 	}
 
@@ -1469,43 +1378,24 @@ namespace UE::CoreUObject::Private
 	}
 
 	template <typename CppStruct>
+		requires (HasIntrusiveUnsetOptionalState<CppStruct>())
 	void InitializeIntrusiveUnsetOptionalValue(void* Data)
 	{
-		if constexpr(::HasIntrusiveUnsetOptionalState<CppStruct>())
-		{
-			new (Data) TOptional<CppStruct>();
-		}
-		else
-		{
-			checkf(false, TEXT("InitializeIntrusiveUnsetOptionalValue should not be called for types which do not have an intrusive unset optional state"));
-		}
+		new (Data) TOptional<CppStruct>();
 	}
 
 	template <typename CppStruct>
+		requires (HasIntrusiveUnsetOptionalState<CppStruct>())
 	bool IsIntrusiveOptionalValueSet(const void* Data)
 	{
-		if constexpr(::HasIntrusiveUnsetOptionalState<CppStruct>())
-		{
-			return reinterpret_cast<const TOptional<CppStruct>*>(Data)->IsSet();
-		}
-		else 
-		{
-			checkf(false, TEXT("IsIntrusiveOptionalValueSet should not be called for types which do not have an intrusive unset optional state"));
-			return false;
-		}
+		return reinterpret_cast<const TOptional<CppStruct>*>(Data)->IsSet();
 	}
 
 	template <typename CppStruct>
+		requires (HasIntrusiveUnsetOptionalState<CppStruct>())
 	void ClearIntrusiveOptionalValue(void* Data)
 	{
-		if constexpr(::HasIntrusiveUnsetOptionalState<CppStruct>())
-		{
-			reinterpret_cast<TOptional<CppStruct>*>(Data)->Reset();
-		}
-		else
-		{
-			checkf(false, TEXT("ClearIntrusiveOptionalValue should not be called for types which do not have an intrusive unset optional state"));
-		}
+		reinterpret_cast<TOptional<CppStruct>*>(Data)->Reset();
 	}
 
 	template <typename CppStruct>
@@ -1567,42 +1457,42 @@ namespace UE::CoreUObject::Private
 		}
 	}
 
-	#define UE_FOREACH_FAKE_VTABLE_FUNC(IndexNameFunctypeCondition) \
-		IndexNameFunctypeCondition(0,	Construct,								void(void* Address),																																			!TStructOpsTypeTraits<CppStruct>::WithZeroConstructor) \
-		IndexNameFunctypeCondition(1,	ConstructForTests,						void(void* Address),																																			!TStructOpsTypeTraits<CppStruct>::WithZeroConstructor) \
-		IndexNameFunctypeCondition(2,	Destruct,								void(void* Address),																																			!TStructOpsTypeTraits<CppStruct>::WithNoDestructor && !TIsPODType<CppStruct>::Value) \
-		IndexNameFunctypeCondition(3,	SerializeArchive,						bool(FArchive& Ar, void* Data, UStruct* DefaultsStruct, const void* Defaults),																					TStructOpsTypeTraits<CppStruct>::WithSerializer) \
-		IndexNameFunctypeCondition(4,	SerializeSlot,							bool(FStructuredArchive::FSlot Slot, void* Data, UStruct* DefaultsStruct, const void* Defaults),																TStructOpsTypeTraits<CppStruct>::WithStructuredSerializer) \
-		IndexNameFunctypeCondition(5,	PostSerialize,							void(const FArchive& Ar, void* Data),																															TStructOpsTypeTraits<CppStruct>::WithPostSerialize) \
-		IndexNameFunctypeCondition(6,	NetSerialize,							bool(FArchive& Ar, UPackageMap* Map, bool& bOutSuccess, void* Data),																							true) \
-		IndexNameFunctypeCondition(7,	NetDeltaSerialize,						bool(FNetDeltaSerializeInfo& DeltaParms, void* Data),																											true) \
-		IndexNameFunctypeCondition(8,	PostScriptConstruct,					void(void* Data),																																				TStructOpsTypeTraits<CppStruct>::WithPostScriptConstruct) \
-		IndexNameFunctypeCondition(9,	GetPreloadDependencies,					void(void* Data, TArray<UObject*>& OutDeps),																													true) \
-		IndexNameFunctypeCondition(10,	Copy,									bool(void* Dest, void const* Src, int32 ArrayDim),																												true) \
-		IndexNameFunctypeCondition(11,	Identical,								bool(const void* A, const void* B, uint32 PortFlags, bool& bOutResult),																							TStructOpsTypeTraits<CppStruct>::WithIdentical || TStructOpsTypeTraits<CppStruct>::WithIdenticalViaEquality) \
-		IndexNameFunctypeCondition(12,	ExportTextItem,							bool(FString& ValueStr, const void* PropertyValue, const void* DefaultValue, UObject* Parent, int32 PortFlags, UObject* ExportRootScope),						TStructOpsTypeTraits<CppStruct>::WithExportTextItem) \
-		IndexNameFunctypeCondition(13,	ImportTextItem,							bool(const TCHAR*& Buffer, void* Data, int32 PortFlags, UObject* OwnerObject, FOutputDevice* ErrorText),														TStructOpsTypeTraits<CppStruct>::WithImportTextItem) \
-		IndexNameFunctypeCondition(14,	FindInnerPropertyInstance,				bool(FName PropertyName, const void* Data, const FProperty*& OutProp, const void*& OutData),																	TStructOpsTypeTraits<CppStruct>::WithFindInnerPropertyInstance) \
-		IndexNameFunctypeCondition(15,	AddStructReferencedObjects,				void(*())(void* A, FReferenceCollector& Collector),						 																						TStructOpsTypeTraits<CppStruct>::WithAddStructReferencedObjects) \
-		IndexNameFunctypeCondition(16,	SerializeFromMismatchedTag,				bool(FPropertyTag const& Tag, FArchive& Ar, void* Data),																										TStructOpsTypeTraits<CppStruct>::WithSerializeFromMismatchedTag) \
-		IndexNameFunctypeCondition(17,	StructuredSerializeFromMismatchedTag,	bool(FPropertyTag const& Tag, FStructuredArchive::FSlot Slot, void* Data),																						TStructOpsTypeTraits<CppStruct>::WithStructuredSerializeFromMismatchedTag) \
-		IndexNameFunctypeCondition(18,	GetStructTypeHash,						uint32(const void* Src),																																		true) \
-		IndexNameFunctypeCondition(19,	InitializeIntrusiveUnsetOptionalValue,	void(void* Data),																																				::HasIntrusiveUnsetOptionalState<CppStruct>()) \
-		IndexNameFunctypeCondition(20,	IsIntrusiveOptionalValueSet,			bool(const void* Data),																																			::HasIntrusiveUnsetOptionalState<CppStruct>()) \
-		IndexNameFunctypeCondition(21,	ClearIntrusiveOptionalValue,			void(void* Data),																																				::HasIntrusiveUnsetOptionalState<CppStruct>()) \
-		IndexNameFunctypeCondition(22,	IsIntrusiveOptionalSafeForGC,			bool(),																																							true) \
+	#define UE_FOREACH_FAKE_VTABLE_FUNC(IndexNameFunctype) \
+		IndexNameFunctype(0,	Construct,								void(void* Address)) \
+		IndexNameFunctype(1,	ConstructForTests,						void(void* Address)) \
+		IndexNameFunctype(2,	Destruct,								void(void* Address)) \
+		IndexNameFunctype(3,	SerializeArchive,						bool(FArchive& Ar, void* Data, UStruct* DefaultsStruct, const void* Defaults)) \
+		IndexNameFunctype(4,	SerializeSlot,							bool(FStructuredArchive::FSlot Slot, void* Data, UStruct* DefaultsStruct, const void* Defaults)) \
+		IndexNameFunctype(5,	PostSerialize,							void(const FArchive& Ar, void* Data)) \
+		IndexNameFunctype(6,	NetSerialize,							bool(FArchive& Ar, UPackageMap* Map, bool& bOutSuccess, void* Data)) \
+		IndexNameFunctype(7,	NetDeltaSerialize,						bool(FNetDeltaSerializeInfo& DeltaParms, void* Data)) \
+		IndexNameFunctype(8,	PostScriptConstruct,					void(void* Data)) \
+		IndexNameFunctype(9,	GetPreloadDependencies,					void(void* Data, TArray<UObject*>& OutDeps)) \
+		IndexNameFunctype(10,	Copy,									bool(void* Dest, void const* Src, int32 ArrayDim)) \
+		IndexNameFunctype(11,	Identical,								bool(const void* A, const void* B, uint32 PortFlags, bool& bOutResult)) \
+		IndexNameFunctype(12,	ExportTextItem,							bool(FString& ValueStr, const void* PropertyValue, const void* DefaultValue, UObject* Parent, int32 PortFlags, UObject* ExportRootScope)) \
+		IndexNameFunctype(13,	ImportTextItem,							bool(const TCHAR*& Buffer, void* Data, int32 PortFlags, UObject* OwnerObject, FOutputDevice* ErrorText)) \
+		IndexNameFunctype(14,	FindInnerPropertyInstance,				bool(FName PropertyName, const void* Data, const FProperty*& OutProp, const void*& OutData)) \
+		IndexNameFunctype(15,	AddStructReferencedObjects,				void(*())(void* A, FReferenceCollector& Collector)) \
+		IndexNameFunctype(16,	SerializeFromMismatchedTag,				bool(FPropertyTag const& Tag, FArchive& Ar, void* Data)) \
+		IndexNameFunctype(17,	StructuredSerializeFromMismatchedTag,	bool(FPropertyTag const& Tag, FStructuredArchive::FSlot Slot, void* Data)) \
+		IndexNameFunctype(18,	GetStructTypeHash,						uint32(const void* Src)) \
+		IndexNameFunctype(19,	InitializeIntrusiveUnsetOptionalValue,	void(void* Data)) \
+		IndexNameFunctype(20,	IsIntrusiveOptionalValueSet,			bool(const void* Data)) \
+		IndexNameFunctype(21,	ClearIntrusiveOptionalValue,			void(void* Data)) \
+		IndexNameFunctype(22,	IsIntrusiveOptionalSafeForGC,			bool()) \
 		UE_FOREACH_FAKE_VTABLE_FUNC_EDITOR_ONLY( \
-		IndexNameFunctypeCondition(23,	CanEditChange,							bool(const FEditPropertyChain& PropertyChain, const void* Data),																								true) \
+		IndexNameFunctype(23,	CanEditChange,							bool(const FEditPropertyChain& PropertyChain, const void* Data)) \
 		) \
-		IndexNameFunctypeCondition(24,	Visit,									EPropertyVisitorControlFlow(FPropertyVisitorContext& Context, const TFunctionRef<EPropertyVisitorControlFlow(const FPropertyVisitorContext& Context)> InFunc),	true) \
-		IndexNameFunctypeCondition(25,	ResolveVisitedPathInfo,					void*(void* Data, const FPropertyVisitorInfo& Info),																											true)
+		IndexNameFunctype(24,	Visit,									EPropertyVisitorControlFlow(FPropertyVisitorContext& Context, const TFunctionRef<EPropertyVisitorControlFlow(const FPropertyVisitorContext& Context)> InFunc)) \
+		IndexNameFunctype(25,	ResolveVisitedPathInfo,					void*(void* Data, const FPropertyVisitorInfo& Info))
 
 	// These represent one function per optional function in the fake vtable.
 	enum class EStructOpsFakeVTableFlags : uint32
 	{
 		None = 0,
 
-		#define UE_VTABLE_ENUM(Index, Name, Functype, Condition) Name = 1u << Index,
+		#define UE_VTABLE_ENUM(Index, Name, Functype) Name = 1u << Index,
 		UE_FOREACH_FAKE_VTABLE_FUNC(UE_VTABLE_ENUM)
 		#undef UE_VTABLE_ENUM
 	};
@@ -1613,7 +1503,7 @@ namespace UE::CoreUObject::Private
 	template <typename CppStruct, EStructOpsFakeVTableFlags Flag>
 	struct THandlingFunctionPtr;
 
-	#define UE_VTABLE_FUNC(Index, Name, Functype, Condition) \
+	#define UE_VTABLE_FUNC(Index, Name, Functype) \
 		template <> \
 		struct THandlingFunctionType<EStructOpsFakeVTableFlags::Name> \
 		{ \
@@ -1634,8 +1524,8 @@ namespace UE::CoreUObject::Private
 
 		EStructOpsFakeVTableFlags Result = EStructOpsFakeVTableFlags::None;
 
-		#define UE_VTABLE_SETFLAG(Index, Name, Functype, Condition) \
-			if constexpr (Condition) \
+		#define UE_VTABLE_SETFLAG(Index, Name, Functype) \
+			if constexpr (requires(TIdentity_T<Functype>*& PtrRef) { PtrRef = &Name<CppStruct>; }) /* assignment to RefPtr shouldn't be needed, but is a workaround for MSVC bug: https://developercommunity.microsoft.com/t/Compiler-fails-overload-resolution-when/10939568 */ \
 			{ \
 				Result |= EStructOpsFakeVTableFlags::Name; \
 			}
