Made FCapabilities init more readable.

diff --git a/Engine/Source/Runtime/CoreUObject/Public/UObject/Class.h b/Engine/Source/Runtime/CoreUObject/Public/UObject/Class.h
--- a/Engine/Source/Runtime/CoreUObject/Public/UObject/Class.h
+++ b/Engine/Source/Runtime/CoreUObject/Public/UObject/Class.h
@@ -1062,40 +1062,41 @@ namespace UE::CoreUObject::Private
 		using TraitsType = TStructOpsTypeTraits<CppStruct>;
 
 		constexpr FCapabilities Capabilities {
-			(TIsPODType<CppStruct>::Value ? CPF_IsPlainOldData : CPF_None)
-			| (std::is_trivially_destructible_v<CppStruct> ? CPF_NoDestructor : CPF_None)
-			| (TIsZeroConstructType<CppStruct>::Value ? CPF_ZeroConstructor : CPF_None)
-			| (TModels_V<CGetTypeHashable, CppStruct> ? CPF_HasGetValueTypeHash : CPF_None),
-			TraitsType::WithSerializerObjectReferences,
-			TraitsType::WithNoInitConstructor,
-			TraitsType::WithZeroConstructor,
-			!(TraitsType::WithNoDestructor || TIsPODType<CppStruct>::Value),
-			TraitsType::WithSerializer,
-			TraitsType::WithStructuredSerializer,
-			TraitsType::WithPostSerialize,
-			TraitsType::WithNetSerializer,
-			TraitsType::WithNetSharedSerialization,
-			TraitsType::WithNetDeltaSerializer,
-			TraitsType::WithPostScriptConstruct,
-			TIsPODType<CppStruct>::Value,
-			TIsUECoreType<CppStruct>::Value,
-			TIsUECoreVariant<CppStruct>::Value,
-			TraitsType::WithCopy,
-			TraitsType::WithIdentical || TraitsType::WithIdenticalViaEquality,
-			TraitsType::WithExportTextItem,
-			TraitsType::WithImportTextItem,
-			TraitsType::WithAddStructReferencedObjects,
-			TraitsType::WithSerializeFromMismatchedTag,
-			TraitsType::WithStructuredSerializeFromMismatchedTag,
-			TModels_V<CGetTypeHashable, CppStruct>,
-			::HasIntrusiveUnsetOptionalState<CppStruct>(),
-			TIsAbstract<CppStruct>::Value,
-			TraitsType::WithFindInnerPropertyInstance,
-			TraitsType::WithClearOnFinishDestroy,
+			.ComputedPropertyFlags =
+				(TIsPODType<CppStruct>::Value ? CPF_IsPlainOldData : CPF_None) |
+				(std::is_trivially_destructible_v<CppStruct> ? CPF_NoDestructor : CPF_None) |
+				(TIsZeroConstructType<CppStruct>::Value ? CPF_ZeroConstructor : CPF_None) |
+				(TModels_V<CGetTypeHashable, CppStruct> ? CPF_HasGetValueTypeHash : CPF_None),
+			.HasSerializerObjectReferences				= TraitsType::WithSerializerObjectReferences,
+			.HasNoopConstructor							= TraitsType::WithNoInitConstructor,
+			.HasZeroConstructor							= TraitsType::WithZeroConstructor,
+			.HasDestructor								= !(TraitsType::WithNoDestructor || TIsPODType<CppStruct>::Value),
+			.HasSerializer								= TraitsType::WithSerializer,
+			.HasStructuredSerializer					= TraitsType::WithStructuredSerializer,
+			.HasPostSerialize							= TraitsType::WithPostSerialize,
+			.HasNetSerializer							= TraitsType::WithNetSerializer,
+			.HasNetSharedSerialization					= TraitsType::WithNetSharedSerialization,
+			.HasNetDeltaSerializer						= TraitsType::WithNetDeltaSerializer,
+			.HasPostScriptConstruct						= TraitsType::WithPostScriptConstruct,
+			.IsPlainOldData								= TIsPODType<CppStruct>::Value,
+			.IsUECoreType								= TIsUECoreType<CppStruct>::Value,
+			.IsUECoreVariant							= TIsUECoreVariant<CppStruct>::Value,
+			.HasCopy									= TraitsType::WithCopy,
+			.HasIdentical								= TraitsType::WithIdentical || TraitsType::WithIdenticalViaEquality,
+			.HasExportTextItem							= TraitsType::WithExportTextItem,
+			.HasImportTextItem							= TraitsType::WithImportTextItem,
+			.HasAddStructReferencedObjects				= TraitsType::WithAddStructReferencedObjects,
+			.HasSerializeFromMismatchedTag				= TraitsType::WithSerializeFromMismatchedTag,
+			.HasStructuredSerializeFromMismatchedTag	= TraitsType::WithStructuredSerializeFromMismatchedTag,
+			.HasGetTypeHash								= TModels_V<CGetTypeHashable, CppStruct>,
+			.HasIntrusiveUnsetOptionalState				= ::HasIntrusiveUnsetOptionalState<CppStruct>(),
+			.IsAbstract									= TIsAbstract<CppStruct>::Value,
+			.HasFindInnerPropertyInstance				= TraitsType::WithFindInnerPropertyInstance,
+			.ClearOnFinishDestroy						= TraitsType::WithClearOnFinishDestroy,
 #if WITH_EDITOR
-			TraitsType::WithCanEditChange,
+			.HasCanEditChange							= TraitsType::WithCanEditChange,
 #endif
-			TraitsType::WithVisitor,
+			.HasVisitor									= TraitsType::WithVisitor,
 		};
 		return Capabilities;
 	}
