Remove unused field `bLoadedByEditorPropertiesOnly`.

This was set up in UPackage::PostInitProperties but never actually referenced again.

I noticed this due to the std::atomic, which caused an issue in AutoRTFM, but there's no multithreaded access and in fact no loads are issued at all.

diff --git a/Engine/Source/Runtime/CoreUObject/Private/UObject/Package.cpp b/Engine/Source/Runtime/CoreUObject/Private/UObject/Package.cpp
--- a/Engine/Source/Runtime/CoreUObject/Private/UObject/Package.cpp
+++ b/Engine/Source/Runtime/CoreUObject/Private/UObject/Package.cpp
@@ -135,10 +135,6 @@ void UPackage::PostInitProperties()
 
 	SetPIEInstanceID(INDEX_NONE);
 	bIsCookedForEditor = false;
-	// Mark this package as editor-only by default. As soon as something in it is accessed through a non editor-only
-	// property the flag will be removed.
-	bLoadedByEditorPropertiesOnly = !HasAnyFlags(RF_ClassDefaultObject) && !HasAnyPackageFlags(PKG_CompiledIn) && (IsRunningCommandlet());
-
 	bIsDynamicPIEPackagePending = false;
 #endif // WITH_EDITORONLY_DATA
 }
diff --git a/Engine/Source/Runtime/CoreUObject/Public/UObject/Package.h b/Engine/Source/Runtime/CoreUObject/Public/UObject/Package.h
--- a/Engine/Source/Runtime/CoreUObject/Public/UObject/Package.h
+++ b/Engine/Source/Runtime/CoreUObject/Public/UObject/Package.h
@@ -293,13 +293,9 @@ public:
 	bool bIsCookedForEditor{ false };
 
 private:
-	/** True if this package is only referenced by editor-only properties 
-	*   Note: This flag is manipulated on different threads, do not mix with other bitfields above as they might get corrupted.
-	*/
-	std::atomic<bool> bLoadedByEditorPropertiesOnly{ false };
-
 	/** This flag is manipulated on different threads, do not mix with other bitfields above as they might get corrupted. */
 	std::atomic<bool> bHasBeenEndLoaded{ false };
+
 public:
 	/** This flag becomes true for loaded packages after serialization and postload and before returning from LoadPackage or calling load completion delegate.
 	  *  For newly created packages in editor, it becomes true once the package is saved. 
