[AutoRTFM] Fix for task_group memory corruption

When `UObject::ProcessContextOpcode()` hit a runtime error, the following would happen:
* `FBlueprintCoreDelegates::ThrowScriptException()` was called within an `AUTORTFM_OPEN`
* This would eventually call into `RaiseVerseRuntimeError()`, which would call `AutoRTFM::CascadingAbortTransaction()` in the open. When `AutoRTFM::CascadingAbortTransaction()` is called the open, it will immediately undo the transactional memory of the inner-most transaction only - the outer transactions abort only later, when we transition back from open -> closed.
* After calling `AutoRTFM::CascadingAbortTransaction()`, we would call `DoRaiseVerseRuntimeErrorInOpen()`, which calls `FContentScopeGuard::GetActiveScope()->Terminate()`, and that modifies the `task_group` in the open.

If the BPVM had created any tasks in an outer transaction before the runtime error, then we'd end up with open modification to memory that is also recorded as part of an outer transaction's write log. When the outer transactions finally abort, we end up with mixed open/closed writes, causing memory corruption to the `task_group::ActiveTasks` array.

This CL fixes the issue in two ways:
* The `UE_AUTORTFM_OPEN` has been removed from `UObject::ProcessContextOpcode()`. Testing suggests that this is now unnecessary, and by removing this we keep the broadcast of the runtime failure in the closed.
* `RaiseVerseRuntimeError()` has been changed to use the `AutoRTFM::OnComplete()` callback for the call to `DoRaiseVerseRuntimeErrorInOpen()`, which ensures that all the transactions have been fully aborted before we reset the `task_group`. In order to use this API, `RaiseVerseRuntimeError()` has been restructured to assume that it is only ever called in the closed, when in a transaction. If called in the open, an ensure will fire to alert us.

diff --git a/Engine/Source/Runtime/CoreUObject/Private/UObject/ScriptCore.cpp b/Engine/Source/Runtime/CoreUObject/Private/UObject/ScriptCore.cpp
--- a/Engine/Source/Runtime/CoreUObject/Private/UObject/ScriptCore.cpp
+++ b/Engine/Source/Runtime/CoreUObject/Private/UObject/ScriptCore.cpp
@@ -5,6 +5,7 @@
 =============================================================================*/
 
 #include "CoreMinimal.h"
+#include "Misc/AssertionMacros.h"
 #include "Misc/CoreMisc.h"
 #include "Misc/CommandLine.h"
 #include "Logging/LogScopedCategoryAndVerbosityOverride.h"
@@ -3161,50 +3162,47 @@ void UObject::ProcessContextOpcode( FFrame& Stack, RESULT_DECL, bool bCanFailSil
 
 		if (!bCanFailSilently)
 		{
-			UE_AUTORTFM_OPEN
+			if (NewContext && !IsValid(NewContext))
 			{
-				if (NewContext && !IsValid(NewContext))
-				{
-					FBlueprintExceptionInfo ExceptionInfo(
-						EBlueprintExceptionType::AccessViolation,
-						FText::Format(
-							LOCTEXT("AccessPendingKill", "Attempted to access {0} via property {1}, but {0} is not valid (pending kill or garbage)"),
-							FText::FromString(GetNameSafe(NewContext)),
-							FText::FromString(GetNameSafe(Stack.MostRecentProperty))
-						)
-					);
-					FBlueprintCoreDelegates::ThrowScriptException(this, Stack, ExceptionInfo);
-				}
-				else if (Stack.MostRecentProperty != NULL)
-				{
-						FBlueprintExceptionInfo ExceptionInfo(
-							EBlueprintExceptionType::AccessViolation,
-							Stack.MostRecentProperty->GetOwner<UClass>()
-							? FText::Format(
-								LOCTEXT("AccessNoneUClassContext", "Accessed None trying to read {2} property {0} in {1}"),
-								FText::FromString(Stack.MostRecentProperty->GetName()),
-								FText::FromString(Stack.MostRecentProperty->GetOwner<UClass>()->GetName()),
-								FText::FromString(Stack.MostRecentProperty->HasAllPropertyFlags(CPF_Virtual)?"(virtual)":"(real)"))
-							: FText::Format(
-								LOCTEXT("AccessNoneContext", "Accessed None trying to read {1} property {0} in not an UClass"),
-								FText::FromString(Stack.MostRecentProperty->GetName()),
-								FText::FromString(Stack.MostRecentProperty->HasAllPropertyFlags(CPF_Virtual) ? "(virtual)" : "(real)"))
-						);
-						FBlueprintCoreDelegates::ThrowScriptException(this, Stack, ExceptionInfo);
-				}
-				else
-				{
-					// Stack.MostRecentProperty will be NULL under the following conditions:
-					//   1. the context expression was a function call which returned an object
-					//   2. the context expression was a literal object reference
-					//   3. the context expression was an instance variable that no longer exists (it was editor-only, etc.)
-					FBlueprintExceptionInfo ExceptionInfo(
-						EBlueprintExceptionType::AccessViolation,
-						LOCTEXT("AccessNoneNoContext", "Accessed None")
-					);
-					FBlueprintCoreDelegates::ThrowScriptException(this, Stack, ExceptionInfo);
-				}
-			};
+				FBlueprintExceptionInfo ExceptionInfo(
+					EBlueprintExceptionType::AccessViolation,
+					FText::Format(
+						LOCTEXT("AccessPendingKill", "Attempted to access {0} via property {1}, but {0} is not valid (pending kill or garbage)"),
+						FText::FromString(GetNameSafe(NewContext)),
+						FText::FromString(GetNameSafe(Stack.MostRecentProperty))
+					)
+				);
+				FBlueprintCoreDelegates::ThrowScriptException(this, Stack, ExceptionInfo);
+			}
+			else if (Stack.MostRecentProperty != NULL)
+			{
+				FBlueprintExceptionInfo ExceptionInfo(
+					EBlueprintExceptionType::AccessViolation,
+					Stack.MostRecentProperty->GetOwner<UClass>()
+					? FText::Format(
+						LOCTEXT("AccessNoneUClassContext", "Accessed None trying to read {2} property {0} in {1}"),
+						FText::FromString(Stack.MostRecentProperty->GetName()),
+						FText::FromString(Stack.MostRecentProperty->GetOwner<UClass>()->GetName()),
+						FText::FromString(Stack.MostRecentProperty->HasAllPropertyFlags(CPF_Virtual)?"(virtual)":"(real)"))
+					: FText::Format(
+						LOCTEXT("AccessNoneContext", "Accessed None trying to read {1} property {0} in not an UClass"),
+						FText::FromString(Stack.MostRecentProperty->GetName()),
+						FText::FromString(Stack.MostRecentProperty->HasAllPropertyFlags(CPF_Virtual) ? "(virtual)" : "(real)"))
+				);
+				FBlueprintCoreDelegates::ThrowScriptException(this, Stack, ExceptionInfo);
+			}
+			else
+			{
+				// Stack.MostRecentProperty will be NULL under the following conditions:
+				//   1. the context expression was a function call which returned an object
+				//   2. the context expression was a literal object reference
+				//   3. the context expression was an instance variable that no longer exists (it was editor-only, etc.)
+				FBlueprintExceptionInfo ExceptionInfo(
+					EBlueprintExceptionType::AccessViolation,
+					LOCTEXT("AccessNoneNoContext", "Accessed None")
+				);
+				FBlueprintCoreDelegates::ThrowScriptException(this, Stack, ExceptionInfo);
+			}
 		}
 
 		const CodeSkipSizeType wSkip = Stack.ReadCodeSkipCount(); // Code offset for NULL expressions. Code += sizeof(CodeSkipSizeType)
