Removing structop function template instantiations which don't do anything useful with fallbacks in the non-templated calls.  This replaces a lot of duplicate instantiations which simple capability bit checks, and those removed duplicate instantiations will also result into smaller fake (compact) vtables.

diff --git a/Engine/Source/Runtime/CoreUObject/Public/UObject/Class.h b/Engine/Source/Runtime/CoreUObject/Public/UObject/Class.h
--- a/Engine/Source/Runtime/CoreUObject/Public/UObject/Class.h
+++ b/Engine/Source/Runtime/CoreUObject/Public/UObject/Class.h
@@ -1054,6 +1054,8 @@ namespace UE::CoreUObject::Private
 		bool HasCanEditChange : 1;
 #endif
 		bool HasVisitor : 1;
+		bool HasGetPreloadDependencies : 1;
+		bool IsIntrusiveOptionalSafeForGC : 1;
 	};
 
 	template <typename CppStruct>
@@ -1097,6 +1099,8 @@ namespace UE::CoreUObject::Private
 			.HasCanEditChange							= TraitsType::WithCanEditChange,
 #endif
 			.HasVisitor									= TraitsType::WithVisitor,
+			.HasGetPreloadDependencies					= TraitsType::WithGetPreloadDependencies,
+			.IsIntrusiveOptionalSafeForGC				= TraitsType::WithIntrusiveOptionalSafeForGC,
 		};
 		return Capabilities;
 	}
@@ -1186,33 +1190,17 @@ namespace UE::CoreUObject::Private
 	}
 
 	template <typename CppStruct>
+		requires (!!TStructOpsTypeTraits<CppStruct>::WithNetSerializer)
 	bool NetSerialize(FArchive& Ar, UPackageMap* Map, bool& bOutSuccess, void* Data)
 	{
-		using TraitsType = TStructOpsTypeTraits<CppStruct>;
-
-		if constexpr (TraitsType::WithNetSerializer)
-		{
-			return ((CppStruct*)Data)->NetSerialize(Ar, Map, bOutSuccess);
-		}
-		else
-		{
-			return false;
-		}
+		return ((CppStruct*)Data)->NetSerialize(Ar, Map, bOutSuccess);
 	}
 
 	template <typename CppStruct>
+		requires (!!TStructOpsTypeTraits<CppStruct>::WithNetDeltaSerializer)
 	bool NetDeltaSerialize(FNetDeltaSerializeInfo& DeltaParms, void* Data)
 	{
-		using TraitsType = TStructOpsTypeTraits<CppStruct>;
-
-		if constexpr (TraitsType::WithNetDeltaSerializer)
-		{
-			return ((CppStruct*)Data)->NetDeltaSerialize(DeltaParms);
-		}
-		else
-		{
-			return false;
-		}
+		return ((CppStruct*)Data)->NetDeltaSerialize(DeltaParms);
 	}
 
 	template <typename CppStruct>
@@ -1223,38 +1211,26 @@ namespace UE::CoreUObject::Private
 	}
 
 	template <typename CppStruct>
+		requires (!!TStructOpsTypeTraits<CppStruct>::WithGetPreloadDependencies)
 	void GetPreloadDependencies(void* Data, TArray<UObject*>& OutDeps)
 	{
-		using TraitsType = TStructOpsTypeTraits<CppStruct>;
-
-		if constexpr (TraitsType::WithGetPreloadDependencies)
-		{
-			((CppStruct*)Data)->GetPreloadDependencies(OutDeps);
-		}
+		((CppStruct*)Data)->GetPreloadDependencies(OutDeps);
 	}
 
 	template <typename CppStruct>
+		requires (!!TStructOpsTypeTraits<CppStruct>::WithCopy)
 	bool Copy(void* Dest, void const* Src, int32 ArrayDim)
 	{
-		using TraitsType = TStructOpsTypeTraits<CppStruct>;
-
-		if constexpr (TraitsType::WithCopy)
-		{
-			static_assert((!TIsPODType<CppStruct>::Value), "You probably don't want custom copy for a POD type.");
+		static_assert((!TIsPODType<CppStruct>::Value), "You probably don't want custom copy for a POD type.");
 
-			CppStruct* TypedDest = (CppStruct*)Dest;
-			const CppStruct* TypedSrc  = (const CppStruct*)Src;
+		CppStruct* TypedDest = (CppStruct*)Dest;
+		const CppStruct* TypedSrc  = (const CppStruct*)Src;
 
-			for (; ArrayDim; --ArrayDim)
-			{
-				*TypedDest++ = *TypedSrc++;
-			}
-			return true;
-		}
-		else
+		for (; ArrayDim; --ArrayDim)
 		{
-			return false;
+			*TypedDest++ = *TypedSrc++;
 		}
+		return true;
 	}
 
 	// If you get an overload resolution error here, it's likely you've specified both WithIdentical and WithIdenticalViaEquality
@@ -1364,18 +1340,10 @@ namespace UE::CoreUObject::Private
 	}
 
 	template <typename CppStruct>
+		requires (TModels_V<CGetTypeHashable, CppStruct>)
 	uint32 GetStructTypeHash(const void* Src)
 	{
-		ensure(GetCapabilities<CppStruct>().HasGetTypeHash);
-
-		if constexpr (TModels_V<CGetTypeHashable, CppStruct>)
-		{
-			return GetTypeHashHelper(*(const CppStruct*)Src);
-		}
-		else
-		{
-			return 0;
-		}
+		return GetTypeHashHelper(*(const CppStruct*)Src);
 	}
 
 	template <typename CppStruct>
@@ -1399,28 +1367,12 @@ namespace UE::CoreUObject::Private
 		reinterpret_cast<TOptional<CppStruct>*>(Data)->Reset();
 	}
 
-	template <typename CppStruct>
-	bool IsIntrusiveOptionalSafeForGC()
-	{
-		using TraitsType = TStructOpsTypeTraits<CppStruct>;
-
-		return TraitsType::WithIntrusiveOptionalSafeForGC;
-	}
-
 #if WITH_EDITOR
 	template <typename CppStruct>
+		requires (!!TStructOpsTypeTraits<CppStruct>::WithCanEditChange)
 	bool CanEditChange(const FEditPropertyChain& PropertyChain, const void* Data)
 	{
-		using TraitsType = TStructOpsTypeTraits<CppStruct>;
-
-		if constexpr (TraitsType::WithCanEditChange)
-		{
-			return ((const CppStruct*)Data)->CanEditChange(PropertyChain);
-		}
-		else
-		{
-			return false;
-		}
+		return ((const CppStruct*)Data)->CanEditChange(PropertyChain);
 	}
 	#define UE_FOREACH_FAKE_VTABLE_FUNC_EDITOR_ONLY(...) __VA_ARGS__
 #else
@@ -1428,34 +1380,18 @@ namespace UE::CoreUObject::Private
 #endif // WITH_EDITOR
 
 	template <typename CppStruct>
+		requires (!!TStructOpsTypeTraits<CppStruct>::WithVisitor)
 	EPropertyVisitorControlFlow Visit(FPropertyVisitorContext& Context, const TFunctionRef<EPropertyVisitorControlFlow(const FPropertyVisitorContext& Context)> InFunc)
 	{
-		using TraitsType = TStructOpsTypeTraits<CppStruct>;
-
-		if constexpr (TraitsType::WithVisitor)
-		{
-			CppStruct* Struct = (CppStruct*)Context.Data.PropertyData;
-			return Struct->Visit(Context, InFunc);
-		}
-		else
-		{
-			return EPropertyVisitorControlFlow::StepOver;
-		}
+		CppStruct* Struct = (CppStruct*)Context.Data.PropertyData;
+		return Struct->Visit(Context, InFunc);
 	}
 
 	template <typename CppStruct>
+		requires (!!TStructOpsTypeTraits<CppStruct>::WithVisitor)
 	void* ResolveVisitedPathInfo(void* Data, const FPropertyVisitorInfo& Info)
 	{
-		using TraitsType = TStructOpsTypeTraits<CppStruct>;
-
-		if constexpr (TraitsType::WithVisitor)
-		{
-			return ((CppStruct*)Data)->ResolveVisitedPathInfo(Info);
-		}
-		else
-		{
-			return nullptr;
-		}
+		return ((CppStruct*)Data)->ResolveVisitedPathInfo(Info);
 	}
 
 	#define UE_FOREACH_FAKE_VTABLE_FUNC(IndexNameFunctype) \
@@ -1481,12 +1417,11 @@ namespace UE::CoreUObject::Private
 		IndexNameFunctype(19,	InitializeIntrusiveUnsetOptionalValue,	void(void* Data)) \
 		IndexNameFunctype(20,	IsIntrusiveOptionalValueSet,			bool(const void* Data)) \
 		IndexNameFunctype(21,	ClearIntrusiveOptionalValue,			void(void* Data)) \
-		IndexNameFunctype(22,	IsIntrusiveOptionalSafeForGC,			bool()) \
 		UE_FOREACH_FAKE_VTABLE_FUNC_EDITOR_ONLY( \
-		IndexNameFunctype(23,	CanEditChange,							bool(const FEditPropertyChain& PropertyChain, const void* Data)) \
+		IndexNameFunctype(22,	CanEditChange,							bool(const FEditPropertyChain& PropertyChain, const void* Data)) \
 		) \
-		IndexNameFunctype(24,	Visit,									EPropertyVisitorControlFlow(FPropertyVisitorContext& Context, const TFunctionRef<EPropertyVisitorControlFlow(const FPropertyVisitorContext& Context)> InFunc)) \
-		IndexNameFunctype(25,	ResolveVisitedPathInfo,					void*(void* Data, const FPropertyVisitorInfo& Info))
+		IndexNameFunctype(23,	Visit,									EPropertyVisitorControlFlow(FPropertyVisitorContext& Context, const TFunctionRef<EPropertyVisitorControlFlow(const FPropertyVisitorContext& Context)> InFunc)) \
+		IndexNameFunctype(24,	ResolveVisitedPathInfo,					void*(void* Data, const FPropertyVisitorInfo& Info))
 
 	// These represent one function per optional function in the fake vtable.
 	enum class EStructOpsFakeVTableFlags : uint32
@@ -1644,6 +1579,14 @@ namespace UE::CoreUObject::Private
 #endif
 		}
 	};
+	template <typename CppStruct>
+	inline const TStructOpsFakeVTable<0> TStructOpsFakeVTableInstance<CppStruct> = 
+	{
+		{
+			GetStructOpsFakeVTableFlags<CppStruct>(),
+			&GetCapabilities<CppStruct>
+		}
+	};
 
 	template <typename CppStruct, EStructOpsFakeVTableFlags Flags, auto... Funcs>
 	struct TStructOpsFakeVTableForTypeImpl : TStructOpsFakeVTableForTypeImpl<CppStruct, EnumRemoveLowestSetFlag(Flags), Funcs..., THandlingFunctionPtr<CppStruct, EnumLowestSetFlag(Flags)>::Value>
@@ -1806,7 +1749,7 @@ public:
 		 */
 		bool NetSerialize(FArchive& Ar, UPackageMap* Map, bool& bOutSuccess, void *Data)
 		{
-			return FakeVPtr->GetHandlingFunc<UE::CoreUObject::Private::EStructOpsFakeVTableFlags::NetSerialize>()(Ar, Map, bOutSuccess, Data);
+			return HasNetSerializer() && FakeVPtr->GetHandlingFunc<UE::CoreUObject::Private::EStructOpsFakeVTableFlags::NetSerialize>()(Ar, Map, bOutSuccess, Data);
 		}
 
 		/** return true if this struct can net delta serialize delta (serialize a network delta from a base state) **/
@@ -1821,7 +1764,7 @@ public:
 		 */
 		bool NetDeltaSerialize(FNetDeltaSerializeInfo & DeltaParms, void *Data)
 		{
-			return FakeVPtr->GetHandlingFunc<UE::CoreUObject::Private::EStructOpsFakeVTableFlags::NetDeltaSerialize>()(DeltaParms, Data);
+			return HasNetDeltaSerializer() && FakeVPtr->GetHandlingFunc<UE::CoreUObject::Private::EStructOpsFakeVTableFlags::NetDeltaSerialize>()(DeltaParms, Data);
 		}
 
 		/** return true if this class implements a post script construct call **/
@@ -1836,10 +1779,19 @@ public:
 			FakeVPtr->GetHandlingFunc<UE::CoreUObject::Private::EStructOpsFakeVTableFlags::PostScriptConstruct>()(Data);
 		}
 
+		/** return true if this struct has GetPreloadDependencies **/
+		bool HasGetPreloadDependencies() const
+		{
+			return GetCapabilities().HasGetPreloadDependencies;
+		}
+
 		/** Call PreloadDependencies on this structure */
 		void GetPreloadDependencies(void* Data, TArray<UObject*>& OutDeps)
 		{
-			FakeVPtr->GetHandlingFunc<UE::CoreUObject::Private::EStructOpsFakeVTableFlags::GetPreloadDependencies>()(Data, OutDeps);
+			if (HasGetPreloadDependencies())
+			{
+				FakeVPtr->GetHandlingFunc<UE::CoreUObject::Private::EStructOpsFakeVTableFlags::GetPreloadDependencies>()(Data, OutDeps);
+			}
 		}
 
 		/** return true if this struct should be memcopied **/
@@ -1872,7 +1824,7 @@ public:
 		 */
 		bool Copy(void* Dest, void const* Src, int32 ArrayDim)
 		{
-			return FakeVPtr->GetHandlingFunc<UE::CoreUObject::Private::EStructOpsFakeVTableFlags::Copy>()(Dest, Src, ArrayDim);
+			return HasCopy() && FakeVPtr->GetHandlingFunc<UE::CoreUObject::Private::EStructOpsFakeVTableFlags::Copy>()(Dest, Src, ArrayDim);
 		}
 
 		/** return true if this struct can compare **/
@@ -1988,7 +1940,14 @@ public:
 		/** Calls GetTypeHash if enabled */
 		uint32 GetStructTypeHash(const void* Src)
 		{
-			return FakeVPtr->GetHandlingFunc<UE::CoreUObject::Private::EStructOpsFakeVTableFlags::GetStructTypeHash>()(Src);
+			if (HasGetTypeHash())
+			{
+				return FakeVPtr->GetHandlingFunc<UE::CoreUObject::Private::EStructOpsFakeVTableFlags::GetStructTypeHash>()(Src);
+			}
+			else
+			{
+				return 0;
+			}
 		}
 
 		/** Returns property flag values that can be computed at compile time */
@@ -2039,7 +1998,7 @@ public:
 		*/
 		bool IsIntrusiveOptionalSafeForGC() const
 		{
-			return FakeVPtr->GetHandlingFunc<UE::CoreUObject::Private::EStructOpsFakeVTableFlags::IsIntrusiveOptionalSafeForGC>()();
+			return GetCapabilities().IsIntrusiveOptionalSafeForGC;
 		}
 
 #if WITH_EDITOR
@@ -2052,7 +2011,7 @@ public:
 		/** Returns true if this struct would allow the given property to be edited in the details panel. */
 		bool CanEditChange(const FEditPropertyChain& PropertyChain, const void* Data) const
 		{
-			return FakeVPtr->GetHandlingFunc<UE::CoreUObject::Private::EStructOpsFakeVTableFlags::CanEditChange>()(PropertyChain, Data);
+			return HasCanEditChange() && FakeVPtr->GetHandlingFunc<UE::CoreUObject::Private::EStructOpsFakeVTableFlags::CanEditChange>()(PropertyChain, Data);
 		}
 #endif
 
@@ -2065,12 +2024,26 @@ public:
 		/** Structs property visitor signature */
 		EPropertyVisitorControlFlow Visit(FPropertyVisitorContext& Context, const TFunctionRef<EPropertyVisitorControlFlow(const FPropertyVisitorContext& Context)> InFunc) const
 		{
-			return FakeVPtr->GetHandlingFunc<UE::CoreUObject::Private::EStructOpsFakeVTableFlags::Visit>()(Context, InFunc);
+			if (HasVisitor())
+			{
+				return FakeVPtr->GetHandlingFunc<UE::CoreUObject::Private::EStructOpsFakeVTableFlags::Visit>()(Context, InFunc);
+			}
+			else
+			{
+				return EPropertyVisitorControlFlow::StepOver;
+			}
 		}
 
 		void* ResolveVisitedPathInfo(void* Data, const FPropertyVisitorInfo& Info) const
 		{
-			return FakeVPtr->GetHandlingFunc<UE::CoreUObject::Private::EStructOpsFakeVTableFlags::ResolveVisitedPathInfo>()(Data, Info);
+			if (HasVisitor())
+			{
+				return FakeVPtr->GetHandlingFunc<UE::CoreUObject::Private::EStructOpsFakeVTableFlags::ResolveVisitedPathInfo>()(Data, Info);
+			}
+			else
+			{
+				return nullptr;
+			}
 		}
 
 	private:
@@ -4998,4 +4971,3 @@ UE_DECLARE_CORE_VARIANT_TYPE(Ray3,		Ray);
 UE_DECLARE_CORE_VARIANT_TYPE(Sphere3, Sphere);
 
 #undef UE_DECLARE_CORE_VARIANT_TYPE
-
