[Bugfix] Overridden components on an entity return to their default value after undo-redo with impersonation enabled
this occurs because when IDOs are created, the ArMergeOverrides flag caused the root node to get it's override cleared and so the override state wasn't correctly serialized into the transaction buffer.

diff --git a/Engine/Source/Runtime/CoreUObject/Private/UObject/OverridableManager.cpp b/Engine/Source/Runtime/CoreUObject/Private/UObject/OverridableManager.cpp
--- a/Engine/Source/Runtime/CoreUObject/Private/UObject/OverridableManager.cpp
+++ b/Engine/Source/Runtime/CoreUObject/Private/UObject/OverridableManager.cpp
@@ -125,12 +125,20 @@ FOverriddenPropertySet* FOverridableManager::RestoreOverrideOperation(TNotNull<U
 
 #if WITH_EDITORONLY_DATA
 	FOverriddenPropertySet& ObjectOverriddenProperties = FindOrAdd(Object);
+
+	// IDOs need to deserialize a couple times in a row and they use ArMergeOverrides to help. This sets the bDontClearOverrides parameter for this function,
+	// which tells us not to stomp the existing overrides. if bDontClearOverrides is true, it's important that we don't reset the tree or change non-none 
+	// property operations to None.
 	if (!bDontClearOverrides)
 	{
 		ObjectOverriddenProperties.Reset();
 	}
-	ObjectOverriddenProperties.SetOverriddenPropertyOperation(Operation, /*CurrentPropertyChain*/nullptr, /*Property*/nullptr);
-	ObjectOverriddenProperties.bNeedsSubobjectTemplateInstantiation = bNeedsSubobjectTemplateInstantiation;
+
+	if (!bDontClearOverrides || ObjectOverriddenProperties.GetOverriddenPropertyOperation(nullptr, nullptr) == EOverriddenPropertyOperation::None)
+	{
+		ObjectOverriddenProperties.SetOverriddenPropertyOperation(Operation, /*CurrentPropertyChain*/nullptr, /*Property*/nullptr);
+		ObjectOverriddenProperties.bNeedsSubobjectTemplateInstantiation = bNeedsSubobjectTemplateInstantiation;
+	}
 	return &ObjectOverriddenProperties;
 #else
 	if (bNeedsSubobjectTemplateInstantiation)
