[PluginManager]

- Treat commandlets are Editor when determining which leak report severity to use when plugin leaks occur.

#rnx

diff --git a/Engine/Source/Runtime/CoreUObject/Private/Misc/CoreUObjectPluginManager.cpp b/Engine/Source/Runtime/CoreUObject/Private/Misc/CoreUObjectPluginManager.cpp
--- a/Engine/Source/Runtime/CoreUObject/Private/Misc/CoreUObjectPluginManager.cpp
+++ b/Engine/Source/Runtime/CoreUObject/Private/Misc/CoreUObjectPluginManager.cpp
@@ -147,9 +147,13 @@ namespace UE::CoreUObject::Private
 	{
 		TRACE_CPUPROFILER_EVENT_SCOPE(HandlePossibleAssetLeaks);
 
-		// Allow the editor to customize it's own severity during development.
-#if WITH_EDITOR && !UE_BUILD_SHIPPING
+		// Allow the editor/commandlet to customize its own severity during development.
+#if !UE_BUILD_SHIPPING
+#if WITH_EDITOR
 		const int32 LeakedAssetSeverity = GLeakedAssetTrace_EditorSeverity >= 0 ? GLeakedAssetTrace_EditorSeverity : GLeakedAssetTrace_Severity;
+#else
+		const int32 LeakedAssetSeverity = GLeakedAssetTrace_EditorSeverity >= 0 && IsRunningCommandlet() ? GLeakedAssetTrace_EditorSeverity : GLeakedAssetTrace_Severity;
+#endif
 #else
 		const int32 LeakedAssetSeverity = GLeakedAssetTrace_Severity;
 #endif
