Replaced GetCapabilities function with an inline member.

diff --git a/Engine/Source/Runtime/CoreUObject/Public/UObject/Class.h b/Engine/Source/Runtime/CoreUObject/Public/UObject/Class.h
--- a/Engine/Source/Runtime/CoreUObject/Public/UObject/Class.h
+++ b/Engine/Source/Runtime/CoreUObject/Public/UObject/Class.h
@@ -1476,8 +1476,8 @@ namespace UE::CoreUObject::Private
 		// These flags will indicate which functions are present in the HandlingFunctions array.
 		EStructOpsFakeVTableFlags Flags;
 
-		/** returns struct capabilities */
-		FCapabilities (*GetCapabilities)();
+		/** Struct capabilities */
+		FCapabilities Capabilities;
 
 // MSVC doesn't like zero-sized arrays in the middle of structs
 #ifndef _MSC_VER
@@ -1566,7 +1566,7 @@ namespace UE::CoreUObject::Private
 	{
 		{
 			GetStructOpsFakeVTableFlags<CppStruct>(),
-			&GetCapabilities<CppStruct>
+			GetCapabilities<CppStruct>()
 		},
 		{
 #ifdef _MSC_VER
@@ -1584,7 +1584,7 @@ namespace UE::CoreUObject::Private
 	{
 		{
 			GetStructOpsFakeVTableFlags<CppStruct>(),
-			&GetCapabilities<CppStruct>
+			GetCapabilities<CppStruct>()
 		}
 	};
 
@@ -1629,9 +1629,9 @@ public:
 		}
 
 		/** returns struct capabilities */
-		FCapabilities GetCapabilities() const
+		const FCapabilities& GetCapabilities() const
 		{
-			return FakeVPtr->GetCapabilities();
+			return FakeVPtr->Capabilities;
 		}
 
 		/** return true if this class has a no-op constructor and takes EForceInit to init **/
