[GarbageCollector]
- Make sure that pinned UObject are preventing GC of clusters

diff --git a/Engine/Source/Runtime/CoreUObject/Private/UObject/GarbageCollection.cpp b/Engine/Source/Runtime/CoreUObject/Private/UObject/GarbageCollection.cpp
--- a/Engine/Source/Runtime/CoreUObject/Private/UObject/GarbageCollection.cpp
+++ b/Engine/Source/Runtime/CoreUObject/Private/UObject/GarbageCollection.cpp
@@ -4186,7 +4186,7 @@ public:
 					FUObjectItem* RootItem = &GUObjectArray.GetObjectItemArrayUnsafe()[Cluster.RootIndex];
 					if (!RootItem->IsGarbage())
 					{
-						bool bKeepCluster = RootItem->HasAnyFlags(EInternalObjectFlags_RootFlags);
+						bool bKeepCluster = RootItem->HasAnyFlags(EInternalObjectFlags_RootFlags) || RootItem->GetRefCount() > 0;
 						if (bKeepCluster)
 						{
 							FGCFlags::FastMarkAsReachableInterlocked_ForGC(RootItem);
@@ -4199,7 +4199,7 @@ public:
 
 							FGCFlags::FastMarkAsReachableAndClearReachableInClusterInterlocked_ForGC(ClusteredItem);
 
-							if (!bKeepCluster && ClusteredItem->HasAnyFlags(EInternalObjectFlags_RootFlags))
+							if (!bKeepCluster && (ClusteredItem->HasAnyFlags(EInternalObjectFlags_RootFlags) || ClusteredItem->GetRefCount() > 0))
 							{
 								ThreadState.Payload.KeepClusters.Add(RootItem);
 								bKeepCluster = true;
@@ -7282,4 +7282,4 @@ void ProcessAsync(void (*ProcessSync)(void*, FWorkerContext&), void* Processor,
 int32 GetNumCollectReferenceWorkers()
 {
 	return FMath::Clamp(FTaskGraphInterface::Get().GetNumWorkerThreads(), 1, UE::GC::MaxWorkers);
-}
\ No newline at end of file
+}
