Avoid resolving a remote object if it's not been GC'd yet when calling FObjectPtr::GetClass()

diff --git a/Engine/Source/Runtime/CoreUObject/Public/UObject/ObjectHandle.h b/Engine/Source/Runtime/CoreUObject/Public/UObject/ObjectHandle.h
--- a/Engine/Source/Runtime/CoreUObject/Public/UObject/ObjectHandle.h
+++ b/Engine/Source/Runtime/CoreUObject/Public/UObject/ObjectHandle.h
@@ -406,7 +406,7 @@ namespace UE::CoreUObject::Private
 		}
 #elif UE_WITH_REMOTE_OBJECT_HANDLE
 		UObject* Obj = nullptr;
-		if (IsObjectHandleResolved(Handle))
+		if (IsObjectHandleResolved_ForGC(Handle))
 		{
 			Obj = ReadObjectHandlePointerNoCheck(Handle);
 		}
