SavePackage: Fix errors in IllegalReferences diagnostic, by changing FindMostlikelyCulprit to use a recorded instigator from the FPackageHarvester when it harvests an import or export, rather than executing its own slightly different search for the first object in the package that references the bad object.
[RN] Minor, Foundation

diff --git a/Engine/Source/Runtime/CoreUObject/Private/UObject/SavePackage/PackageHarvester.cpp b/Engine/Source/Runtime/CoreUObject/Private/UObject/SavePackage/PackageHarvester.cpp
--- a/Engine/Source/Runtime/CoreUObject/Private/UObject/SavePackage/PackageHarvester.cpp
+++ b/Engine/Source/Runtime/CoreUObject/Private/UObject/SavePackage/PackageHarvester.cpp
@@ -6,6 +6,7 @@
 #include "Cooker/CookDependency.h"
 #include "Interfaces/ITargetPlatform.h"
 #include "Misc/PackageAccessTracking.h"
+#include "Serialization/ArchiveSerializedPropertyChain.h"
 #include "UObject/OverridableManager.h"
 #include "UObject/PropertyPortFlags.h"
 #include "UObject/SavePackage/SaveContext.h"
@@ -759,7 +760,7 @@ void FPackageHarvester::TryHarvestExportInternal(UObject* InObject, UE::SavePack
 		return;
 	}
 
-	HarvestExport(InObject);
+	HarvestExport(InObject, ObjectStatus);
 
 	// @todo: Remove this code once we know that the SubObjectsShadowSerialization works with SceneGraph and it is tested correctly.
 	// Objects that has overridable serialization enabled, needs to separately export their subobject independently of the property serializing it.
@@ -852,7 +853,7 @@ void FPackageHarvester::TryHarvestImport(TObjectPtr<UObject> InObject, UE::SaveP
 		return;
 	}
 
-	HarvestImport(InObject);
+	HarvestImport(InObject, ObjectStatus);
 	ProcessImport(InObject, ObjectStatus);
 }
 
@@ -1307,33 +1308,62 @@ void FPackageHarvester::HarvestSearchableName(TObjectPtr<UObject> TypeObject, FN
 		});
 }
 
-void FPackageHarvester::HarvestExport(UObject* InObject)
+void FPackageHarvester::HarvestExport(UObject* InObject, UE::SavePackageUtilities::FObjectStatus& ObjectStatus)
 {
 	bool bFromOptionalRef = CurrentExportDependencies.CurrentExport &&
 		CurrentExportDependencies.CurrentExport->GetClass()->HasAnyClassFlags(CLASS_Optional);
 	FTaggedExport TaggedExport(InObject, !DoesObjectNeedLoadForEditorGame(InObject), bFromOptionalRef);
 
-	ForEachExportHarvestingRealm([this, &TaggedExport, InObject](ESaveRealm HarvestingRealm)
+	ForEachExportHarvestingRealm([this, &TaggedExport, InObject, &ObjectStatus](ESaveRealm HarvestingRealm)
 		{
 			FHarvestedRealm& RealmData = SaveContext.GetHarvestedRealm(HarvestingRealm);
 			RealmData.AddExport(TaggedExport);
 			RealmData.GetNamesReferencedFromPackageHeader().Add(InObject->GetFName().GetDisplayIndex());
+			SetInstigatorFromProcessingStack(ObjectStatus, HarvestingRealm);
 		});
 	ExportsToProcess.Enqueue({ InObject, nullptr, CurrentExportHarvestingRealms });
 }
 
-void FPackageHarvester::HarvestImport(TObjectPtr<UObject> InObject)
+void FPackageHarvester::HarvestImport(TObjectPtr<UObject> InObject, UE::SavePackageUtilities::FObjectStatus& ObjectStatus)
 {
-	ForEachExportHarvestingRealm([this, InObject](ESaveRealm HarvestingRealm)
+	ForEachExportHarvestingRealm([this, InObject, &ObjectStatus](ESaveRealm HarvestingRealm)
 		{
 			SaveContext.GetHarvestedRealm(HarvestingRealm).AddImport(InObject);
 			if (CurrentExportDependencies.ProcessImportDepth == 0)
 			{
 				SaveContext.GetHarvestedRealm(HarvestingRealm).AddDirectImport(InObject);
 			}
+			SetInstigatorFromProcessingStack(ObjectStatus, HarvestingRealm);
 		});
 }
 
+void FPackageHarvester::SetInstigatorFromProcessingStack(UE::SavePackageUtilities::FObjectStatus& ObjectStatus,
+	ESaveRealm HarvestingRealm)
+{
+	using namespace UE::SavePackageUtilities;
+
+	FRealmInstigator& Instigator = ObjectStatus.RealmInstigator[(uint32)HarvestingRealm];
+	if (Instigator.Object)
+	{
+		return;
+	}
+
+	Instigator.Property = nullptr;
+	if (CurrentExportDependencies.CurrentExport)
+	{
+		Instigator.Object = CurrentExportDependencies.CurrentExport;
+		const FArchiveSerializedPropertyChain* PropertyChain = GetSerializedPropertyChain();
+		if (PropertyChain && PropertyChain->GetNumProperties() > 0)
+		{
+			Instigator.Property = PropertyChain->GetPropertyFromRoot(0);
+		}
+	}
+	else
+	{
+		Instigator.Object = SaveContext.GetPackage();
+	}
+}
+
 #if WITH_VERSE_VM || defined(__INTELLISENSE__)
 void FPackageHarvester::HarvestCellExport(Verse::VCell* InCell)
 {
diff --git a/Engine/Source/Runtime/CoreUObject/Private/UObject/SavePackage/PackageHarvester.h b/Engine/Source/Runtime/CoreUObject/Private/UObject/SavePackage/PackageHarvester.h
--- a/Engine/Source/Runtime/CoreUObject/Private/UObject/SavePackage/PackageHarvester.h
+++ b/Engine/Source/Runtime/CoreUObject/Private/UObject/SavePackage/PackageHarvester.h
@@ -151,11 +151,13 @@ private:
 	void ResolveOverrides();
 
 	void TryHarvestExportInternal(UObject* InObject, UE::SavePackageUtilities::FObjectStatus& ObjectStatus);
-	void HarvestExport(UObject* InObject);
+	void HarvestExport(UObject* InObject, UE::SavePackageUtilities::FObjectStatus& ObjectStatus);
 	void TryHarvestImport(TObjectPtr<UObject> InObject, UE::SavePackageUtilities::FObjectStatus& ObjectStatus);
-	void HarvestImport(TObjectPtr<UObject> InObject);
+	void HarvestImport(TObjectPtr<UObject> InObject, UE::SavePackageUtilities::FObjectStatus& ObjectStatus);
 	/** Add objects/names/others that are referenced even when an object is just being imported */
 	void ProcessImport(TObjectPtr<UObject> InObject, UE::SavePackageUtilities::FObjectStatus& ObjectStatus);
+	void SetInstigatorFromProcessingStack(UE::SavePackageUtilities::FObjectStatus& ObjectStatus,
+		ESaveRealm HarvestingRealm);
 
 #if WITH_VERSE_VM || defined(__INTELLISENSE__)
 	void HarvestCellExport(Verse::VCell* InCell);
diff --git a/Engine/Source/Runtime/CoreUObject/Private/UObject/SavePackage/SaveContext.h b/Engine/Source/Runtime/CoreUObject/Private/UObject/SavePackage/SaveContext.h
--- a/Engine/Source/Runtime/CoreUObject/Private/UObject/SavePackage/SaveContext.h
+++ b/Engine/Source/Runtime/CoreUObject/Private/UObject/SavePackage/SaveContext.h
@@ -125,6 +125,13 @@ enum class ESaveableStatus
 namespace UE::SavePackageUtilities
 {
 
+/** Information about the instigator that caused an object to be added as import or export to a SaveRealm. */
+struct FRealmInstigator
+{
+	UObject* Object = nullptr;
+	FProperty* Property = nullptr;
+};
+
 /** Calculated flags about saveable status and other state for UObjects encountered during save. */
 struct FObjectStatus
 {
@@ -168,6 +175,7 @@ struct FObjectStatus
 		return bIsInSavePackage;
 	}
 
+	FRealmInstigator RealmInstigator[(uint32)ESaveRealm::RealmCount];
 	UObject* SaveableStatusCulprit = nullptr;
 	ESaveableStatus SaveableStatus = ESaveableStatus::Success;
 	ESaveableStatus SaveableStatusCulpritStatus = ESaveableStatus::Success;
diff --git a/Engine/Source/Runtime/CoreUObject/Private/UObject/SavePackage/SavePackageUtilities.cpp b/Engine/Source/Runtime/CoreUObject/Private/UObject/SavePackage/SavePackageUtilities.cpp
--- a/Engine/Source/Runtime/CoreUObject/Private/UObject/SavePackage/SavePackageUtilities.cpp
+++ b/Engine/Source/Runtime/CoreUObject/Private/UObject/SavePackage/SavePackageUtilities.cpp
@@ -329,15 +329,16 @@ EObjectMark GetExcludedObjectMarksForTargetPlatform(const class ITargetPlatform*
  * Find the most likely culprit that caused the objects in the passed in array to be considered illegal for saving.
  *
  * @param	BadObjects				Array of objects that are considered "bad" (e.g. non- RF_Public, in different map package, ...)
+ * @param	SaveContext				SaveContext for the SavePackage call.
+ * @param	HarvestedRealm				Which realm encountered the BadObjects.
  * @param	OutMostLikelyCulprit	UObject that is considered the most likely culprit causing the "bad" objects to be referenced or NULL
  * @param	OutReferencer			UObject referencing the most likely culprit
  * @param	OutReferencerProperty	Property (belonging to referencer) storing the offending reference
  * @param	OutIsCulpritArchetype	Is the most likely culprit an archetype object
- * @param	InOptionalSaveContext	Optional save context
  */
-void FindMostLikelyCulprit(const TArray<UObject*>& BadObjects, 
-						   UObject*& OutMostLikelyCulprit, UObject*& OutReferencer, const FProperty*& OutReferencerProperty, bool& OutIsCulpritArchetype,
-						   FSaveContext* InOptionalSaveContext)
+void FindMostLikelyCulprit(const TArray<UObject*>& BadObjects, FSaveContext& SaveContext,
+	ESaveRealm HarvestedRealm, UObject*& OutMostLikelyCulprit, UObject*& OutReferencer,
+	const FProperty*& OutReferencerProperty, bool& OutIsCulpritArchetype, FString& OutDiagnosticText)
 {
 	UObject* ArchetypeCulprit = nullptr;
 	UObject* ReferencedCulprit = nullptr;
@@ -349,15 +350,8 @@ void FindMostLikelyCulprit(const TArray<UObject*>& BadObjects,
 	OutReferencerProperty = nullptr;
 	OutIsCulpritArchetype = false;
 
-	auto IsObjectIncluded = [InOptionalSaveContext](UObject* InObject)
-	{
-		// if we passed in a SaveContext use that to validate if the object is an import or an export instead of marks
-		if (InOptionalSaveContext)
-		{
-			return InOptionalSaveContext->IsIncluded(InObject);
-		}
-		return InObject->HasAnyMarks(EObjectMark(OBJECTMARK_TagExp | OBJECTMARK_TagImp));
-	};
+	TStringBuilder<1024> DiagnosticText;
+	DiagnosticText << TEXT("SavePackage Invalid Reference Diagnostics:");
 
 	// Iterate over all objects that are marked as unserializable/ bad and print out their referencers.
 	for (int32 BadObjIndex = 0; BadObjIndex < BadObjects.Num(); BadObjIndex++)
@@ -368,7 +362,7 @@ void FindMostLikelyCulprit(const TArray<UObject*>& BadObjects,
 		// for that reason rather than checking references. Class archetypes must be public since instances of their class in other packages can refer to them
 		if (Obj->HasAnyFlags(RF_ArchetypeObject | RF_DefaultSubObject | RF_ClassDefaultObject))
 		{
-			UE_LOG(LogSavePackage, Warning, TEXT("%s is a private Archetype object"), *Obj->GetFullName());
+			DiagnosticText.Appendf(TEXT("\n\t%s is a private Archetype object"), *Obj->GetFullName());
 			TArray<const TCHAR*> Flags;
 			auto AddFlagIfPresent = [&Flags, Obj](EObjectFlags InFlag, const TCHAR* Descriptor)
 			{
@@ -380,7 +374,7 @@ void FindMostLikelyCulprit(const TArray<UObject*>& BadObjects,
 			AddFlagIfPresent(RF_ArchetypeObject, TEXT("RF_ArchetypeObject"));
 			AddFlagIfPresent(RF_ClassDefaultObject, TEXT("RF_ClassDefaultObject"));
 			AddFlagIfPresent(RF_DefaultSubObject, TEXT("RF_DefaultSubObject"));
-			UE_LOG(LogSavePackage, Warning, TEXT("\tThis object is an archetype (flags include %s) but is private. This is a code error from the generator of the object. All archetype objects must be public."),
+			DiagnosticText.Appendf(TEXT("\n\t\tThis object is an archetype (flags include %s) but is private. This is a code error from the generator of the object. All archetype objects must be public."),
 				*FString::Join(Flags, TEXT("|")));
 
 			if (ArchetypeCulprit == nullptr)
@@ -390,40 +384,25 @@ void FindMostLikelyCulprit(const TArray<UObject*>& BadObjects,
 			continue;
 		}
 
-		// Do not run more than a certain number of ref gather if there's a lot of illegal references
-		const int32 MaxNumberOfRefGather = 5;
-		if (BadObjIndex < MaxNumberOfRefGather)
+		FObjectStatus* ObjectStatus = SaveContext.FindCachedObjectStatus(Obj);
+		FRealmInstigator* Instigator = ObjectStatus ? &ObjectStatus->RealmInstigator[(uint32)HarvestedRealm] : nullptr;
+		DiagnosticText.Appendf(TEXT("\n\tInstigator of %s:"), *Obj->GetFullName());
+		if (!Instigator || !Instigator->Object)
 		{
-			UE_LOG(LogSavePackage, Warning, TEXT("\r\nReferencers of %s:"), *Obj->GetFullName());
-			FReferencerInformationList Refs;
-			if (IsReferenced(Obj, RF_Public, EInternalObjectFlags::Native, true, &Refs))
-			{
-				for (int32 i = 0; i < Refs.ExternalReferences.Num(); i++)
-				{
-					UObject* RefObj = Refs.ExternalReferences[i].Referencer;
-					if (IsObjectIncluded(RefObj))
-					{
-						if (RefObj->GetFName() == NAME_PersistentLevel || RefObj->GetClass()->GetFName() == NAME_World)
-						{
-							// these types of references should be ignored
-							continue;
-						}
-
-						CulpritReferencer = RefObj;
-
-						UE_LOG(LogSavePackage, Warning, TEXT("\t%s (%i refs)"), *RefObj->GetFullName(), Refs.ExternalReferences[i].TotalReferences);
-						for (int32 j = 0; j < Refs.ExternalReferences[i].ReferencingProperties.Num(); j++)
-						{
-							const FProperty* Prop = Refs.ExternalReferences[i].ReferencingProperties[j];
-							UE_LOG(LogSavePackage, Warning, TEXT("\t\t%i) %s"), j, *Prop->GetFullName());
-							CulpritReferencerProperty = Prop;
-						}
-
-						// Later ReferencedCulprits are higher priority than earlier culprits. TODO: Not sure if this is an intentional behavior or if they choice was arbitrary.
-						ReferencedCulprit = Obj;
-					}
-				}
-			}
+			DiagnosticText << TEXT("\n\t\tObject: <unknown>");
+			DiagnosticText << TEXT("\n\t\tProperty: <unknown>");
+		}
+		else
+		{
+			CulpritReferencer = Instigator->Object;
+			CulpritReferencerProperty = Instigator->Property;
+			DiagnosticText.Appendf(TEXT("\n\t\tObject: %s"), *Instigator->Object->GetFullName());
+			DiagnosticText.Appendf(TEXT("\n\t\tProperty: %s"),
+				Instigator->Property ? *Instigator->Property->GetPathName() : TEXT("<unknown>"));
+
+			// Later ReferencedCulprits are higher priority than earlier culprits.
+			// TODO: Not sure if this is an intentional behavior or if they choice was arbitrary.
+			ReferencedCulprit = Obj;
 		}
 	}
 
@@ -452,6 +431,7 @@ void FindMostLikelyCulprit(const TArray<UObject*>& BadObjects,
 			}
 		}
 	}
+	OutDiagnosticText = DiagnosticText;
 }
 
 ESavePackageResult FinalizeTempOutputFiles(const FPackagePath& PackagePath, const FSavePackageOutputFileArray& OutputFiles, const FDateTime& FinalTimeStamp)
diff --git a/Engine/Source/Runtime/CoreUObject/Private/UObject/SavePackage/SavePackageUtilities.h b/Engine/Source/Runtime/CoreUObject/Private/UObject/SavePackage/SavePackageUtilities.h
--- a/Engine/Source/Runtime/CoreUObject/Private/UObject/SavePackage/SavePackageUtilities.h
+++ b/Engine/Source/Runtime/CoreUObject/Private/UObject/SavePackage/SavePackageUtilities.h
@@ -29,6 +29,7 @@ class FSavePackageContext;
 class IPackageWriter;
 class ITargetPlatform;
 enum class ESavePackageResult;
+enum class ESaveRealm : uint32;
 struct FIoHash;
 
 // Save Time trace
@@ -175,7 +176,9 @@ ESavePackageResult CreatePayloadSidecarFile(FLinkerSave& Linker, const FPackageP
 void SaveMetaData(UPackage* InOuter, FLinkerSave* Linker, FStructuredArchive::FRecord Record);
 void SaveWorldLevelInfo(UPackage* InOuter, FLinkerSave* Linker, FStructuredArchive::FRecord Record);
 EObjectMark GetExcludedObjectMarksForTargetPlatform(const ITargetPlatform* TargetPlatform);
-void FindMostLikelyCulprit(const TArray<UObject*>& BadObjects, UObject*& OutMostLikelyCulprit, UObject*& OutReferencer, const FProperty*& OutReferencerProperty, bool& OutIsCulpritArchetype, FSaveContext* InOptionalSaveContext = nullptr);
+void FindMostLikelyCulprit(const TArray<UObject*>& BadObjects, FSaveContext& SaveContext,
+	ESaveRealm HarvestedRealm, UObject*& OutMostLikelyCulprit, UObject*& OutReferencer,
+	const FProperty*& OutReferencerProperty, bool& OutIsCulpritArchetype, FString& OutDiagnosticText);
 	
 /** 
 	* Search 'OutputFiles' for output files that were saved to the temp directory and move those files
diff --git a/Engine/Source/Runtime/CoreUObject/Private/UObject/SavePackage2.cpp b/Engine/Source/Runtime/CoreUObject/Private/UObject/SavePackage2.cpp
--- a/Engine/Source/Runtime/CoreUObject/Private/UObject/SavePackage2.cpp
+++ b/Engine/Source/Runtime/CoreUObject/Private/UObject/SavePackage2.cpp
@@ -780,7 +780,8 @@ ESavePackageResult ValidateExports(FSaveContext& SaveContext)
 }
 
 ESavePackageResult ValidateIllegalReferences(FSaveContext& SaveContext, TArray<UObject*>& PrivateObjects,
-	TArray<UObject*>& PrivateContentObjects, TArray<UObject*>& ObjectsInOtherMaps)
+	TArray<UObject*>& PrivateContentObjects, TArray<UObject*>& ObjectsInOtherMaps,
+	ESaveRealm HarvestedRealm)
 {
 	FFormatNamedArguments Args;
 
@@ -810,17 +811,13 @@ ESavePackageResult ValidateIllegalReferences(FSaveContext& SaveContext, TArray<U
 
 	if (ObjectsToCheck)
 	{
-		// Override GIsSavingPackage because we're no longer saving packages at this point and some serialization code can trigger a
-		// a StaticFindObject during referencing finding, which is disallowed during package saving.
-		// 
-		// This will temporarily override the value of GIsSavingPackage by FScopedSavingFlag in InnerSave.
-		TGuardValue<bool> IsSavingPackageGuard(GIsSavingPackage, false);
-
 		UObject* MostLikelyCulprit = nullptr;
 		UObject* Referencer = nullptr;
 		const FProperty* ReferencerProperty = nullptr;
 		bool IsCulpritArchetype;
-		UE::SavePackageUtilities::FindMostLikelyCulprit(*ObjectsToCheck, MostLikelyCulprit, Referencer, ReferencerProperty, IsCulpritArchetype, &SaveContext);
+		FString CulpritDiagnostics;
+		UE::SavePackageUtilities::FindMostLikelyCulprit(*ObjectsToCheck, SaveContext, HarvestedRealm, MostLikelyCulprit,
+			Referencer, ReferencerProperty, IsCulpritArchetype, CulpritDiagnostics);
 
 		UObject* ReferencerOuter = Referencer ? Referencer->GetOuter() : nullptr;
 
@@ -843,6 +840,7 @@ ESavePackageResult ValidateIllegalReferences(FSaveContext& SaveContext, TArray<U
 
 		FString ErrorMessage = FString::Printf(TEXT("%s%s%s"), *ErrorPrologue, *ReferenceDescription, *ErrorEpilogue);
 
+		UE_LOG(LogSavePackage, Warning, TEXT("%s"), *CulpritDiagnostics);
 		if (SaveContext.IsGenerateSaveError())
 		{
 			SaveContext.GetError()->Logf(ELogVerbosity::Warning, TEXT("%s"), *ErrorMessage);
@@ -956,6 +954,7 @@ ESavePackageResult ValidateImports(FSaveContext& SaveContext)
 	TArray<UObject*> PrivateContentObjects;
 	TArray<UObject*> ObjectsInOtherMaps;
 	const TSet<TObjectPtr<UObject>>& Imports = SaveContext.GetImports();
+	ESaveRealm HarvestedRealm = SaveContext.GetCurrentHarvestingRealm();
 	const TSet<TObjectPtr<UObject>>& DirectImports = SaveContext.GetDirectImports();
 	UE_COOK_RESULTPROJECTION_SCOPED(UE::Cook::ResultProjection::PackageAndClass);
 	for (const TObjectPtr<UObject>& Import : Imports)
@@ -1049,7 +1048,7 @@ ESavePackageResult ValidateImports(FSaveContext& SaveContext)
 	}
 	if (PrivateObjects.Num() > 0 || PrivateContentObjects.Num() > 0 || ObjectsInOtherMaps.Num() > 0)
 	{
-		return ValidateIllegalReferences(SaveContext, PrivateObjects, PrivateContentObjects, ObjectsInOtherMaps);
+		return ValidateIllegalReferences(SaveContext, PrivateObjects, PrivateContentObjects, ObjectsInOtherMaps, HarvestedRealm);
 	}
 
 	// Validate External Import Rules
