Fix for all string properties being marked as 'loose'.

This was happening because we would always return EConvertFromTypeResult::Converted in UE_STRPROPERTY_CLASS::ConvertFromType even when the tag and property would match.

UStruct::SerializeTaggedProperties would then mark these converted properties as 'unknown'.

#rnx

diff --git a/Engine/Source/Runtime/CoreUObject/Private/UObject/StrProperty.cpp.inl b/Engine/Source/Runtime/CoreUObject/Private/UObject/StrProperty.cpp.inl
--- a/Engine/Source/Runtime/CoreUObject/Private/UObject/StrProperty.cpp.inl
+++ b/Engine/Source/Runtime/CoreUObject/Private/UObject/StrProperty.cpp.inl
@@ -45,6 +45,11 @@ UE_STRPROPERTY_CLASS::UE_STRPROPERTY_CLASS(UField* InField)
 
 EConvertFromTypeResult UE_STRPROPERTY_CLASS::ConvertFromType(const FPropertyTag& Tag, FStructuredArchive::FSlot Slot, uint8* Data, UStruct* DefaultsStruct, const uint8* Defaults)
 {
+	if (GetClass()->GetFName() == Tag.Type)
+	{
+		return EConvertFromTypeResult::UseSerializeItem;
+	}
+
 	// Convert serialized text to string.
 	if (Tag.Type==NAME_TextProperty)
 	{ 
@@ -52,8 +57,6 @@ EConvertFromTypeResult UE_STRPROPERTY_CLASS::ConvertFromType(const FPropertyTag&
 		Slot << Text;
 		const UE_STRPROPERTY_STRINGTYPE String(FTextInspector::GetSourceString(Text) ? *FTextInspector::GetSourceString(Text) : TEXT(""));
 		SetPropertyValue_InContainer(Data, String, Tag.ArrayIndex);
-
-		return EConvertFromTypeResult::Converted;
 	}
 	else if (Tag.Type==NAME_StrProperty)
 	{
@@ -61,8 +64,6 @@ EConvertFromTypeResult UE_STRPROPERTY_CLASS::ConvertFromType(const FPropertyTag&
 		Slot << SavedStr;
 		const UE_STRPROPERTY_STRINGTYPE ConvertedString(SavedStr);
 		SetPropertyValue_InContainer(Data, ConvertedString, Tag.ArrayIndex);
-
-		return EConvertFromTypeResult::Converted;
 	}
 	else if (Tag.Type==NAME_Utf8StrProperty)
 	{
@@ -70,8 +71,6 @@ EConvertFromTypeResult UE_STRPROPERTY_CLASS::ConvertFromType(const FPropertyTag&
 		Slot << SavedStr;
 		const UE_STRPROPERTY_STRINGTYPE ConvertedString(SavedStr);
 		SetPropertyValue_InContainer(Data, ConvertedString, Tag.ArrayIndex);
-
-		return EConvertFromTypeResult::Converted;
 	}
 	else if (Tag.Type==NAME_AnsiStrProperty)
 	{
@@ -79,11 +78,13 @@ EConvertFromTypeResult UE_STRPROPERTY_CLASS::ConvertFromType(const FPropertyTag&
 		Slot << SavedStr;
 		const UE_STRPROPERTY_STRINGTYPE ConvertedString(SavedStr);
 		SetPropertyValue_InContainer(Data, ConvertedString, Tag.ArrayIndex);
-
-		return EConvertFromTypeResult::Converted;
+	}
+	else
+	{
+		return EConvertFromTypeResult::CannotConvert;
 	}
 
-	return EConvertFromTypeResult::UseSerializeItem;
+	return EConvertFromTypeResult::Converted;
 }
 
 void UE_STRPROPERTY_CLASS::ExportText_Internal( FString& ValueStr, const void* PropertyValueOrContainer, EPropertyPointerType PropertyPointerType, const void* DefaultValue, UObject* Parent, int32 PortFlags, UObject* ExportRootScope ) const
