When creating new IDOs, overrides aren't copied. This causes components to not get saved because the components array isn't overridden

diff --git a/Engine/Source/Runtime/CoreUObject/Private/UObject/InstanceDataObjectUtils.cpp b/Engine/Source/Runtime/CoreUObject/Private/UObject/InstanceDataObjectUtils.cpp
--- a/Engine/Source/Runtime/CoreUObject/Private/UObject/InstanceDataObjectUtils.cpp
+++ b/Engine/Source/Runtime/CoreUObject/Private/UObject/InstanceDataObjectUtils.cpp
@@ -1748,7 +1748,6 @@ UObject* CreateInstanceDataObject(UObject* Owner)
 	TArray<uint8> OwnerData;
 
 	FObjectWriter Writer(OwnerData);
-	Writer.ArNoDelta = true;
 	Owner->SerializeScriptProperties(Writer);
 
 	FObjectReader Reader(OwnerData);
