Addressing some late comments on https://p4-swarm.epicgames.net/reviews/41046371 from Devin
- fixed spelling of Inherit / inheritance
- fixed issue that caused unknown properties to get lost if they came from structs stored in the owner's super-type

diff --git a/Engine/Source/Runtime/CoreUObject/Private/UObject/InstanceDataObjectUtils.cpp b/Engine/Source/Runtime/CoreUObject/Private/UObject/InstanceDataObjectUtils.cpp
--- a/Engine/Source/Runtime/CoreUObject/Private/UObject/InstanceDataObjectUtils.cpp
+++ b/Engine/Source/Runtime/CoreUObject/Private/UObject/InstanceDataObjectUtils.cpp
@@ -1107,11 +1107,9 @@ namespace UE
 		return Property->GetBoolMetaData(NAME_IsLooseMetadata);
 	}
 
-	// constructs an InstanceDataObject struct by merging the properties in 
-	static UStruct* CreateInstanceDataObjectStructRec(const UClass* StructClass, UStruct* TemplateStruct, UObject* Outer, const FPropertyPathNameTree* PropertyTree, const FUnknownEnumNames* EnumNames)
+	// creates a struct that doesn't contain unknown (loose) properties except those in sub-structs
+	static UStruct* CreateInstanceDataObjectSuperHierarchy(const UClass* StructClass, UStruct* TemplateStruct, UObject* Outer, const FPropertyPathNameTree* PropertyTree, const FUnknownEnumNames* EnumNames, TSet<FPropertyPathName>& OutSuperPropertyPathsFromTree)
 	{
-		TSet<FPropertyPathName> SuperPropertyPathsFromTree;
-
 		UStruct* Super = nullptr;
 		if (StructClass->IsChildOf<UClass>())
 		{
@@ -1119,7 +1117,7 @@ namespace UE
 			UClass* TemplateSuper = TemplateStruct ? Cast<UClass>(TemplateStruct->GetSuperStruct()) : nullptr;
 			if (bEnableFullIDOInherritance && TemplateSuper && TemplateSuper != UObject::StaticClass())
 			{
-				Super = CreateInstanceDataObjectClass(nullptr, nullptr, TemplateSuper, Outer);
+				Super = CreateInstanceDataObjectSuperHierarchy(StructClass, TemplateSuper, Outer, PropertyTree, EnumNames, OutSuperPropertyPathsFromTree);
 			}
 			else
 			{
@@ -1127,24 +1125,20 @@ namespace UE
 			}
 		}
 
-
 		if (TemplateStruct)
 		{
-			{
-				const FName SuperName(WriteToString<128>(TemplateStruct->GetName(), TEXTVIEW("_Super")));
-				const UClass* SuperStructClass = StructClass->GetSuperClass();
-				UStruct* NewSuper = NewObject<UStruct>(Outer, SuperStructClass, MakeUniqueObjectName(nullptr, SuperStructClass, SuperName));
-				NewSuper->SetSuperStruct(Super);
-				Super = NewSuper;
-				Super->SetMetaData(NAME_IsInstanceDataObjectStruct, TEXT("True"));
-			}
+			const FName SuperName(WriteToString<128>(TemplateStruct->GetName(), TEXTVIEW("_Super")));
+			const UClass* SuperStructClass = StructClass->GetSuperClass();
+			UStruct* Result = NewObject<UStruct>(Outer, SuperStructClass, MakeUniqueObjectName(nullptr, SuperStructClass, SuperName));
+			Result->SetSuperStruct(Super);
+			Result->SetMetaData(NAME_IsInstanceDataObjectStruct, TEXT("True"));
 
-			// Gather properties for Super Struct
+			// Gather properties
 			TArray<FProperty*> SuperProperties;
 			EFieldIterationFlags ExcludeSuperFlags = (EFieldIterationFlags)((uint8)EFieldIterationFlags::Default & ~(uint8)EFieldIterationFlags::IncludeSuper);
 			for (const FProperty* TemplateProperty : TFieldRange<FProperty>(TemplateStruct, ExcludeSuperFlags))
 			{
-				FProperty* SuperProperty = CastFieldChecked<FProperty>(FField::Duplicate(TemplateProperty, Super));
+				FProperty* SuperProperty = CastFieldChecked<FProperty>(FField::Duplicate(TemplateProperty, Result));
 				SuperProperties.Add(SuperProperty);
 
 				FField::CopyMetaData(TemplateProperty, SuperProperty);
@@ -1156,11 +1150,11 @@ namespace UE
 				if (PropertyTree)
 				{
 					FPropertyPathName Path;
-					Path.Push({TemplateProperty->GetFName(), Type});
+					Path.Push({ TemplateProperty->GetFName(), Type });
 					if (FPropertyPathNameTree::FConstNode Node = PropertyTree->Find(Path))
 					{
 						SubTree = Node.GetSubTree();
-						SuperPropertyPathsFromTree.Add(MoveTemp(Path));
+						OutSuperPropertyPathsFromTree.Add(MoveTemp(Path));
 					}
 				}
 
@@ -1171,12 +1165,25 @@ namespace UE
 			// AddCppProperty expects reverse property order for StaticLink to work correctly
 			for (FProperty* Property : ReverseIterate(SuperProperties))
 			{
-				Super->AddCppProperty(Property);
+				Result->AddCppProperty(Property);
 			}
-			Super->Bind();
-			Super->StaticLink(/*bRelinkExistingProperties*/true);
+			Result->Bind();
+			Result->StaticLink(/*bRelinkExistingProperties*/true);
+			return Result;
 		}
 
+		return Super;
+	}
+
+	// constructs an InstanceDataObject struct by merging the properties in 
+	static UStruct* CreateInstanceDataObjectStructRec(const UClass* StructClass, UStruct* TemplateStruct, UObject* Outer, const FPropertyPathNameTree* PropertyTree, const FUnknownEnumNames* EnumNames)
+	{
+
+		// The super of IDOs map 1-1 with the template struct both in it's properties and it's inherritance chain. 
+		// It can however have unknown properties in it's sub-structs. structs with those will appear in the SuperPropertyPathsFromTree set
+		TSet<FPropertyPathName> SuperPropertyPathsFromTree;
+		UStruct* Super = CreateInstanceDataObjectSuperHierarchy(StructClass, TemplateStruct, Outer, PropertyTree, EnumNames, SuperPropertyPathsFromTree);
+		
 		const FName InstanceDataObjectName = (TemplateStruct) ? FName(WriteToString<128>(TemplateStruct->GetName(), TEXTVIEW("_InstanceDataObject"))) : FName(TEXTVIEW("InstanceDataObject"));
 		UStruct* Result = NewObject<UStruct>(Outer, StructClass, MakeUniqueObjectName(Outer, StructClass, InstanceDataObjectName));
 		Result->SetSuperStruct(Super);
