Moving CLASS_MigratingAsset to RF_MigratingAsset. This allows marking individual objects as 'migrating assets' instead of entire class hierarchies. Added editable AActor::bMigratingAsset property to allow marking individual actors as migrating asset in the editor.

#virtualized

diff --git a/Engine/Source/Programs/Shared/EpicGames.Core/UnrealEngineTypes.cs b/Engine/Source/Programs/Shared/EpicGames.Core/UnrealEngineTypes.cs
--- a/Engine/Source/Programs/Shared/EpicGames.Core/UnrealEngineTypes.cs
+++ b/Engine/Source/Programs/Shared/EpicGames.Core/UnrealEngineTypes.cs
@@ -158,11 +158,6 @@ namespace EpicGames.Core
 		/// </summary>
 		Native = 0x00000080u,
 
-		/// <summary>
-		/// Instances of this class can migrate to another server even if loaded from disk.
-		/// </summary>
-		MigratingAsset = 0x00000100u,
-
 		/// <summary>
 		/// Do not allow users to create in the editor.
 		/// </summary>
@@ -281,7 +276,7 @@ namespace EpicGames.Core
 		/// <summary>
 		/// Flags to inherit from base class
 		/// </summary>
-		Inherit = Transient | Optional | DefaultConfig | Config | PerObjectConfig | ConfigDoNotCheckDefaults | MigratingAsset | NotPlaceable | Const | HasInstancedReference |
+		Inherit = Transient | Optional | DefaultConfig | Config | PerObjectConfig | ConfigDoNotCheckDefaults | NotPlaceable | Const | HasInstancedReference |
 			Deprecated | DefaultToInstanced | GlobalUserConfig | ProjectUserConfig | NeedsDeferredDependencyLoading,
 
 		/// <summary>
@@ -302,7 +297,7 @@ namespace EpicGames.Core
 		/// <summary>
 		/// This is used as a mask for the flags put into generated code for "compiled in" classes.
 		/// </summary>
-		SaveInCompiledInClasses = Abstract | DefaultConfig | GlobalUserConfig | ProjectUserConfig | PerPlatformConfig | Config | Transient | Optional | Native | MigratingAsset | NotPlaceable | PerObjectConfig |
+		SaveInCompiledInClasses = Abstract | DefaultConfig | GlobalUserConfig | ProjectUserConfig | PerPlatformConfig | Config | Transient | Optional | Native | NotPlaceable | PerObjectConfig |
 			ConfigDoNotCheckDefaults | EditInlineNew | CollapseCategories | Interface | DefaultToInstanced | HasInstancedReference | Hidden | Deprecated |
 			HideDropDown | Intrinsic | Const | MinimalAPI | RequiredAPI | MatchedSerializers | NeedsDeferredDependencyLoading,
 	};
diff --git a/Engine/Source/Programs/Shared/EpicGames.UHT/Specifiers/UhtClassSpecifiers.cs b/Engine/Source/Programs/Shared/EpicGames.UHT/Specifiers/UhtClassSpecifiers.cs
--- a/Engine/Source/Programs/Shared/EpicGames.UHT/Specifiers/UhtClassSpecifiers.cs
+++ b/Engine/Source/Programs/Shared/EpicGames.UHT/Specifiers/UhtClassSpecifiers.cs
@@ -64,13 +64,6 @@ namespace EpicGames.UHT.Parsers
 			classObj.RemoveClassFlags(EClassFlags.NotPlaceable);
 		}
 
-		[UhtSpecifier(Extends = UhtTableNames.Class, ValueType = UhtSpecifierValueType.Legacy)]
-		private static void MigratingAssetSpecifier(UhtSpecifierContext specifierContext)
-		{
-			UhtClass classObj = (UhtClass)specifierContext.Type;
-			classObj.AddClassFlags(EClassFlags.MigratingAsset);
-		}
-
 		[UhtSpecifier(Extends = UhtTableNames.Class, ValueType = UhtSpecifierValueType.Legacy)]
 		private static void NotPlaceableSpecifier(UhtSpecifierContext specifierContext)
 		{
diff --git a/Engine/Source/Runtime/CoreUObject/Private/UObject/RemoteObject.cpp b/Engine/Source/Runtime/CoreUObject/Private/UObject/RemoteObject.cpp
--- a/Engine/Source/Runtime/CoreUObject/Private/UObject/RemoteObject.cpp
+++ b/Engine/Source/Runtime/CoreUObject/Private/UObject/RemoteObject.cpp
@@ -799,10 +799,10 @@ FRemoteObjectId FRemoteObjectId::Generate(UObjectBase* InObject, const TCHAR* In
 					bIsClassOrArchetypeSubobject = true;
 					break;
 				}
-				if (Class->HasAnyClassFlags(CLASS_MigratingAsset))
+				if (!!(OuterIt->GetFlags() & RF_MigratingAsset))
 				{
 					bIsMigratingAsset = true;
-				}					
+				}
 			}
 
 			bIsAsset = bIsClassOrArchetypeSubobject || !bIsMigratingAsset;
diff --git a/Engine/Source/Runtime/CoreUObject/Private/UObject/UObjectGlobals.cpp b/Engine/Source/Runtime/CoreUObject/Private/UObject/UObjectGlobals.cpp
--- a/Engine/Source/Runtime/CoreUObject/Private/UObject/UObjectGlobals.cpp
+++ b/Engine/Source/Runtime/CoreUObject/Private/UObject/UObjectGlobals.cpp
@@ -390,7 +390,7 @@ FString LexToString(EObjectFlags Flags)
 		TEXT("Dynamic"),
 		TEXT("WillBeLoaded"),
 		TEXT("HasExternalPackage"),
-		TEXT("PendingKill"),
+		TEXT("MigratingAsset"),
 		TEXT("Garbage"),
 		TEXT("AllocatedInSharedPage"),
 	};
diff --git a/Engine/Source/Runtime/CoreUObject/Public/UObject/ObjectMacros.h b/Engine/Source/Runtime/CoreUObject/Public/UObject/ObjectMacros.h
--- a/Engine/Source/Runtime/CoreUObject/Public/UObject/ObjectMacros.h
+++ b/Engine/Source/Runtime/CoreUObject/Public/UObject/ObjectMacros.h
@@ -204,8 +204,9 @@ enum EClassFlags
 	CLASS_ProjectUserConfig	  = 0x00000040u,
 	/** Class is a native class - native interfaces will have CLASS_Native set, but not RF_MarkAsNative */
 	CLASS_Native			  = 0x00000080u,
-	/** Instances of this class can migrate to another server even if loaded from disk. */
-	CLASS_MigratingAsset	  = 0x00000100u,
+
+	// CLASS_Unused			  = 0x00000100u,
+
 	/** Do not allow users to create in the editor. */
 	CLASS_NotPlaceable        = 0x00000200u,
 	/** Handle object configuration on a per-object basis, rather than per-class. */
@@ -266,7 +267,7 @@ enum EClassFlags
 ENUM_CLASS_FLAGS(EClassFlags);
 
 /** Flags to inherit from base class */
-#define CLASS_Inherit ((EClassFlags)(CLASS_Transient | CLASS_Optional | CLASS_DefaultConfig | CLASS_Config | CLASS_PerObjectConfig | CLASS_ConfigDoNotCheckDefaults | CLASS_MigratingAsset | CLASS_NotPlaceable \
+#define CLASS_Inherit ((EClassFlags)(CLASS_Transient | CLASS_Optional | CLASS_DefaultConfig | CLASS_Config | CLASS_PerObjectConfig | CLASS_ConfigDoNotCheckDefaults | CLASS_NotPlaceable \
 						| CLASS_Const | CLASS_HasInstancedReference | CLASS_Deprecated | CLASS_DefaultToInstanced | CLASS_GlobalUserConfig | CLASS_ProjectUserConfig | CLASS_PerPlatformConfig | CLASS_NeedsDeferredDependencyLoading))
 
 /** These flags will be cleared by the compiler when the class is parsed during script compilation */
@@ -289,7 +290,6 @@ ENUM_CLASS_FLAGS(EClassFlags);
 	CLASS_Transient | \
 	CLASS_Optional | \
 	CLASS_Native | \
-	CLASS_MigratingAsset | \
 	CLASS_NotPlaceable | \
 	CLASS_PerObjectConfig | \
 	CLASS_PerPlatformConfig | \
@@ -580,7 +580,7 @@ enum EObjectFlags
 	RF_ImmutableDefaultObject	=0x04000000,	///< Hidden version of the class default object that cannot and will never be mutated
 	RF_WillBeLoaded				=0x08000000,	///< This object was constructed during load and will be loaded shortly
 	RF_HasExternalPackage		=0x10000000,	///< This object has an external package assigned and should look it up when getting the outermost package
-	// RF_Unused				=0x20000000,
+	RF_MigratingAsset			=0x20000000,	///< Objects marked as MigratingAsset can be migrated to another server even if they've been loaded from disk
 
 	// RF_MirroredGarbage is mirrored in EInternalObjectFlags::Garbage because checking the internal flags is much faster for the Garbage Collector
 	// while checking the object flags is much faster outside of it where the Object pointer is already available and most likely cached.
@@ -601,7 +601,7 @@ inline constexpr EObjectFlags RF_HasPlaceholderType = EObjectFlags(0x20000000);
 #define RF_AllFlags				(EObjectFlags)0xffffffff	///< All flags, used mainly for error checking
 
 /** Flags to load from unreal asset files */
-#define RF_Load						((EObjectFlags)(RF_Public | RF_Standalone | RF_Transactional | RF_ClassDefaultObject | RF_ArchetypeObject | RF_DefaultSubObject | RF_TextExportTransient | RF_InheritableComponentTemplate | RF_DuplicateTransient | RF_NonPIEDuplicateTransient)) 
+#define RF_Load						((EObjectFlags)(RF_Public | RF_Standalone | RF_Transactional | RF_ClassDefaultObject | RF_ArchetypeObject | RF_DefaultSubObject | RF_TextExportTransient | RF_InheritableComponentTemplate | RF_DuplicateTransient | RF_NonPIEDuplicateTransient | RF_MigratingAsset)) 
 
 /** Sub-objects will inherit these flags from their SuperObject */
 #define RF_PropagateToSubObjects	((EObjectFlags)(RF_Public | RF_ArchetypeObject | RF_Transactional | RF_Transient))
diff --git a/Engine/Source/Runtime/Engine/Classes/GameFramework/Actor.h b/Engine/Source/Runtime/Engine/Classes/GameFramework/Actor.h
--- a/Engine/Source/Runtime/Engine/Classes/GameFramework/Actor.h
+++ b/Engine/Source/Runtime/Engine/Classes/GameFramework/Actor.h
@@ -249,7 +249,7 @@ struct TIsContiguousContainer<TInlineComponentArray<T, NumElements>>
  * @see https://docs.unrealengine.com/Programming/UnrealArchitecture/Actors/ActorLifecycle
  * @see UActorComponent
  */
-UCLASS(BlueprintType, Blueprintable, MigratingAsset, config=Engine, meta=(ShortTooltip="An Actor is an object that can be placed or spawned in the world."), MinimalAPI)
+UCLASS(BlueprintType, Blueprintable, config=Engine, meta=(ShortTooltip="An Actor is an object that can be placed or spawned in the world."), MinimalAPI)
 class AActor : public UObject
 {
 	GENERATED_BODY()
@@ -563,6 +563,12 @@ protected:
 	UPROPERTY(Category=Actor, EditAnywhere, AdvancedDisplay)
 	uint8 bCanBeInCluster:1;
 
+#if WITH_EDITORONLY_DATA
+	/** If true, this actor can be migrated to another server even if it's been loaded from disk */
+	UPROPERTY(Category = Actor, EditAnywhere, AdvancedDisplay)
+	uint8 bMigratingAsset : 1;
+#endif
+
 	/**
 	 * If false, the Blueprint ReceiveTick() event will be disabled on dedicated servers.
 	 * @see AllowReceiveTickEventOnDedicatedServer()
diff --git a/Engine/Source/Runtime/Engine/Private/ActorEditor.cpp b/Engine/Source/Runtime/Engine/Private/ActorEditor.cpp
--- a/Engine/Source/Runtime/Engine/Private/ActorEditor.cpp
+++ b/Engine/Source/Runtime/Engine/Private/ActorEditor.cpp
@@ -154,6 +154,7 @@ void AActor::PostEditChangeProperty(FPropertyChangedEvent& PropertyChangedEvent)
 	static const FName Name_RelativeLocation = USceneComponent::GetRelativeLocationPropertyName();
 	static const FName Name_RelativeRotation = USceneComponent::GetRelativeRotationPropertyName();
 	static const FName Name_RelativeScale3D = USceneComponent::GetRelativeScale3DPropertyName();
+	static const FName Name_MigratingAsset = GET_MEMBER_NAME_CHECKED(AActor, bMigratingAsset);
 
 	FProperty* MemberPropertyThatChanged = PropertyChangedEvent.MemberProperty;
 	const FName MemberPropertyName = MemberPropertyThatChanged != NULL ? MemberPropertyThatChanged->GetFName() : NAME_None;
@@ -246,6 +247,18 @@ void AActor::PostEditChangeProperty(FPropertyChangedEvent& PropertyChangedEvent)
 		GEngine->BroadcastOnActorMoved( this );
 	}
 
+	if (MemberPropertyName == Name_MigratingAsset)
+	{
+		if (bMigratingAsset)
+		{
+			SetFlags(RF_MigratingAsset);
+		}
+		else
+		{
+			ClearFlags(RF_MigratingAsset);
+		}
+	}
+
 	FEditorSupportDelegates::UpdateUI.Broadcast();
 	Super::PostEditChangeProperty(PropertyChangedEvent);
 }
