Enabling placeholders

diff --git a/Engine/Source/Runtime/CoreUObject/Private/UObject/InstanceDataObjectUtils.cpp b/Engine/Source/Runtime/CoreUObject/Private/UObject/InstanceDataObjectUtils.cpp
--- a/Engine/Source/Runtime/CoreUObject/Private/UObject/InstanceDataObjectUtils.cpp
+++ b/Engine/Source/Runtime/CoreUObject/Private/UObject/InstanceDataObjectUtils.cpp
@@ -382,7 +382,7 @@ namespace UE
 		TEXT("When enabled, IDOs will be constructed using an archetype chain")
 	);
 
-	bool bEnableIDOPlaceholderObjects = false;
+	bool bEnableIDOPlaceholderObjects = true;
 	FAutoConsoleVariableRef EnableIDOPlaceholderObjectsCVar(
 		TEXT("IDO.Placeholder.Enable"),
 		bEnableIDOPlaceholderObjects,
