IncrementalCook: Auto-collected AssetRegistryQuery dependencies: Set bIncludeOnlyOnDiskAssets=true, and issue a warning if not already set, to prevent indeterminism due to different packages being in memory.
#rnx

diff --git a/Engine/Source/Editor/UnrealEd/Private/Cooker/CookAssetRegistryAccessTracker.cpp b/Engine/Source/Editor/UnrealEd/Private/Cooker/CookAssetRegistryAccessTracker.cpp
--- a/Engine/Source/Editor/UnrealEd/Private/Cooker/CookAssetRegistryAccessTracker.cpp
+++ b/Engine/Source/Editor/UnrealEd/Private/Cooker/CookAssetRegistryAccessTracker.cpp
@@ -6,6 +6,7 @@
 #include "AssetRegistry/AssetData.h"
 #include "AssetRegistry/AssetRegistryModule.h"
 #include "Cooker/CookAssetRegistryAccessTracker_AddPackageDep.h"
+#include "CookOnTheSide/CookLog.h"
 #include "Interfaces/ITargetPlatform.h"
 #include "Misc/PackageAccessTracking.h"
 
@@ -163,6 +164,17 @@ namespace UE::CookAssetRegistryAccessTracker
 			NewRecord.Platform = AccumulatedScopeData->TargetPlatform;
 			NewRecord.bAddPackageDependencies = FAddPackageDependenciesScope::GetAddPackageDependenciesEnabled();
 			NewRecord.Filter = UE::AssetRegistry::ConvertToNonCompiledFilter(CompiledFilter);
+			if (!NewRecord.Filter.bIncludeOnlyOnDiskAssets)
+			{
+				// TODO: Bump this log severity up to warning and enable DumpStackTraceToLog once we have fixed all engine classes.
+				UE_LOG(LogCook, Verbose,
+					TEXT("An AssetRegistry query is executing during a package load or save, and it uses an ARFilter with bIncludeOnlyOnDiskAssets=false. ")
+					TEXT("This will create cookindeterminism since the packages in can vary between cooks. Caller should set bIncludeOnlyOnDiskAssets=true. ")
+					TEXT("For the recorded dependency we are forcing bIncludeOnlyOnDiskAssets=true. Callstack of the AssetRegistryQuery construction: ")
+				);
+				//FDebug::DumpStackTraceToLog(ELogVerbosity::Display);
+				NewRecord.Filter.bIncludeOnlyOnDiskAssets = true;
+			}
 			NewRecord.Filter.SortForSaving();
 
 			Records.Add(NewRecord);
diff --git a/Engine/Source/Editor/UnrealEd/Private/Cooker/CookIncrementalVersion.cpp b/Engine/Source/Editor/UnrealEd/Private/Cooker/CookIncrementalVersion.cpp
--- a/Engine/Source/Editor/UnrealEd/Private/Cooker/CookIncrementalVersion.cpp
+++ b/Engine/Source/Editor/UnrealEd/Private/Cooker/CookIncrementalVersion.cpp
@@ -8,6 +8,6 @@ namespace UE::Cook
 {
 
 // Change CookIncrementalVersion to a new guid when all packages in an incremental cook need to be invalidated
-FGuid CookIncrementalVersion(0x93DE0E1A, 0x721143D5, 0x9487FF18, 0xEF03D6B1);
+FGuid CookIncrementalVersion(0xECD6A91D, 0x830D4F75, 0x9BD16AAA, 0x913CC837);
 
 }
diff --git a/Engine/Source/Runtime/CoreUObject/Private/Cooker/CookDependency.cpp b/Engine/Source/Runtime/CoreUObject/Private/Cooker/CookDependency.cpp
--- a/Engine/Source/Runtime/CoreUObject/Private/Cooker/CookDependency.cpp
+++ b/Engine/Source/Runtime/CoreUObject/Private/Cooker/CookDependency.cpp
@@ -227,6 +227,18 @@ FCookDependency FCookDependency::AssetRegistryQuery(FARFilter Filter)
 {
 	FCookDependency Result(ECookDependency::AssetRegistryQuery);
 	Filter.SortForSaving();
+	if (!Filter.bIncludeOnlyOnDiskAssets)
+	{
+		// TODO: Bump this log severity up to warning and enable DumpStackTraceToLog once we have fixed all engine classes.
+		UE_LOG(LogCore, Verbose,
+			TEXT("An AssetRegistryQuery dependency has an ARFilter with bIncludeOnlyOnDiskAssets=false. ")
+			TEXT("This will create cookindeterminism since the packages in can vary between cooks. Caller should set bIncludeOnlyOnDiskAssets=true. ")
+			TEXT("For the recorded dependency we are forcing bIncludeOnlyOnDiskAssets=true. Callstack of the AssetRegistryQuery construction: ")
+		);
+		//FDebug::DumpStackTraceToLog(ELogVerbosity::Display);
+		Filter.bIncludeOnlyOnDiskAssets = true;
+	}
+
 	Result.ARFilter = MakeUnique<FARFilter>(MoveTemp(Filter));
 	return Result;
 }
@@ -575,6 +587,13 @@ void FCookDependency::UpdateHash(FCookDependencyContext& Context, FBlake3Hash* O
 			if (AssetRegistry)
 			{
 				TArray<FName> PackageNames;
+				if (!ARFilter->bIncludeOnlyOnDiskAssets)
+				{
+					UE_LOG(LogCore, Error,
+						TEXT("An AssetRegistryQuery dependency has an ARFilter with bIncludeOnlyOnDiskAssets=false. ")
+						TEXT("This is supposed to be prevented on construction of FCookDependency::AssetRegistryQuery."));
+					ARFilter->bIncludeOnlyOnDiskAssets = true;
+				}
 				AssetRegistry->EnumerateAssets(*ARFilter, [&Context, &PackageNames](const FAssetData& AssetData)
 					{
 						PackageNames.Add(AssetData.PackageName);
