[UE] Defer invoking the class initializer on duplicated objects until after subobject instancing has occurred.

Notes:
- This addresses an issue where a class initializer function (if implemented) could mutate subobject archetype instance data during UObject duplication.

diff --git a/Engine/Source/Runtime/CoreUObject/Private/UObject/UObjectGlobals.cpp b/Engine/Source/Runtime/CoreUObject/Private/UObject/UObjectGlobals.cpp
--- a/Engine/Source/Runtime/CoreUObject/Private/UObject/UObjectGlobals.cpp
+++ b/Engine/Source/Runtime/CoreUObject/Private/UObject/UObjectGlobals.cpp
@@ -3302,6 +3302,10 @@ UObject* StaticDuplicateObjectEx( FObjectDuplicationParameters& Parameters )
 				}
 
 				DuplicatedObject->CheckDefaultSubobjects();
+
+				// Now that subobjects are instanced, we can invoke the class initializer function on the instance. Note that this will get
+				// deferred during object construction because subobject instancing was initially disabled through the instance graph above.
+				DuplicatedObject->GetClass()->PostInitInstance(DuplicatedObject, &InstanceGraph);
 			}
 		}
 	}
@@ -4373,7 +4377,12 @@ void FObjectInitializer::PostConstructInit()
 		Obj->PostInitProperties();
 	}
 
-	Class->PostInitInstance(Obj, InstanceGraph);
+	// Only invoke the class initializer if subobject instancing is enabled. Since object properties may have been default-initialized to
+	// reference archetypes before instancing, archetypes could end up being mutated by compiled bytecode if executed before fixup occurs.
+	if (bAllowInstancing)
+	{
+		Class->PostInitInstance(Obj, InstanceGraph);
+	}
 
 #if WITH_EDITORONLY_DATA
 	if (Obj->IsTemplate())
