Capture expected warning from parsing deprecated format in an OutputDeviceNull to avoid confusing warnings on Horde.

diff --git a/Engine/Source/Runtime/CoreUObject/Private/UObject/TopLevelAssetPath.cpp b/Engine/Source/Runtime/CoreUObject/Private/UObject/TopLevelAssetPath.cpp
--- a/Engine/Source/Runtime/CoreUObject/Private/UObject/TopLevelAssetPath.cpp
+++ b/Engine/Source/Runtime/CoreUObject/Private/UObject/TopLevelAssetPath.cpp
@@ -6,6 +6,7 @@
 
 #include "Misc/AsciiSet.h"
 #include "Misc/Optional.h"
+#include "Misc/OutputDeviceNull.h"
 #include "Misc/PackageName.h"
 #include "Misc/RedirectCollector.h"
 #include "Serialization/CompactBinaryWriter.h"
@@ -601,11 +602,15 @@ TEST_CASE_NAMED(FTopLevelAssetPathTest, "System::Core::Misc::TopLevelAssetPath",
 	CHECK_FALSE_MESSAGE(TEXT("TrySetPath with malformed path string fails"), MalformedPath.TrySetPath("/Path/To/Package.Asset.Malformed"));
 	CHECK_EQUALS(TEXT("Failed set to string is empty string"), MalformedPath.ToString(), FString());
 
-	FTopLevelAssetPath FromStructuredText;
-	const TCHAR* StructuredTextBuffer = TEXT("(PackageName=\"/Path/To/Package\", AssetName=\"Asset\")");
-	CHECK_MESSAGE(TEXT("FromStructuredText ImportTextItem succeeds"), FromStructuredText.ImportTextItem(StructuredTextBuffer, PPF_None, nullptr, GLog->Get(), nullptr));
-	CHECK_EQUALS(TEXT("FromStructuredText ImportTextItem advances buffer"), *StructuredTextBuffer, '\0');
-	CHECK_EQUALS(TEXT("FromStructuredText imports correct path"), FromStructuredText.ToString(), TEXT("/Path/To/Package.Asset"));
+	{ 
+		// Temporary output device to swallow expected warning
+		FOutputDeviceNull NullOut;
+		FTopLevelAssetPath FromStructuredText;
+		const TCHAR* StructuredTextBuffer = TEXT("(PackageName=\"/Path/To/Package\", AssetName=\"Asset\")");
+		CHECK_MESSAGE(TEXT("FromStructuredText ImportTextItem succeeds"), FromStructuredText.ImportTextItem(StructuredTextBuffer, PPF_None, nullptr, &NullOut, nullptr));
+		CHECK_EQUALS(TEXT("FromStructuredText ImportTextItem advances buffer"), *StructuredTextBuffer, '\0');
+		CHECK_EQUALS(TEXT("FromStructuredText imports correct path"), FromStructuredText.ToString(), TEXT("/Path/To/Package.Asset"));
+	}
 
 	FTopLevelAssetPath FromExportTextPath;
 	const TCHAR* ExportTextBuffer = TEXT("/Script/Module.ClassName'/Path/To/Package.Asset'");
