Free remote object related memory on exit

diff --git a/Engine/Source/Runtime/CoreUObject/Private/UObject/Obj.cpp b/Engine/Source/Runtime/CoreUObject/Private/UObject/Obj.cpp
--- a/Engine/Source/Runtime/CoreUObject/Private/UObject/Obj.cpp
+++ b/Engine/Source/Runtime/CoreUObject/Private/UObject/Obj.cpp
@@ -5743,6 +5743,7 @@ void StaticExit()
 	FDeferredMessageLog::Cleanup();	
 	CleanupLinkerAnnotations();
 	CleanupCachedArchetypes();
+	UE::RemoteObject::Private::ShutdownRemoteObjects();
 
 	UE_LOG(LogExit, Log, TEXT("Object subsystem successfully closed.") );
 }
diff --git a/Engine/Source/Runtime/CoreUObject/Private/UObject/RemoteObject.cpp b/Engine/Source/Runtime/CoreUObject/Private/UObject/RemoteObject.cpp
--- a/Engine/Source/Runtime/CoreUObject/Private/UObject/RemoteObject.cpp
+++ b/Engine/Source/Runtime/CoreUObject/Private/UObject/RemoteObject.cpp
@@ -75,6 +75,14 @@ public:
 	using Super::begin;
 	using Super::end;
 
+	virtual ~FRemoteObjectStubMap()
+	{
+		for (TPair<FRemoteObjectId, UE::RemoteObject::Handle::FRemoteObjectStub*>& Pair : *this)
+		{
+			delete Pair.Value;
+		}
+	}
+
 	UE::RemoteObject::Handle::FRemoteObjectStub*& FindOrAdd(FRemoteObjectId Id)
 	{		
 		return Super::FindOrAdd(UE::RemoteObject::Private::FRemoteIdLocalizationHelper::GetLocalized(Id));
@@ -309,6 +317,12 @@ void InitRemoteObjects()
 #endif
 }
 
+void ShutdownRemoteObjects()
+{
+	delete ObjectMaps;
+	ObjectMaps = nullptr;
+}
+
 void RegisterRemoteObjectId(FRemoteObjectId ObjectId, FRemoteServerId ResidentServerId)
 {
 	ObjectMaps->FindOrAddRemoteObjectStub(ObjectId, ResidentServerId);
diff --git a/Engine/Source/Runtime/CoreUObject/Private/UObject/RemoteObjectPrivate.h b/Engine/Source/Runtime/CoreUObject/Private/UObject/RemoteObjectPrivate.h
--- a/Engine/Source/Runtime/CoreUObject/Private/UObject/RemoteObjectPrivate.h
+++ b/Engine/Source/Runtime/CoreUObject/Private/UObject/RemoteObjectPrivate.h
@@ -19,6 +19,11 @@ namespace UE::RemoteObject::Private
 	*/
 	void InitRemoteObjects();
 
+	/**
+	* Frees memory associated with remote objects subsystems
+	*/
+	void ShutdownRemoteObjects();
+
 	/**
 	* Marks object memory as remote and creates its stub
 	*/
