[PropertyObject]
- Leverage lazy loading when possible instead of doing in-place loading during type conversion that can cause partial loading and/or deadlock

diff --git a/Engine/Source/Runtime/CoreUObject/Private/UObject/PropertyObject.cpp b/Engine/Source/Runtime/CoreUObject/Private/UObject/PropertyObject.cpp
--- a/Engine/Source/Runtime/CoreUObject/Private/UObject/PropertyObject.cpp
+++ b/Engine/Source/Runtime/CoreUObject/Private/UObject/PropertyObject.cpp
@@ -11,6 +11,8 @@
 #include "UObject/PropertyHelper.h"
 #include "UObject/RemoteObjectTransfer.h"
 #include "UObject/CoreNet.h"
+#include "UObject/ObjectPathId.h"
+#include "UObject/PackedObjectRef.h"
 
 /*-----------------------------------------------------------------------------
 	FObjectProperty.
@@ -57,6 +59,7 @@ FString FObjectProperty::GetCPPMacroType( FString& ExtendedTypeText ) const
 
 EConvertFromTypeResult FObjectProperty::ConvertFromType(const FPropertyTag& Tag, FStructuredArchive::FSlot Slot, uint8* Data, UStruct* DefaultsStruct, const uint8* Defaults)
 {
+	using namespace UE::CoreUObject::Private;
 	static FName NAME_AssetObjectProperty = "AssetObjectProperty"; // old name of soft object property
 
 	if (Tag.Type == NAME_SoftObjectProperty || Tag.Type == NAME_AssetObjectProperty)
@@ -67,6 +70,7 @@ EConvertFromTypeResult FObjectProperty::ConvertFromType(const FPropertyTag& Tag,
 
 		UObject* PreviousValueObj = nullptr;
 
+		bool bLazyResolve = false;
 		// If we're async loading, only zenloader supports synchronous loads from the async loading thread.
 		// Other loaders are not safe to do a sync load because they may crash or fail to set the variable, so throw an error if it's not already in memory
 		ELoaderType LoaderType = GetLoaderType();
@@ -80,6 +84,28 @@ EConvertFromTypeResult FObjectProperty::ConvertFromType(const FPropertyTag& Tag,
 				{
 					UE_LOG(LogClass, Error, TEXT("Failed to convert soft path %s to unloaded object as this is not supported during async loading with the currently active loader '%s'. Load and resave %s in the editor to fix!"), *PreviousValue.ToString(), LexToString(LoaderType), *Slot.GetUnderlyingArchive().GetArchiveName());
 				}
+#if WITH_EDITOR
+				// Only for the parallel loading path first so that we can do some opt-in testing
+				else if (IsInParallelLoadingThread())
+				{
+					const FSoftObjectPath SoftObjectPath = PreviousValue.ToSoftObjectPath();
+
+					const FObjectRef ImportRef(SoftObjectPath.GetLongPackageFName(), NAME_None, NAME_None, FObjectPathId(PreviousValue.GetAssetName()));
+					const FPackedObjectRef PackedObjectRef = MakePackedObjectRef(ImportRef);
+					PreviousValueObj = (UObject*)PackedObjectRef.EncodedRef;
+
+					UE_LOG(LogProperty, Log, TEXT("ConvertFromType is replacing softobjectpath %s with lazyref %p"), *SoftObjectPath.GetLongPackageFName().ToString(), PreviousValueObj);
+
+					SoftObjectPath.LoadAsync(
+						FLoadSoftObjectPathAsyncDelegate(),
+						FLoadAssetAsyncOptionalParams
+						{
+							.PackagePriority = MAX_int32
+						}
+					);
+					bLazyResolve = true;
+				}
+#endif
 				else
 				{
 					PreviousValueObj = PreviousValue.LoadSynchronous();
@@ -94,8 +120,11 @@ EConvertFromTypeResult FObjectProperty::ConvertFromType(const FPropertyTag& Tag,
 		// Now copy the value into the object's address space
 		SetPropertyValue_InContainer(Data, PreviousValueObj, Tag.ArrayIndex);
 
-		// Validate the type is proper
-		CheckValidObject(GetPropertyValuePtr_InContainer(Data, Tag.ArrayIndex), PreviousValueObj, Defaults);
+		if (!bLazyResolve)
+		{
+			// Validate the type is proper
+			CheckValidObject(GetPropertyValuePtr_InContainer(Data, Tag.ArrayIndex), PreviousValueObj, Defaults);
+		}
 
 		return EConvertFromTypeResult::Converted;
 	}
