Store the class of a remote object in its stub to avoid resolving the object when checking its class with TObjectPtr's utility functions (see FObjectPtr::GetClass())

diff --git a/Engine/Source/Runtime/CoreUObject/Private/UObject/RemoteObject.cpp b/Engine/Source/Runtime/CoreUObject/Private/UObject/RemoteObject.cpp
--- a/Engine/Source/Runtime/CoreUObject/Private/UObject/RemoteObject.cpp
+++ b/Engine/Source/Runtime/CoreUObject/Private/UObject/RemoteObject.cpp
@@ -100,11 +100,21 @@ public:
 class FRemoteObjectMaps
 {
 	mutable FTransactionallySafeCriticalSection ObjectMapCritical;
+	// Maps remote object id to to a stub. Note that at the moment stubs are never destroyed (this is required by FRemoteObjectHandlePrivate)
 	FRemoteObjectStubMap RemoteObjects;
-	TMap<FRemoteObjectId, FRemoteObjectPathName> AssetPaths;
+	// Maps remote object id to to a pathname. Note that at the moment these pathnames are never destroyed (this is required by FRemoteObjectClass)
+	TMap<FRemoteObjectId, FRemoteObjectPathName*> AssetPaths;
 
 public:
 
+	virtual ~FRemoteObjectMaps()
+	{
+		for (TPair<FRemoteObjectId, FRemoteObjectPathName*>& Pair : AssetPaths)
+		{
+			delete Pair.Value;
+		}
+
+	}
 	UE::RemoteObject::Handle::FRemoteObjectStub* FindRemoteObjectStub(FRemoteObjectId Id)
 	{
 		using namespace UE::RemoteObject::Handle;
@@ -148,6 +158,11 @@ public:
 			FRemoteServerId ObjectServerId = Stub->Id.GetServerId();
 			Stub->OwningServerId = ObjectServerId.IsValid() ? ObjectServerId : FRemoteServerId::GetLocalServerId();
 		}
+		else if (!Stub->Class.IsValid())
+		{
+			// Always make sure we store the object class in the stub (it's possible the stub was created only with a remote id in which case the class wouldn't be stored)
+			Stub->Class = FRemoteObjectClass(Object->GetClass());
+		}
 
 		FUObjectItem* ObjectItem = GUObjectArray.ObjectToObjectItem(Object);
 		checkf(ObjectItem, TEXT("Attempting to get a serial number for an object that does not exist in the global UObject array (it's possible GUObjectArray is not initialized yet, ObjectIndex=%d)"), GUObjectArray.ObjectToIndex(Object));
@@ -172,16 +187,21 @@ public:
 #endif
 	}
 
-	void StoreAssetPath(UObject* InObject)
+	FRemoteObjectPathName* StoreAssetPath(UObject* InObject)
 	{
 		FRemoteObjectId ObjectId(InObject);
-		FRemoteObjectPathName& Path = AssetPaths.FindOrAdd(ObjectId);
-		Path = FRemoteObjectPathName(InObject);
+		FRemoteObjectPathName*& Path = AssetPaths.FindOrAdd(ObjectId);
+		if (!Path)
+		{
+			Path = new FRemoteObjectPathName(InObject);
+		}
+		return Path;
 	}
 
 	FRemoteObjectPathName* FindAssetPath(FRemoteObjectId ObjectId)
 	{
-		return AssetPaths.Find(ObjectId);
+		FRemoteObjectPathName** PathName = AssetPaths.Find(ObjectId);
+		return PathName ? *PathName : nullptr;
 	}
 
 	/**
@@ -447,14 +467,53 @@ bool IsSafeToMigrateObjects()
 namespace UE::RemoteObject::Handle
 {
 
+FRemoteObjectClass::FRemoteObjectClass(UClass* InClass)
+{
+	checkf(InClass, TEXT("FRemoteClassStub requires a valid class"));
+	
+	if (InClass->IsNative())
+	{
+		// Native classes are never GC'd so we can just store a raw pointer to a class object
+		PathNameOrClass = UPTRINT(InClass);
+	}
+	else
+	{
+		// Blueprints and Verse classes are assets that can be GC'd after an object of such class is sent to another server
+		// Since we don't want a strong reference to an asset class we store a pathname instead (here we take advantage of the fact that pathnames are never destroyed, see FRemoteObjectMaps)
+		// We could cache the class object here for faster access but it wouldn't remove the requirement of storing its pathname (because the class may get GC'd and we want to be able to load it)
+		// An alternative to this approach would be to store the class object's remote id but that would require one more lookup (id -> pathname)
+		// It's also important to keep this structure small so that FRemoteObjectStub is also small because stubs are never destroyed atm
+		PathNameOrClass = UPTRINT(UE::RemoteObject::Private::ObjectMaps->StoreAssetPath(InClass)) | 1;
+	}
+}
+
+UClass* FRemoteObjectClass::GetClass() const
+{
+	if (IsNative())
+	{
+		return reinterpret_cast<UClass*>(PathNameOrClass);
+	}
+	else
+	{
+		const FRemoteObjectPathName* PathName = reinterpret_cast<const FRemoteObjectPathName*>(PathNameOrClass & ~UPTRINT(1));
+		// Resolve() will first attempt to find the class in memory and if it doesn't exist, load it.
+		return Cast<UClass>(PathName->Resolve());
+	}
+}
+
 FRemoteObjectStub::FRemoteObjectStub(UObject* Object)
 {
 	using namespace UE::CoreUObject::Private;
 
-	Id = FObjectHandleUtils::GetRemoteId(Object);
-	if (UObject* Outer = FObjectHandleUtils::GetNonAccessTrackedOuterNoResolve(Object))
+	if (Object)
 	{
-		OuterId = FObjectHandleUtils::GetRemoteId(Outer);
+		Id = FObjectHandleUtils::GetRemoteId(Object);
+		Class = FRemoteObjectClass(Object->GetClass());
+
+		if (UObject* Outer = FObjectHandleUtils::GetNonAccessTrackedOuterNoResolve(Object))
+		{
+			OuterId = FObjectHandleUtils::GetRemoteId(Outer);
+		}
 	}
 }
 
diff --git a/Engine/Source/Runtime/CoreUObject/Public/UObject/ObjectHandle.h b/Engine/Source/Runtime/CoreUObject/Public/UObject/ObjectHandle.h
--- a/Engine/Source/Runtime/CoreUObject/Public/UObject/ObjectHandle.h
+++ b/Engine/Source/Runtime/CoreUObject/Public/UObject/ObjectHandle.h
@@ -412,7 +412,15 @@ namespace UE::CoreUObject::Private
 		}
 		else
 		{
-			Obj = ResolveObjectHandle(Handle);
+			const UE::RemoteObject::Handle::FRemoteObjectStub* Stub = Handle.ToStub();
+			if (UClass* Class = Stub->Class.GetClass())
+			{
+				return Class;
+			}
+			else
+			{
+				Obj = ResolveObjectHandle(Handle);
+			}
 		}
 		return Obj != nullptr ? UE::CoreUObject::Private::GetClass(Obj) : nullptr;
 #else
diff --git a/Engine/Source/Runtime/CoreUObject/Public/UObject/RemoteObject.h b/Engine/Source/Runtime/CoreUObject/Public/UObject/RemoteObject.h
--- a/Engine/Source/Runtime/CoreUObject/Public/UObject/RemoteObject.h
+++ b/Engine/Source/Runtime/CoreUObject/Public/UObject/RemoteObject.h
@@ -7,6 +7,7 @@
 #include "UObject/ObjectMacros.h"
 
 class UObject;
+class UClass;
 
 /**
  * Remote objects are unique UObjects that are referenced by a local server instance but their memory is actually owned by (exists on) another server.
@@ -88,6 +89,29 @@ namespace UE::RemoteObject::Handle
 		return HashCombine(::GetTypeHash(Key.PhysicsServerId.GetIdNumber()), ::GetTypeHash(Key.PhysicsLocalIslandId));
 	}
 
+	/**
+	* Structure that holds a pointer to a native class or to a pathname of a class on disk (Blueprint etc)
+	*/
+	struct FRemoteObjectClass
+	{
+		UPTRINT PathNameOrClass = 0;
+
+		FRemoteObjectClass() = default;
+		COREUOBJECT_API explicit FRemoteObjectClass(UClass* InClass);
+
+		bool IsNative() const
+		{
+			return !(PathNameOrClass & 1);
+		}
+
+		bool IsValid() const
+		{
+			return !!PathNameOrClass;
+		}
+		
+		COREUOBJECT_API UClass* GetClass() const;
+	};
+
 	// Structure that holds basic information about a remote object
 	// This is what FObjectPtr that references a remote object actually points to after the remote object's memory has been claimed by GC
 	struct FRemoteObjectStub
@@ -95,6 +119,7 @@ namespace UE::RemoteObject::Handle
 		FRemoteObjectId Id;
 		FRemoteObjectId OuterId;
 		FName Name;
+		FRemoteObjectClass Class;
 
 		/** SerialNumber this object had on this server */
 		int32 SerialNumber = 0;
