[UObjectGlobals]
- Fix ClassUnique thread-safety and prevent the monotically incremented value from being allowed to rollback

diff --git a/Engine/Source/Runtime/CoreUObject/Private/UObject/UObjectGlobals.cpp b/Engine/Source/Runtime/CoreUObject/Private/UObject/UObjectGlobals.cpp
--- a/Engine/Source/Runtime/CoreUObject/Private/UObject/UObjectGlobals.cpp
+++ b/Engine/Source/Runtime/CoreUObject/Private/UObject/UObjectGlobals.cpp
@@ -2848,15 +2848,21 @@ FName MakeUniqueObjectName(UObject* Parent, const UClass* Class, FName InBaseNam
 			static const FName NamePackage(NAME_Package);
 			if (BaseName == NamePackage)
 			{
+				int32 NextClassUnique = 0;
+				UE_AUTORTFM_OPEN
+				{
+					NextClassUnique = FPlatformAtomics::InterlockedIncrement((volatile int32*)&Class->ClassUnique);
+				};
+
 				if (Parent == NULL)
 				{
 					//package names should default to "/Temp/Untitled" when their parent is NULL. Otherwise they are a group.
-					TestName = FName(*FString::Printf(TEXT("/Temp/%s"), LexToString(NAME_Untitled)), ++Class->ClassUnique);
+					TestName = FName(*FString::Printf(TEXT("/Temp/%s"), LexToString(NAME_Untitled)), NextClassUnique);
 				}
 				else
 				{
 					//package names should default to "Untitled"
-					TestName = FName(NAME_Untitled, ++Class->ClassUnique);
+					TestName = FName(NAME_Untitled, NextClassUnique);
 				}
 			}
 			else
@@ -2893,7 +2899,7 @@ FName MakeUniqueObjectName(UObject* Parent, const UClass* Class, FName InBaseNam
 					}
 					else
 					{
-						NameNumber = ++Class->ClassUnique;
+						NameNumber = FPlatformAtomics::InterlockedIncrement((volatile int32*)&Class->ClassUnique);
 					}
 				};
 				TestName = FName(BaseName, NameNumber);
