[ObjectAnnotation]
- Fix potential deadlock between multiple threads handling annotations because of 2 different locks being used in different order
- Fix race condition where we could end up registering twice or unregistering while not being registered
- Add unittest for object annotation thread-safety

diff --git a/Engine/Source/Programs/LowLevelTests/FoundationTests/Tests/UObjectAnnotationTests.cpp b/Engine/Source/Programs/LowLevelTests/FoundationTests/Tests/UObjectAnnotationTests.cpp
new file mode 100644
--- /dev/null
+++ b/Engine/Source/Programs/LowLevelTests/FoundationTests/Tests/UObjectAnnotationTests.cpp
@@ -0,0 +1,81 @@
+// Copyright Epic Games, Inc. All Rights Reserved.
+
+#include "CoreTypes.h"
+
+#include "TestCommon/Expectations.h"
+#include "TestHarness.h"
+#include "TestObject.h"
+#include "UObject/Class.h"
+#include "UObject/Package.h"
+#include "UObject/UObjectAnnotation.h"
+#include "Tasks/Task.h"
+#include "Async/ParallelFor.h"
+#include "Async/ManualResetEvent.h"
+
+TEST_CASE("UObjectAnnotation::ThreadSafety", "[UObjectAnnotation]")
+{
+	FUObjectAnnotationSparseBool AnnotationTest;
+	FGCObject::StaticInit();
+
+	UE::FManualResetEvent DoneEvent;
+	UE::Tasks::FTask WatchdogTask =
+		UE::Tasks::Launch(
+			TEXT("WatchdogTask"),
+			[&DoneEvent]()
+			{
+				double StartTime = FPlatformTime::Seconds();
+				while (!DoneEvent.WaitFor(UE::FMonotonicTimeSpan::FromSeconds(1)))
+				{
+					if (FPlatformTime::Seconds() - StartTime > 10.0f)
+					{
+						UE_LOG(LogTemp, Fatal, TEXT("UObjectAnnotation::ThreadSafety is non responsive and likely dead-locked. Aborting program."));
+					}
+				}
+			}
+		);
+
+	ParallelFor(16,
+		[&AnnotationTest](int32)
+		{
+			TStrongObjectPtr<UPackage> Package = nullptr;
+			{
+				FGCScopeGuard Guard;
+				Package = TStrongObjectPtr<UPackage>(NewObject<UPackage>());
+			}
+
+			for (int32 Index = 0; Index < 1000; ++Index)
+			{
+				AnnotationTest.Set(Package.Get());
+				CHECK(AnnotationTest.Get(Package.Get()));
+
+				FPlatformProcess::YieldCycles(FMath::RandRange(0, 10000));
+
+				if (IsInGameThread())
+				{
+					UObject* OldPackage = Package.Get();
+
+					// Replace by a new object and garbage collect the old one
+					Package = TStrongObjectPtr<UPackage>(NewObject<UPackage>());
+					CollectGarbage(GARBAGE_COLLECTION_KEEPFLAGS);
+
+					// Garbage collection should have unregistered the annotation automatically.
+					CHECK(!AnnotationTest.Get(OldPackage));
+				}
+				else
+				{
+					AnnotationTest.Clear(Package.Get());
+					CHECK(!AnnotationTest.Get(Package.Get()));
+				}
+
+				FPlatformProcess::YieldCycles(FMath::RandRange(0, 10000));
+			}
+
+			Package->ClearInternalFlags(EInternalObjectFlags::Async);
+		}
+	);
+
+	CollectGarbage(GARBAGE_COLLECTION_KEEPFLAGS);
+
+	DoneEvent.Notify();
+	WatchdogTask.Wait();
+}
diff --git a/Engine/Source/Runtime/CoreUObject/Public/UObject/UObjectAnnotation.h b/Engine/Source/Runtime/CoreUObject/Public/UObject/UObjectAnnotation.h
--- a/Engine/Source/Runtime/CoreUObject/Public/UObject/UObjectAnnotation.h
+++ b/Engine/Source/Runtime/CoreUObject/Public/UObject/UObjectAnnotation.h
@@ -9,6 +9,7 @@
 #include "AutoRTFM.h"
 #include "UObject/UObjectArray.h"
 #include "UObject/UObjectBaseUtility.h"
+#include "Misc/ScopeExit.h"
 #include "Misc/ScopeLock.h"
 #include "Misc/ScopeRWLock.h"
 #include "Misc/TransactionallySafeCriticalSection.h"
@@ -92,6 +93,7 @@ private:
 	void AddAnnotationInternal(const UObjectBase* Object, T&& Annotation)
 	{
 		check(Object);
+		bool bProcessRegistration = false;
 		TAnnotation LocalAnnotation = Forward<T>(Annotation);
 		if (LocalAnnotation.IsDefault())
 		{
@@ -112,12 +114,46 @@ private:
 				if (bAutoRemove)
 #endif
 				{
-					GUObjectArray.AddUObjectDeleteListener(this);
+					bProcessRegistration = true;
 				}
 			}
 		}
+
+		if (bProcessRegistration)
+		{
+			ProcessRegistration();
+		}
 	}
 
+	void ProcessRegistration()
+	{
+		// Always take the UObjectDeleteListeners mutex first to preserve lock order in all case
+		// and avoid deadlock. This is important because we need to follow this lock flow :
+		// ConditionalFinishDestroy -> RemoveObjectFromDeleteListeners (UObjectDeleteListenersCritical) -> NotifyUObjectDeleted -> RemoveAnnotation (AnnotationMapCritical)
+		GUObjectArray.LockUObjectDeleteListeners();
+		ON_SCOPE_EXIT
+		{
+			GUObjectArray.UnlockUObjectDeleteListeners();
+		};
+
+		UE::TWriteScopeLock AnnotationMapLock(AnnotationMapCritical);
+		if (AnnotationMap.Num() == 0)
+		{
+			if (bRegistered)
+			{
+				GUObjectArray.RemoveUObjectDeleteListener(this);
+				bRegistered = false;
+			}
+		}
+		else
+		{
+			if (!bRegistered)
+			{
+				GUObjectArray.AddUObjectDeleteListener(this);
+				bRegistered = true;
+			}
+		}
+	}
 public:
 	/**
 	 * Add an annotation to the annotation list. If the Annotation is the default, then the annotation is removed.
@@ -143,30 +179,33 @@ public:
 	TAnnotation GetAndRemoveAnnotation(const UObjectBase *Object)
 	{		
 		check(Object);
-		bool bHadElements = false;
-		bool bIsNowEmpty = false;
+		bool bProcessRegistration = false;
 
-		// Avoid holding the lock while we call GUObjectArray.RemoveUObjectDeleteListener as it could deadlock
 		TAnnotation Result;
 		{
 			SetAnnotationCacheKeyAndValue(Object, TAnnotation());
 
 			UE::TRWScopeLock AnnotationMapLock(AnnotationMapCritical, SLT_Write);
-			bHadElements = (AnnotationMap.Num() > 0);
+			const bool bHadElements = (AnnotationMap.Num() > 0);
 			AnnotationMap.RemoveAndCopyValue(Object, Result);
-			bIsNowEmpty = (AnnotationMap.Num() == 0);
-		}
+			const bool bIsNowEmpty = (AnnotationMap.Num() == 0);
 
-		if (bHadElements && bIsNowEmpty)
-		{
-			// we are removing the last one, so if we are auto removing or verifying removal, unregister now
+			if (bHadElements && bIsNowEmpty)
+			{
+				// we are removing the last one, so if we are auto removing or verifying removal, unregister now
 #if (UE_BUILD_SHIPPING || UE_BUILD_TEST)
-			if (bAutoRemove)
+				if (bAutoRemove)
 #endif
-			{
-				GUObjectArray.RemoveUObjectDeleteListener(this);
+				{
+					bProcessRegistration = true;
+				}
 			}
 		}
+
+		if (bProcessRegistration)
+		{
+			ProcessRegistration();
+		}
 		return Result;
 	}
 	/**
@@ -177,27 +216,32 @@ public:
 	void RemoveAnnotation(const UObjectBase *Object)
 	{
 		check(Object);
-		bool bHadElements = false;
-		bool bIsNowEmpty = false;
-		// Avoid holding the lock while we call GUObjectArray.RemoveUObjectDeleteListener as it could deadlock
+		bool bProcessRegistration = false;
+
 		{
 			SetAnnotationCacheKeyAndValue(Object, TAnnotation());
 
 			UE::TRWScopeLock AnnotationMapLock(AnnotationMapCritical, SLT_Write);
-			bHadElements = (AnnotationMap.Num() > 0);
+			const bool bHadElements = (AnnotationMap.Num() > 0);
 			AnnotationMap.Remove(Object);
-			bIsNowEmpty = (AnnotationMap.Num() == 0);
-		}
-		if (bHadElements && bIsNowEmpty)
-		{
-			// we are removing the last one, so if we are auto removing or verifying removal, unregister now
+			const bool bIsNowEmpty = (AnnotationMap.Num() == 0);
+
+			if (bHadElements && bIsNowEmpty)
+			{
+				// we are removing the last one, so if we are auto removing or verifying removal, unregister now
 #if (UE_BUILD_SHIPPING || UE_BUILD_TEST)
-			if (bAutoRemove)
+				if (bAutoRemove)
 #endif
-			{
-				GUObjectArray.RemoveUObjectDeleteListener(this);
+				{
+					bProcessRegistration = true;
+				}
 			}
 		}
+
+		if (bProcessRegistration)
+		{
+			ProcessRegistration();
+		}
 	}
 	/**
 	 * Removes all annotation from the annotation list. 
@@ -205,25 +249,30 @@ public:
 	 */
 	void RemoveAllAnnotations()
 	{
-		bool bHadElements = false;
+		bool bProcessRegistration = false;
 
-		// Avoid holding the lock while we call GUObjectArray.RemoveUObjectDeleteListener as it could deadlock
 		{
 			SetAnnotationCacheKeyAndValue(nullptr, TAnnotation());
-			bHadElements = (AnnotationMap.Num() > 0);
+			const bool bHadElements = (AnnotationMap.Num() > 0);
 			UE::TRWScopeLock AnnotationMapLock(AnnotationMapCritical, SLT_Write);
 			AnnotationMap.Empty();
-		}
-		if (bHadElements)
-		{
-			// we are removing the last one, so if we are auto removing or verifying removal, unregister now
+
+			if (bHadElements)
+			{
+				// we are removing the last one, so if we are auto removing or verifying removal, unregister now
 #if (UE_BUILD_SHIPPING || UE_BUILD_TEST)
-			if (bAutoRemove)
+				if (bAutoRemove)
 #endif
-			{
-				GUObjectArray.RemoveUObjectDeleteListener(this);
+				{
+					bProcessRegistration = true;
+				}
 			}
 		}
+
+		if (bProcessRegistration)
+		{
+			ProcessRegistration();
+		}
 	}
 
 	/**
@@ -369,6 +418,11 @@ protected:
 	 */
 	TAnnotation AnnotationCacheValue;
 
+	/**
+	 * Monitor registration to GUObjectArray/UObjectDeleteListener.
+	 */
+	bool bRegistered = false;
+
 #if UE_AUTORTFM
 	/**
 	 * Tracks the presence of a AutoRTFM handlers for cleaning up the annotation cache.
@@ -1160,21 +1214,69 @@ public:
 	 */
 	void AddAnnotation(int32 Index, const TAnnotation& Annotation)
 	{
-		FRWScopeLock AnnotationArrayLock(AnnotationArrayCritical, SLT_Write);
-		AddAnnotationInternal(Index, Annotation);
+		bool bProcessRegistration = false;
+		{
+			FRWScopeLock AnnotationArrayLock(AnnotationArrayCritical, SLT_Write);
+			bProcessRegistration = AddAnnotationInternal(Index, Annotation);
+		}
+
+		if (bProcessRegistration)
+		{
+			ProcessRegistration();
+		}
 	}
 
 	void AddAnnotation(int32 Index, TAnnotation&& Annotation)
-	{
-		FRWScopeLock AnnotationArrayLock(AnnotationArrayCritical, SLT_Write);
-		AddAnnotationInternal(Index, MoveTemp(Annotation));
+	{	
+		bool bProcessRegistration = false;
+		{
+			FRWScopeLock AnnotationArrayLock(AnnotationArrayCritical, SLT_Write);
+			bProcessRegistration = AddAnnotationInternal(Index, MoveTemp(Annotation));
+		}
+
+		if (bProcessRegistration)
+		{
+			ProcessRegistration();
+		}
 	}
 
 
 private:
+	void ProcessRegistration()
+	{
+		// Always take the UObjectDeleteListeners mutex first to preserve lock order in all case
+		// and avoid deadlock. This is important because we need to follow this lock flow :
+		// ConditionalFinishDestroy -> RemoveObjectFromDeleteListeners (UObjectDeleteListenersCritical) -> NotifyUObjectDeleted -> RemoveAnnotation (AnnotationMapCritical)
+		GUObjectArray.LockUObjectDeleteListeners();
+		ON_SCOPE_EXIT
+		{
+			GUObjectArray.UnlockUObjectDeleteListeners();
+		};
+
+		UE::TWriteScopeLock AnnotationMapLock(AnnotationArrayCritical);
+		if (AnnotationArray.Num() == 0)
+		{
+			if (bRegistered)
+			{
+				GUObjectArray.RemoveUObjectDeleteListener(this);
+				bRegistered = false;
+			}
+		}
+		else
+		{
+			if (!bRegistered)
+			{
+				GUObjectArray.AddUObjectDeleteListener(this);
+				bRegistered = true;
+			}
+		}
+	}
+
 	template<typename T>
-	void AddAnnotationInternal(int32 Index, T&& Annotation)
+	[[nodiscard]] bool AddAnnotationInternal(int32 Index, T&& Annotation)
 	{
+		bool bProcessRegistration = false;
+
 		check(Index >= 0);
 		if (Annotation.IsDefault())
 		{
@@ -1189,7 +1291,7 @@ private:
 				if (bAutoRemove)
 #endif
 				{
-					GUObjectArray.AddUObjectDeleteListener(this);
+					bProcessRegistration = true;
 				}
 			}
 			if (Index >= AnnotationArray.Num())
@@ -1203,6 +1305,8 @@ private:
 			}
 			AnnotationArray[Index] = Forward<T>(Annotation);
 		}
+
+		return bProcessRegistration;
 	}
 
 public:
@@ -1244,19 +1348,24 @@ public:
 	 */
 	void RemoveAllAnnotations()
 	{
-		FRWScopeLock AnnotationArrayLock(AnnotationArrayCritical, SLT_Write);
-		bool bHadElements = (AnnotationArray.Num() > 0);
-		AnnotationArray.Empty();
-		if (bHadElements)
+		bool bProcessRegistration = false;
 		{
-			// we are removing the last one, so if we are auto removing or verifying removal, unregister now
+			FRWScopeLock AnnotationArrayLock(AnnotationArrayCritical, SLT_Write);
+			bool bHadElements = (AnnotationArray.Num() > 0);
+			AnnotationArray.Empty();
+			if (bHadElements)
+			{
+				// we are removing the last one, so if we are auto removing or verifying removal, unregister now
 #if (UE_BUILD_SHIPPING || UE_BUILD_TEST)
-			if (bAutoRemove)
+				if (bAutoRemove)
 #endif
-			{
-				GUObjectArray.RemoveUObjectDeleteListener(this);
+				{
+					bProcessRegistration = true;
+				}
 			}
 		}
+
+		ProcessRegistration();
 	}
 
 	/**
@@ -1313,10 +1422,18 @@ public:
 	 */
 	inline TAnnotation& GetAnnotationRef(int32 Index)
 	{
-		FRWScopeLock AnnotationArrayLock(AnnotationArrayCritical, SLT_Write);
-		if (Index >= AnnotationArray.Num())
+		bool bProcessRegistration = false;
 		{
-			AddAnnotationInternal(Index, TAnnotation());
+			FRWScopeLock AnnotationArrayLock(AnnotationArrayCritical, SLT_Write);
+			if (Index >= AnnotationArray.Num())
+			{
+				bProcessRegistration = AddAnnotationInternal(Index, TAnnotation());
+			}
+		}
+
+		if (bProcessRegistration)
+		{
+			ProcessRegistration();
 		}
 		return AnnotationArray[Index];
 	}
@@ -1335,6 +1452,10 @@ private:
 	TArray<TAnnotation> AnnotationArray;
 	FRWLock AnnotationArrayCritical;
 
+	/**
+	 * Monitor registration to GUObjectArray/UObjectDeleteListener.
+	 */
+	bool bRegistered = false;
 };
 
 /**
diff --git a/Engine/Source/Runtime/CoreUObject/Public/UObject/UObjectArray.h b/Engine/Source/Runtime/CoreUObject/Public/UObject/UObjectArray.h
--- a/Engine/Source/Runtime/CoreUObject/Public/UObject/UObjectArray.h
+++ b/Engine/Source/Runtime/CoreUObject/Public/UObject/UObjectArray.h
@@ -1430,6 +1430,22 @@ public:
 		mutable FUObjectItem* CurrentObject;
 	};
 
+	/** Locks the mutex protecting the list of delete listeners */
+	void LockUObjectDeleteListeners()
+	{
+#if THREADSAFE_UOBJECTS
+		UObjectDeleteListenersCritical.Lock();
+#endif
+	}
+
+	/** Unlocks the mutex protecting the list of delete listeners */
+	void UnlockUObjectDeleteListeners()
+	{
+#if THREADSAFE_UOBJECTS
+		UObjectDeleteListenersCritical.Unlock();
+#endif
+	}
+
 private:
 
 	//typedef TStaticIndirectArrayThreadSafeRead<UObjectBase, 8 * 1024 * 1024 /* Max 8M UObjects */, 16384 /* allocated in 64K/128K chunks */ > TUObjectArray;
