Changed UStruct parameter in ContainerPtrToValuePtrForDefaults to const.

diff --git a/Engine/Source/Runtime/CoreUObject/Public/UObject/UnrealType.h b/Engine/Source/Runtime/CoreUObject/Public/UObject/UnrealType.h
--- a/Engine/Source/Runtime/CoreUObject/Public/UObject/UnrealType.h
+++ b/Engine/Source/Runtime/CoreUObject/Public/UObject/UnrealType.h
@@ -752,43 +752,43 @@ public:
 		return ContainerPtrToValuePtr<ValueType>((void*)ContainerPtr, ArrayIndex);
 	}
 
-	// Default variants, these accept and return NULL, and also check the property against the size of the container. 
-	// If we copying from a baseclass (like for a CDO), then this will give NULL for a property that doesn't belong to the baseclass
+	// Default variants, these accept and return nullptr, and also check the property against the size of the container. 
+	// If we copy from a baseclass (like for a CDO), then this will give nullptr for any property that doesn't belong to the baseclass.
 	template<typename ValueType>
-	inline ValueType* ContainerPtrToValuePtrForDefaults(UStruct* ContainerClass, UObject* ContainerPtr, int32 ArrayIndex = 0) const
+	inline ValueType* ContainerPtrToValuePtrForDefaults(const UStruct* ContainerClass, UObject* ContainerPtr, int32 ArrayIndex = 0) const
 	{
 		if (ContainerPtr && IsInContainer(ContainerClass))
 		{
 			return ContainerPtrToValuePtr<ValueType>(ContainerPtr, ArrayIndex);
 		}
-		return NULL;
+		return nullptr;
 	}
 	template<typename ValueType>
-	inline ValueType* ContainerPtrToValuePtrForDefaults(UStruct* ContainerClass, void* ContainerPtr, int32 ArrayIndex = 0) const
+	inline ValueType* ContainerPtrToValuePtrForDefaults(const UStruct* ContainerClass, void* ContainerPtr, int32 ArrayIndex = 0) const
 	{
 		if (ContainerPtr && IsInContainer(ContainerClass))
 		{
 			return ContainerPtrToValuePtr<ValueType>(ContainerPtr, ArrayIndex);
 		}
-		return NULL;
+		return nullptr;
 	}
 	template<typename ValueType>
-	inline ValueType const* ContainerPtrToValuePtrForDefaults(UStruct* ContainerClass, UObject const* ContainerPtr, int32 ArrayIndex = 0) const
+	inline const ValueType* ContainerPtrToValuePtrForDefaults(const UStruct* ContainerClass, const UObject* ContainerPtr, int32 ArrayIndex = 0) const
 	{
 		if (ContainerPtr && IsInContainer(ContainerClass))
 		{
 			return ContainerPtrToValuePtr<ValueType>(ContainerPtr, ArrayIndex);
 		}
-		return NULL;
+		return nullptr;
 	}
 	template<typename ValueType>
-	inline ValueType const* ContainerPtrToValuePtrForDefaults(UStruct* ContainerClass, void const* ContainerPtr, int32 ArrayIndex = 0) const
+	inline const ValueType* ContainerPtrToValuePtrForDefaults(const UStruct* ContainerClass, const void* ContainerPtr, int32 ArrayIndex = 0) const
 	{
 		if (ContainerPtr && IsInContainer(ContainerClass))
 		{
 			return ContainerPtrToValuePtr<ValueType>(ContainerPtr, ArrayIndex);
 		}
-		return NULL;
+		return nullptr;
 	}
 	/** See if the offset of this property is below the supplied container size */
 	UE_FORCEINLINE_HINT bool IsInContainer(int32 ContainerSize) const
