Changed the UStruct parameter in InstanceSubobjectTemplates to be const.

Also some minor cleanup of its logic.

diff --git a/Engine/Source/Runtime/CoreUObject/Private/UObject/Class.cpp b/Engine/Source/Runtime/CoreUObject/Private/UObject/Class.cpp
--- a/Engine/Source/Runtime/CoreUObject/Private/UObject/Class.cpp
+++ b/Engine/Source/Runtime/CoreUObject/Private/UObject/Class.cpp
@@ -2704,18 +2704,28 @@ FTopLevelAssetPath UStruct::GetStructPathName() const
 	return FTopLevelAssetPath(GetOuter()->GetFName(), GetFName());
 }
 
-void UStruct::InstanceSubobjectTemplates( TNotNull<void*> Data, void const* DefaultData, UStruct* DefaultStruct, TNotNull<UObject*> Owner, FObjectInstancingGraph* InstanceGraph )
+void UStruct::InstanceSubobjectTemplates(TNotNull<void*> Data, const void* DefaultData, const UStruct* DefaultStruct, TNotNull<UObject*> Owner, FObjectInstancingGraph* InstanceGraph)
 {
 	// Dynamic instancing allows the instance graph to infer when to instance subobjects instead of requiring the
 	// explicit 'CPF_InstancedReference' flag. This also means the code below intentionally visits all properties
 	// in the reflink chain as a result, unless they are otherwise explicitly excluded through the instance graph.
 	const bool bUsesDynamicInstancing = Owner->GetClass()->ShouldUseDynamicSubobjectInstancing();
-	for ( FProperty* Property = RefLink; Property != NULL; Property = Property->NextRef )
+	for (FProperty* Property = RefLink; Property != nullptr; Property = Property->NextRef)
 	{
-		if ((bUsesDynamicInstancing || Property->ContainsInstancedObjectProperty() || Property->HasAnyPropertyFlags(CPF_AllowSelfReference)) &&
-			(!InstanceGraph || !InstanceGraph->IsPropertyInSubobjectExclusionList(Property)))
+		const bool bPropertySupportsInstancing =
+			bUsesDynamicInstancing || 
+			Property->ContainsInstancedObjectProperty() || 
+			Property->HasAnyPropertyFlags(CPF_AllowSelfReference)
+		;
+
+		const bool bPropertyInExclusionList =
+			InstanceGraph && 
+			InstanceGraph->IsPropertyInSubobjectExclusionList(Property)
+		;
+
+		if (bPropertySupportsInstancing && !bPropertyInExclusionList)
 		{
-			Property->InstanceSubobjects( Property->ContainerPtrToValuePtr<uint8>(Data), (uint8*)Property->ContainerPtrToValuePtrForDefaults<uint8>(DefaultStruct, DefaultData), Owner, InstanceGraph );
+			Property->InstanceSubobjects(Property->ContainerPtrToValuePtr<uint8>(Data), Property->ContainerPtrToValuePtrForDefaults<uint8>(DefaultStruct, DefaultData), Owner, InstanceGraph);
 		}
 	}
 }
diff --git a/Engine/Source/Runtime/CoreUObject/Public/UObject/Class.h b/Engine/Source/Runtime/CoreUObject/Public/UObject/Class.h
--- a/Engine/Source/Runtime/CoreUObject/Public/UObject/Class.h
+++ b/Engine/Source/Runtime/CoreUObject/Public/UObject/Class.h
@@ -556,7 +556,7 @@ public:
 	 * @param	Owner						the object that contains the component currently located at Data
 	 * @param	InstanceGraph				contains the mappings of instanced objects and components to their templates
 	 */
-	COREUOBJECT_API void InstanceSubobjectTemplates( TNotNull<void*> Data, void const* DefaultData, UStruct* DefaultStruct, TNotNull<UObject*> Owner, FObjectInstancingGraph* InstanceGraph );
+	COREUOBJECT_API void InstanceSubobjectTemplates(TNotNull<void*> Data, const void* DefaultData, const UStruct* DefaultStruct, TNotNull<UObject*> Owner, FObjectInstancingGraph* InstanceGraph);
 
 	/** Returns the structure used for inheritance, may be changed by child types */
 	virtual UStruct* GetInheritanceSuper() const {return GetSuperStruct();}
