Fix a bug in `FUObjectAnnotationDense` by using a transactionally safe RW lock.

diff --git a/Engine/Source/Programs/AutoRTFMTests/Classes/MyAutoRTFMTestObjectWithSubObjects.h b/Engine/Source/Programs/AutoRTFMTests/Classes/MyAutoRTFMTestObjectWithSubObjects.h
new file mode 100644
--- /dev/null
+++ b/Engine/Source/Programs/AutoRTFMTests/Classes/MyAutoRTFMTestObjectWithSubObjects.h
@@ -0,0 +1,19 @@
+// Copyright Epic Games, Inc. All Rights Reserved.
+#pragma once
+
+#include "MyAutoRTFMTestObjectWithSubObjects.generated.h"
+
+UCLASS()
+class UMyAutoRTFMTestObjectWithSubObjects : public UObject
+{
+	GENERATED_BODY()
+
+public:
+	static FName GetSubObjectName();
+
+	UMyAutoRTFMTestObjectWithSubObjects(const FObjectInitializer& ObjectInitializer = FObjectInitializer::Get());
+
+	UObject* SubObject;
+
+	void DoNothing() const {}
+};
diff --git a/Engine/Source/Programs/AutoRTFMTests/Private/MyAutoRTFMTestObjectWithSubObjects.cpp b/Engine/Source/Programs/AutoRTFMTests/Private/MyAutoRTFMTestObjectWithSubObjects.cpp
new file mode 100644
--- /dev/null
+++ b/Engine/Source/Programs/AutoRTFMTests/Private/MyAutoRTFMTestObjectWithSubObjects.cpp
@@ -0,0 +1,18 @@
+// Copyright Epic Games, Inc. All Rights Reserved.
+
+#include "MyAutoRTFMTestObjectWithSubObjects.h"
+#include "MyAutoRTFMTestObject.h"
+
+FName UMyAutoRTFMTestObjectWithSubObjects::GetSubObjectName()
+{
+	static FName Name(FString::Printf(TEXT("%s%u"), TEXT(__FILE__), __LINE__));
+	return Name;
+}
+
+UMyAutoRTFMTestObjectWithSubObjects::UMyAutoRTFMTestObjectWithSubObjects(const FObjectInitializer& ObjectInitializer /* = FObjectInitializer::Get() */)
+{
+	// Having an object that the default subobject replaces pushes us down a new and exciting codepath.
+	UMyAutoRTFMTestObject* const Replacee = NewObject<UMyAutoRTFMTestObject>(this, UMyAutoRTFMTestObjectWithSubObjects::GetSubObjectName());
+
+	SubObject = CreateDefaultSubobject<UMyAutoRTFMTestObject>(GetSubObjectName());
+}
diff --git a/Engine/Source/Programs/AutoRTFMTests/Private/UObjectTests.cpp b/Engine/Source/Programs/AutoRTFMTests/Private/UObjectTests.cpp
--- a/Engine/Source/Programs/AutoRTFMTests/Private/UObjectTests.cpp
+++ b/Engine/Source/Programs/AutoRTFMTests/Private/UObjectTests.cpp
@@ -6,6 +6,7 @@
 #include "GenericPlatform/GenericPlatformMisc.h"
 #include "HAL/MallocLeakDetection.h"
 #include "MyAutoRTFMTestObject.h"
+#include "MyAutoRTFMTestObjectWithSubObjects.h"
 #include "UObject/GCObject.h"
 #include "UObject/ReachabilityAnalysis.h"
 #include "UObject/UObjectAnnotation.h"
@@ -1281,3 +1282,16 @@ TEST_CASE("FGCObject")
 	}
 }
 
+TEST_CASE("UObject.ObjectWithSubObject")
+{
+	UMyAutoRTFMTestObject* Outer = NewObject<UMyAutoRTFMTestObject>();
+	UMyAutoRTFMTestObjectWithSubObjects* Object = nullptr;
+
+	bool bRetry = true;
+
+	AutoRTFM::Testing::Commit([&]
+	{
+		Object = NewObject<UMyAutoRTFMTestObjectWithSubObjects>(Outer, TEXT("WOWWEE"));
+		REQUIRE(Object);
+	});
+}
diff --git a/Engine/Source/Runtime/CoreUObject/Public/UObject/UObjectAnnotation.h b/Engine/Source/Runtime/CoreUObject/Public/UObject/UObjectAnnotation.h
--- a/Engine/Source/Runtime/CoreUObject/Public/UObject/UObjectAnnotation.h
+++ b/Engine/Source/Runtime/CoreUObject/Public/UObject/UObjectAnnotation.h
@@ -1216,7 +1216,7 @@ public:
 	{
 		bool bProcessRegistration = false;
 		{
-			FRWScopeLock AnnotationArrayLock(AnnotationArrayCritical, SLT_Write);
+			UE::TRWScopeLock AnnotationArrayLock(AnnotationArrayCritical, SLT_Write);
 			bProcessRegistration = AddAnnotationInternal(Index, Annotation);
 		}
 
@@ -1230,7 +1230,7 @@ public:
 	{	
 		bool bProcessRegistration = false;
 		{
-			FRWScopeLock AnnotationArrayLock(AnnotationArrayCritical, SLT_Write);
+			UE::TRWScopeLock AnnotationArrayLock(AnnotationArrayCritical, SLT_Write);
 			bProcessRegistration = AddAnnotationInternal(Index, MoveTemp(Annotation));
 		}
 
@@ -1327,7 +1327,7 @@ public:
 	 */
 	void RemoveAnnotation(int32 Index)
 	{
-		FRWScopeLock AnnotationArrayLock(AnnotationArrayCritical, SLT_Write);
+		UE::TRWScopeLock AnnotationArrayLock(AnnotationArrayCritical, SLT_Write);
 		RemoveAnnotationInternal(Index);
 	}
 
@@ -1350,7 +1350,7 @@ public:
 	{
 		bool bProcessRegistration = false;
 		{
-			FRWScopeLock AnnotationArrayLock(AnnotationArrayCritical, SLT_Write);
+			UE::TRWScopeLock AnnotationArrayLock(AnnotationArrayCritical, SLT_Write);
 			bool bHadElements = (AnnotationArray.Num() > 0);
 			AnnotationArray.Empty();
 			if (bHadElements)
@@ -1393,7 +1393,7 @@ public:
 	inline TAnnotation GetAnnotation(int32 Index)
 	{
 		check(Index >= 0);
-		FRWScopeLock AnnotationArrayLock(AnnotationArrayCritical, SLT_ReadOnly);
+		UE::TRWScopeLock AnnotationArrayLock(AnnotationArrayCritical, SLT_ReadOnly);
 		if (Index < AnnotationArray.Num())
 		{
 			return AnnotationArray[Index];
@@ -1424,7 +1424,7 @@ public:
 	{
 		bool bProcessRegistration = false;
 		{
-			FRWScopeLock AnnotationArrayLock(AnnotationArrayCritical, SLT_Write);
+			UE::TRWScopeLock AnnotationArrayLock(AnnotationArrayCritical, SLT_Write);
 			if (Index >= AnnotationArray.Num())
 			{
 				bProcessRegistration = AddAnnotationInternal(Index, TAnnotation());
@@ -1450,7 +1450,7 @@ private:
 	 * Map from live objects to an annotation
 	 */
 	TArray<TAnnotation> AnnotationArray;
-	FRWLock AnnotationArrayCritical;
+	FTransactionallySafeRWLock AnnotationArrayCritical;
 
 	/**
 	 * Monitor registration to GUObjectArray/UObjectDeleteListener.
