Added UEOps to FObjectHandle.
Added a workaround for a bug on MSVC which caused a failure in ObjectHandleTest.cpp.

diff --git a/Engine/Source/Runtime/Core/Public/Misc/UEOps.h b/Engine/Source/Runtime/Core/Public/Misc/UEOps.h
--- a/Engine/Source/Runtime/Core/Public/Misc/UEOps.h
+++ b/Engine/Source/Runtime/Core/Public/Misc/UEOps.h
@@ -70,6 +70,48 @@ template <typename LhsType>
 	return Lhs.UEOpEquals(nullptr);
 }
 
+#if defined(_MSC_VER) && !defined(__clang__)
+	// On MSVC, operator!= isn't synthesized from operator== in some contexts.  This was found in a case of FObjectHandle:
+	//
+	// FObjectHandle NullObjectHandle = ...;
+	// FObjectHandle TestLateResolveUnsafeObjectHandle = ...;
+	// ...
+	// CHECK(NullObjectHandle != TestLateResolveUnsafeObjectHandle); // catch2\internal\catch_decomposer.hpp(219,9): error C2678: binary '!=': no operator found which takes a left-hand operand of type 'FObjectHandle' (or there is no acceptable conversion)
+	//
+	// These overloads exist to work around this issue - they'll just make overload resolution errors slightly worse.  When the following ticket is fixed, these overloads can be removed.
+	//
+	// https://developercommunity.visualstudio.com/t/Synthesised-operator-not-considered-in/10945932
+
+	template <typename LhsType, typename RhsType>
+		requires UE::Core::Private::CWithUEOpEquals<LhsType, RhsType> || UE::Core::Private::CWithUEOpEquals<RhsType, LhsType>
+	[[nodiscard]] UE_REWRITE constexpr bool operator!=(const LhsType& Lhs, const RhsType& Rhs)
+	{
+		if constexpr (UE::Core::Private::CWithUEOpEquals<LhsType, RhsType>)
+		{
+			return !Lhs.UEOpEquals(Rhs);
+		}
+		else
+		{
+			return !Rhs.UEOpEquals(Lhs);
+		}
+	}
+
+	template <typename LhsType>
+		requires UE::Core::Private::CWithUEOpEquals<LhsType, decltype(nullptr)>
+	[[nodiscard]] UE_REWRITE constexpr bool operator!=(const LhsType& Lhs, UE::Core::Private::FIncomplete* Rhs)
+	{
+		return !Lhs.UEOpEquals(nullptr);
+	}
+
+	template <typename RhsType>
+		requires UE::Core::Private::CWithUEOpEquals<RhsType, decltype(nullptr)>
+	[[nodiscard]] UE_REWRITE constexpr bool operator!=(UE::Core::Private::FIncomplete* Lhs, const RhsType& Rhs)
+	{
+		return !Rhs.UEOpEquals(nullptr);
+	}
+
+#endif
+
 template <typename LhsType, typename RhsType>
 	requires UE::Core::Private::CWithUEOpLessThan<LhsType, RhsType> || UE::Core::Private::CWithUEOpGreaterThan<RhsType, LhsType>
 [[nodiscard]] UE_REWRITE constexpr bool operator<(const LhsType& Lhs, const RhsType& Rhs)
diff --git a/Engine/Source/Runtime/CoreUObject/Public/UObject/ObjectHandle.h b/Engine/Source/Runtime/CoreUObject/Public/UObject/ObjectHandle.h
--- a/Engine/Source/Runtime/CoreUObject/Public/UObject/ObjectHandle.h
+++ b/Engine/Source/Runtime/CoreUObject/Public/UObject/ObjectHandle.h
@@ -8,6 +8,7 @@
 #include "HAL/CriticalSection.h"
 #include "HAL/Platform.h"
 #include "Misc/AssertionMacros.h"
+#include "Misc/UEOps.h"
 #include "Templates/TypeHash.h"
 #include "UObject/NameTypes.h"
 #include "UObject/ObjectHandleTracking.h"
@@ -71,6 +72,8 @@ namespace UE::CoreUObject::Private
 		{
 			return PointerOrRef != 0;
 		}
+
+		bool UEOpEquals(FObjectHandlePrivate Rhs) const;
 	};
 
 	/* Returns the packed object ref for this object IF one exists otherwise returns a null PackedObjectRef */
@@ -104,6 +107,8 @@ namespace UE::CoreUObject::Private
 
 		COREUOBJECT_API static FRemoteObjectHandlePrivate ConvertToRemoteHandle(UObject* Object);
 		COREUOBJECT_API static FRemoteObjectHandlePrivate FromIdNoResolve(FRemoteObjectId ObjectId);
+
+		bool UEOpEquals(FRemoteObjectHandlePrivate Rhs) const;
 	};
 
 #endif
@@ -198,24 +203,24 @@ inline bool IsObjectHandleTypeSafe(FObjectHandle Handle)
 
 /////////////////////////////////////////////////////////////////////////////////////////////////////////////////////////
 #if UE_WITH_OBJECT_HANDLE_LATE_RESOLVE
-inline bool operator==(UE::CoreUObject::Private::FObjectHandlePrivate LHS, UE::CoreUObject::Private::FObjectHandlePrivate RHS)
+inline bool UE::CoreUObject::Private::FObjectHandlePrivate::UEOpEquals(FObjectHandlePrivate RHS) const
 {
 	using namespace UE::CoreUObject::Private;
 
-	bool LhsResolved = IsObjectHandleResolved(LHS);
+	bool LhsResolved = IsObjectHandleResolved(*this);
 	bool RhsResolved = IsObjectHandleResolved(RHS);
 
 	//if both resolved or both unresolved compare the uintptr
 	if (LhsResolved == RhsResolved)
 	{
-		return LHS.PointerOrRef == RHS.PointerOrRef;
+		return PointerOrRef == RHS.PointerOrRef;
 	}
 
 	//only one side can be resolved
 	if (LhsResolved)
 	{
 		//both sides can't be null as resolved status would have be true for both
-		const UObject* Obj = ReadObjectHandlePointerNoCheck(LHS);
+		const UObject* Obj = ReadObjectHandlePointerNoCheck(*this);
 		if (!Obj)
 		{
 			return false;
@@ -245,16 +250,11 @@ inline bool operator==(UE::CoreUObject::Private::FObjectHandlePrivate LHS, UE::C
 		{
 			return false;
 		}
-		return PackedRhs.EncodedRef == LHS.PointerOrRef;
+		return PackedRhs.EncodedRef == PointerOrRef;
 	}
 
 }
 
-inline bool operator!=(UE::CoreUObject::Private::FObjectHandlePrivate LHS, UE::CoreUObject::Private::FObjectHandlePrivate RHS)
-{
-	return !(LHS == RHS);
-}
-
 inline uint32 GetTypeHash(UE::CoreUObject::Private::FObjectHandlePrivate Handle)
 {
 	using namespace UE::CoreUObject::Private;
@@ -278,25 +278,25 @@ inline uint32 GetTypeHash(UE::CoreUObject::Private::FObjectHandlePrivate Handle)
 	return GetTypeHash(Handle.GetPointerOrRef());
 }
 #elif UE_WITH_REMOTE_OBJECT_HANDLE
-inline bool operator==(UE::CoreUObject::Private::FRemoteObjectHandlePrivate LHS, UE::CoreUObject::Private::FRemoteObjectHandlePrivate RHS)
+inline bool UE::CoreUObject::Private::FRemoteObjectHandlePrivate::UEOpEquals(FRemoteObjectHandlePrivate RHS) const
 {
 	using namespace UE::CoreUObject::Private;
 	using namespace UE::RemoteObject::Handle;
 
-	const bool bLhsResolved = IsObjectHandleResolved_ForGC(LHS);
+	const bool bLhsResolved = IsObjectHandleResolved_ForGC(*this);
 	const bool bRhsResolved = IsObjectHandleResolved_ForGC(RHS);
 
 	//if both resolved or both unresolved compare the uintptr
 	if (bLhsResolved == bRhsResolved)
 	{
-		return LHS.PointerOrHandle == RHS.PointerOrHandle;
+		return PointerOrHandle == RHS.PointerOrHandle;
 	}
 
 	//only one side can be resolved
 	if (bLhsResolved)
 	{
 		//both sides can't be null as resolved status would have be true for both
-		const UObject* Obj = ReadObjectHandlePointerNoCheck(LHS);
+		const UObject* Obj = ReadObjectHandlePointerNoCheck(*this);
 		if (!Obj)
 		{
 			return false;
@@ -314,17 +314,12 @@ inline bool operator==(UE::CoreUObject::Private::FRemoteObjectHandlePrivate LHS,
 			return false;
 		}
 
-		const FRemoteObjectStub* LhsStub = LHS.ToStub();
+		const FRemoteObjectStub* LhsStub = ToStub();
 		return LhsStub->Id == FRemoteObjectId((const UObjectBase*)Obj);
 	}
 
 }
 
-inline bool operator!=(UE::CoreUObject::Private::FRemoteObjectHandlePrivate LHS, UE::CoreUObject::Private::FRemoteObjectHandlePrivate RHS)
-{
-	return !(LHS == RHS);
-}
-
 inline FRemoteObjectId GetRemoteObjectId(UE::CoreUObject::Private::FRemoteObjectHandlePrivate Handle)
 {
 	using namespace UE::CoreUObject::Private;
