Print object index when logging issues found when performing GC assumptios verification tests

diff --git a/Engine/Source/Runtime/CoreUObject/Private/UObject/GarbageCollectionVerification.cpp b/Engine/Source/Runtime/CoreUObject/Private/UObject/GarbageCollectionVerification.cpp
--- a/Engine/Source/Runtime/CoreUObject/Private/UObject/GarbageCollectionVerification.cpp
+++ b/Engine/Source/Runtime/CoreUObject/Private/UObject/GarbageCollectionVerification.cpp
@@ -535,7 +535,8 @@ void VerifyNoUnreachableObjects(int32 NumUnreachable)
 
 		for (int32 ObjectIndex = 0; ObjectIndex < NumObjects && (FirstObjectIndex + ObjectIndex) < GUObjectArray.GetObjectArrayNum(); ++ObjectIndex)
 		{
-			FUObjectItem& ObjectItem = GUObjectArray.GetObjectItemArrayUnsafe()[FirstObjectIndex + ObjectIndex];
+			const int32 ObjectIndexInGlobalArray = FirstObjectIndex + ObjectIndex;
+			FUObjectItem& ObjectItem = GUObjectArray.GetObjectItemArrayUnsafe()[ObjectIndexInGlobalArray];
 			if (ObjectItem.GetObject())
 			{
 				UObject* Object = static_cast<UObject*>(ObjectItem.GetObject());
@@ -543,14 +544,16 @@ void VerifyNoUnreachableObjects(int32 NumUnreachable)
 				{
 					if (FGCFlags::IsMaybeUnreachable_ForGC(&ObjectItem))
 					{
-						UE_LOG(LogGarbage, Warning, TEXT("Object %s%s is still marked as MaybeUnreachable after Reachability Analysis is complete."), 
+						UE_LOG(LogGarbage, Warning, TEXT("Object [%d] %s%s is still marked as MaybeUnreachable after Reachability Analysis is complete."),
+							ObjectIndexInGlobalArray,
 							*FReferenceChainSearch::GetObjectFlags(FGCObjectInfo(Object)),
 							*Object->GetFullName());
 						NumErrors++;
 					}
 					if (!FGCFlags::IsReachable_ForGC(&ObjectItem) && !GUObjectArray.IsDisregardForGC(Object))
 					{
-						UE_LOG(LogGarbage, Warning, TEXT("Object %s%s is NOT marked as Unreachable and NOT marked as Reachable."),
+						UE_LOG(LogGarbage, Warning, TEXT("Object [%d] %s%s is NOT marked as Unreachable and NOT marked as Reachable."),
+							ObjectIndexInGlobalArray,
 							*FReferenceChainSearch::GetObjectFlags(FGCObjectInfo(Object)),
 							*Object->GetFullName());
 						NumErrors++;
@@ -563,7 +566,8 @@ void VerifyNoUnreachableObjects(int32 NumUnreachable)
 
 					if (FGCFlags::IsReachable_ForGC(&ObjectItem))
 					{
-						UE_LOG(LogGarbage, Warning, TEXT("Object %s%s is still marked as Reachable."),
+						UE_LOG(LogGarbage, Warning, TEXT("Object [%d] %s%s is still marked as Reachable."),
+							ObjectIndexInGlobalArray,
 							*FReferenceChainSearch::GetObjectFlags(FGCObjectInfo(Object)),
 							*Object->GetFullName());
 						NumErrors++;
