[CIS Fix] Disabling Placeholder Objects again for now because it's breaking CIS tests

diff --git a/Engine/Source/Runtime/CoreUObject/Private/UObject/InstanceDataObjectUtils.cpp b/Engine/Source/Runtime/CoreUObject/Private/UObject/InstanceDataObjectUtils.cpp
--- a/Engine/Source/Runtime/CoreUObject/Private/UObject/InstanceDataObjectUtils.cpp
+++ b/Engine/Source/Runtime/CoreUObject/Private/UObject/InstanceDataObjectUtils.cpp
@@ -382,7 +382,7 @@ namespace UE
 		TEXT("When enabled, IDOs will be constructed using an archetype chain")
 	);
 
-	bool bEnableIDOPlaceholderObjects = true;
+	bool bEnableIDOPlaceholderObjects = false;
 	FAutoConsoleVariableRef EnableIDOPlaceholderObjectsCVar(
 		TEXT("IDO.Placeholder.Enable"),
 		bEnableIDOPlaceholderObjects,
