Make -VerifyGC and -NoVerifyGc take all verification cvars into account.

diff --git a/Engine/Source/Runtime/CoreUObject/Private/UObject/Obj.cpp b/Engine/Source/Runtime/CoreUObject/Private/UObject/Obj.cpp
--- a/Engine/Source/Runtime/CoreUObject/Private/UObject/Obj.cpp
+++ b/Engine/Source/Runtime/CoreUObject/Private/UObject/Obj.cpp
@@ -5691,15 +5691,36 @@ void StaticUObjectInit()
 	GObjTransientPkg = NewObject<UPackage>(nullptr, TEXT("/Engine/Transient"), RF_Transient);
 	GObjTransientPkg->AddToRoot();
 
+	// Force GC verification on or off based on command line parameters.
+	bool bCommandLineForcesVerifyGC = FParse::Param(FCommandLine::Get(), TEXT("VERIFYGC"));
+	bool bCommandLineForcesNoVerifyGC = FParse::Param(FCommandLine::Get(), TEXT("NOVERIFYGC"));
 	if (IConsoleVariable* CVarVerifyGCAssumptions = IConsoleManager::Get().FindConsoleVariable(TEXT("gc.VerifyAssumptions")))
 	{
-	if( FParse::Param( FCommandLine::Get(), TEXT("VERIFYGC") ) )
-	{
+		if (bCommandLineForcesVerifyGC)
+		{
 			CVarVerifyGCAssumptions->Set(true, ECVF_SetByCommandline);
+		}
+		else if (bCommandLineForcesNoVerifyGC)
+		{
+			CVarVerifyGCAssumptions->Set(false, ECVF_SetByCommandline);
+		}
 	}
-	if( FParse::Param( FCommandLine::Get(), TEXT("NOVERIFYGC") ) )
+	if (IConsoleVariable* CVarVerifyGCAssumptionsOnFullPurge = IConsoleManager::Get().FindConsoleVariable(TEXT("gc.VerifyAssumptionsOnFullPurge")))
 	{
-			CVarVerifyGCAssumptions->Set(false, ECVF_SetByCommandline);
+		if (bCommandLineForcesVerifyGC)
+		{
+			CVarVerifyGCAssumptionsOnFullPurge->Set(true, ECVF_SetByCommandline);
+		}
+		else if (bCommandLineForcesNoVerifyGC)
+		{
+			CVarVerifyGCAssumptionsOnFullPurge->Set(false, ECVF_SetByCommandline);
+		}
+	}
+	if (IConsoleVariable* CVarVerifyGCAssumptionsChance = IConsoleManager::Get().FindConsoleVariable(TEXT("gc.VerifyAssumptionsChance")))
+	{
+		if (bCommandLineForcesNoVerifyGC)
+		{
+			CVarVerifyGCAssumptionsChance->Set(0.0f, ECVF_SetByCommandline);
 		}
 	}
 
