Added a debug console command to print all IDOs and placeholder types.

#rnx

diff --git a/Engine/Source/Runtime/CoreUObject/Internal/UObject/PropertyBagRepository.h b/Engine/Source/Runtime/CoreUObject/Internal/UObject/PropertyBagRepository.h
--- a/Engine/Source/Runtime/CoreUObject/Internal/UObject/PropertyBagRepository.h
+++ b/Engine/Source/Runtime/CoreUObject/Internal/UObject/PropertyBagRepository.h
@@ -122,6 +122,9 @@ public:
 	UE_INTERNAL void GatherStats(FPropertyBagRepositoryStats& Stats);
 	#endif
 
+	UE_INTERNAL void DumpIDOs(FOutputDevice& OutputDevice);
+	UE_INTERNAL void DumpPropertyBagPlaceholderTypes(FOutputDevice& OutputDevice);
+
 	// Placeholder object feature flags; external callers can query a particular subfeature to see if it's enabled.
 	// Note: These have CVar counterparts defined in the .cpp file. Add an entry there to also create a toggle switch.
 	enum class EPlaceholderObjectFeature : uint8
diff --git a/Engine/Source/Runtime/CoreUObject/Private/UObject/InstanceDataObjectUtils.cpp b/Engine/Source/Runtime/CoreUObject/Private/UObject/InstanceDataObjectUtils.cpp
--- a/Engine/Source/Runtime/CoreUObject/Private/UObject/InstanceDataObjectUtils.cpp
+++ b/Engine/Source/Runtime/CoreUObject/Private/UObject/InstanceDataObjectUtils.cpp
@@ -563,6 +563,37 @@ namespace UE
 		#endif // STATS
 	}
 
+	static void DumpIDOs(FOutputDevice& OutputDevice)
+	{
+		FPropertyBagRepository::Get().DumpPropertyBagPlaceholderTypes(OutputDevice);
+
+		OutputDevice.Log(TEXT("Instance Data Object classes:"));
+		OutputDevice.Log(TEXT("------------------------------"));
+
+		{
+			TUniqueLock Lock(IDOClassCacheMutex);
+
+			int32 NumIDOClasses = 0;
+
+			for (const TPair<FBlake3Hash, TSharedPtr<TCacheItem<UInstanceDataObjectClass>>>& Pair : IDOClassCache)
+			{
+				UInstanceDataObjectClass* Class = Pair.Value->WeakPtr.Get();
+
+				if (Class)
+				{
+					++NumIDOClasses;
+
+					OutputDevice.Logf(TEXT("%s"), *Class->GetFullName());
+				}
+
+			}
+
+			OutputDevice.Logf(TEXT("%d IDO class(es)"), NumIDOClasses);
+		}
+
+		FPropertyBagRepository::Get().DumpIDOs(OutputDevice);
+	}
+
 	FAutoConsoleCommand DumpIDOStatsCommand(
 		TEXT("IDO.DumpStats"),
 		TEXT("Prints statistics for all current Instance Data Objects."),
@@ -572,6 +603,15 @@ namespace UE
 		})
 	);
 
+	FAutoConsoleCommand DumpIDOsCommand(
+		TEXT("IDO.Dump"),
+		TEXT("Prints all Instance Data Objects and their classes."),
+		FConsoleCommandWithArgsAndOutputDeviceDelegate::CreateLambda([](const TArray<FString>& Args, FOutputDevice& OutputDevice)
+		{
+			DumpIDOs(OutputDevice);
+		})
+	);
+
 	static TSet<FString> GetExcludedUnknownPropertyTypes()
 	{
 		TArray<FString> Result;
diff --git a/Engine/Source/Runtime/CoreUObject/Private/UObject/PropertyBagRepository.cpp b/Engine/Source/Runtime/CoreUObject/Private/UObject/PropertyBagRepository.cpp
--- a/Engine/Source/Runtime/CoreUObject/Private/UObject/PropertyBagRepository.cpp
+++ b/Engine/Source/Runtime/CoreUObject/Private/UObject/PropertyBagRepository.cpp
@@ -63,6 +63,14 @@ public:
 		return PlaceholderTypes.Contains(Type);
 	}
 
+	void GetTypes(TSet<TObjectPtr<const UStruct>>& OutTypes)
+	{
+		ConsumePendingPlaceholderTypes();
+		
+		FScopeLock ScopeLock(&CriticalSection);
+		OutTypes.Append(PlaceholderTypes);
+	}
+
 	int32 Num() const
 	{
 		return PlaceholderTypes.Num();
@@ -1023,6 +1031,47 @@ void FPropertyBagRepository::GatherStats(FPropertyBagRepositoryStats& Stats)
 }
 #endif // STATS
 
+void FPropertyBagRepository::DumpIDOs(FOutputDevice& OutputDevice)
+{
+	OutputDevice.Log(TEXT("Instance Data Objects:"));
+	OutputDevice.Log(TEXT("------------------------------"));
+
+	int32 NumIDOs = 0;
+
+	for (const TPair<const UObject*, FPropertyBagAssociationData>& Pair : AssociatedData)
+	{
+		const FPropertyBagAssociationData& BagData = Pair.Value;
+
+		if (BagData.InstanceDataObject != nullptr)
+		{
+			++NumIDOs;
+
+			OutputDevice.Logf(TEXT("%s"), *BagData.InstanceDataObject->GetFullName());
+
+			const UObject* Owner = FindInstanceForDataObject(BagData.InstanceDataObject);
+			OutputDevice.Logf(TEXT("  Owner: %s"), (Owner ? *Owner->GetFullName() : TEXT("NULL")));
+		}
+	}
+
+	OutputDevice.Logf(TEXT("%d IDO(s)"), NumIDOs);
+}
+
+void FPropertyBagRepository::DumpPropertyBagPlaceholderTypes(FOutputDevice& OutputDevice)
+{
+	OutputDevice.Log(TEXT("Placeholder Types:"));
+	OutputDevice.Log(TEXT("------------------------------"));
+
+	TSet<TObjectPtr<const UStruct>> Types;
+	FPropertyBagPlaceholderTypeRegistry::Get().GetTypes(Types);
+
+	for (TObjectPtr<const UStruct> Type : Types)
+	{
+		OutputDevice.Logf(TEXT("%s"), *Type->GetFullName());
+	}
+
+	OutputDevice.Logf(TEXT("%d Placeholder type(s)"), Types.Num());
+}
+
 } // UE
 
 #endif // WITH_EDITORONLY_DATA
