IncrementalCook: DelegateProperty: Fix a false positive incremental skip in blueprints that subscribe to a delegate when the delegate's function prototype changes. We need to recook those blueprints because the blueprint compiler may need to issue an error if the blueprint's callback function does not match the new prototype.
#rnx

diff --git a/Engine/Source/Runtime/CoreUObject/Private/UObject/PropertyDelegate.cpp b/Engine/Source/Runtime/CoreUObject/Private/UObject/PropertyDelegate.cpp
--- a/Engine/Source/Runtime/CoreUObject/Private/UObject/PropertyDelegate.cpp
+++ b/Engine/Source/Runtime/CoreUObject/Private/UObject/PropertyDelegate.cpp
@@ -215,10 +215,11 @@ void FDelegateProperty::AppendSchemaHash(FBlake3& Builder, bool bSkipEditorOnly)
 	Super::AppendSchemaHash(Builder, bSkipEditorOnly);
 	if (SignatureFunction)
 	{
-		// Hash the function's name instead of recursively hashing the function; the function's schema does not impact how we serialize our pointer to it
-		FNameBuilder ObjectPath;
-		SignatureFunction->GetPathName(nullptr, ObjectPath);
-		Builder.Update(ObjectPath.GetData(), ObjectPath.Len() * sizeof(ObjectPath.GetData()[0]));
+		// Recursively hash the SignatureFunction. The function's schema does not impact how we serialize our pointer
+		// to it, but it is used by the blueprint compiler to validate that a subscriber to the delegate in a blueprint
+		// matches the function prototype of the SignatureFunction.
+		const FIoHash& FunctionSchemaHash = SignatureFunction->GetSchemaHash(bSkipEditorOnly);
+		Builder.Update(&FunctionSchemaHash, sizeof(FunctionSchemaHash));
 	}
 }
 #endif
