Fix crash in FOverrridableManager::GetOverriddenState() when object was containing self references.

#rnx

diff --git a/Engine/Source/Runtime/CoreUObject/Private/UObject/OverridableManager.cpp b/Engine/Source/Runtime/CoreUObject/Private/UObject/OverridableManager.cpp
--- a/Engine/Source/Runtime/CoreUObject/Private/UObject/OverridableManager.cpp
+++ b/Engine/Source/Runtime/CoreUObject/Private/UObject/OverridableManager.cpp
@@ -188,11 +188,9 @@ EOverriddenState FOverridableManager::GetOverriddenState(TNotNull<UObject*> Obje
 			}
 		}
 
-		TSet<UObject*> VisitedObject;
-		VisitedObject.Add(Object);
 		bool bAllOverridden = true;
 		bool bHasOverrides = false;
-		Object->GetClass()->Visit(Object, [this, Object, OverriddenProperties, Path, &VisitedObject, &bAllOverridden, &bHasOverrides, ConsiderPropertyForOverriddenState, &bOutInheritedState](const FPropertyVisitorContext& Context)
+		Object->GetClass()->Visit(Object, [this, Object, OverriddenProperties, Path, &bAllOverridden, &bHasOverrides, ConsiderPropertyForOverriddenState, &bOutInheritedState](const FPropertyVisitorContext& Context)
 		{
 			const EOverriddenPropertyOperation Operation = OverriddenProperties->GetOverriddenPropertyOperation(Context.Path.GetRootIterator());
 			bool bIsEqual = false;
@@ -237,7 +235,7 @@ EOverriddenState FOverridableManager::GetOverriddenState(TNotNull<UObject*> Obje
 				return EPropertyVisitorControlFlow::StepOver;
 			}
 
-			// No need to dig deeper if this is a replace of an add
+			// No need to dig deeper if this is a replace or an add
 			if (Operation == EOverriddenPropertyOperation::Replace || 
 				Operation == EOverriddenPropertyOperation::Add)
 			{
@@ -258,10 +256,22 @@ EOverriddenState FOverridableManager::GetOverriddenState(TNotNull<UObject*> Obje
 				{
 					if (UObject* InstancedSubObject = ObjectProperty->GetObjectPropertyValue(Context.Data.PropertyData))
 					{
-						if (!VisitedObject.Contains(InstancedSubObject) && InstancedSubObject->IsIn(Object))
+						if (InstancedSubObject->IsIn(Object))
 						{
-							VisitedObject.Add(InstancedSubObject);
-							return EPropertyVisitorControlFlow::StepInto;
+							switch(GetOverriddenState(InstancedSubObject, ConsiderPropertyForOverriddenState))
+							{
+								case EOverriddenState::NoOverrides:
+									bAllOverridden = false;
+									break;
+								case EOverriddenState::HasOverrides:
+									bAllOverridden = false;
+									bHasOverrides = true;
+									break;
+								case EOverriddenState::AllOverridden:
+									bHasOverrides = true;
+									break;
+							}
+							return EPropertyVisitorControlFlow::StepOver;
 						}
 					}
 				}
diff --git a/Engine/Source/Runtime/CoreUObject/Private/UObject/OverriddenPropertySet.cpp b/Engine/Source/Runtime/CoreUObject/Private/UObject/OverriddenPropertySet.cpp
--- a/Engine/Source/Runtime/CoreUObject/Private/UObject/OverriddenPropertySet.cpp
+++ b/Engine/Source/Runtime/CoreUObject/Private/UObject/OverriddenPropertySet.cpp
@@ -555,11 +555,9 @@ EOverriddenPropertyOperation FOverriddenPropertySet::GetOverriddenPropertyOperat
 			if (NextPropertyIterator)
 			{
 				// Forward any sub queries to the subobject
-				if (UObject* SubObject = TryGetInstancedSubObjectValue(ObjectProperty, SubValuePtr))
-				{
-					// This should not be needed in the property grid, as it should already been called on the subobject.
-					return OverridableManager.GetOverriddenPropertyOperation(SubObject, NextPropertyIterator, bOutInheritedOperation);
-				}
+				UObject* SubObject = TryGetInstancedSubObjectValue(ObjectProperty, SubValuePtr);
+				// This should not be needed in the property grid, as it should already been called on the subobject.
+				return SubObject ? OverridableManager.GetOverriddenPropertyOperation(SubObject, NextPropertyIterator, bOutInheritedOperation) : EOverriddenPropertyOperation::None;
 			}
 		}
 		// Special handling for array of instanced subobjects 
@@ -582,25 +580,24 @@ EOverriddenPropertyOperation FOverriddenPropertySet::GetOverriddenPropertyOperat
 				FScriptArrayHelper ArrayHelper(ArrayProperty, SubValuePtr);
 				if(ArrayHelper.IsValidIndex(ArrayIndex))
 				{
-					if (UObject* SubObject = TryGetInstancedSubObjectValue(InnerObjectProperty, ArrayHelper.GetElementPtr(ArrayIndex)))
+					UObject* SubObject = TryGetInstancedSubObjectValue(InnerObjectProperty, ArrayHelper.GetElementPtr(ArrayIndex));
+					if (NextPropertyIterator)
 					{
-						if (NextPropertyIterator)
-						{
-							// Forward any sub queries to the subobject
-							return OverridableManager.GetOverriddenPropertyOperation(SubObject, NextPropertyIterator, bOutInheritedOperation);
-						}
-						else if(CurrentOverriddenPropertyNode)
+						// Forward any sub queries to the subobject
+						return SubObject ? OverridableManager.GetOverriddenPropertyOperation(SubObject, NextPropertyIterator, bOutInheritedOperation) : EOverriddenPropertyOperation::None;
+					}
+
+					if(CurrentOverriddenPropertyNode && SubObject)
+					{
+						// Caller wants to know about any override state on the reference of the subobject itself
+						const FOverriddenPropertyNodeID  SubObjectID(SubObject);
+						if (const FOverriddenPropertyNode* SubObjectOverriddenPropertyNode = CurrentOverriddenPropertyNode->SubPropertyNodes.FindByKey(SubObjectID))
 						{
-							// Caller wants to know about any override state on the reference of the subobject itself
-							const FOverriddenPropertyNodeID  SubObjectID(SubObject);
-							if (const FOverriddenPropertyNode* SubObjectOverriddenPropertyNode = CurrentOverriddenPropertyNode->SubPropertyNodes.FindByKey(SubObjectID))
+							if (bOutInheritedOperation)
 							{
-								if (bOutInheritedOperation)
-								{
-									*bOutInheritedOperation = false;
-								}
-								return SubObjectOverriddenPropertyNode->Operation;
+								*bOutInheritedOperation = false;
 							}
+							return SubObjectOverriddenPropertyNode->Operation;
 						}
 					}
 				}
@@ -641,8 +638,11 @@ EOverriddenPropertyOperation FOverriddenPropertySet::GetOverriddenPropertyOperat
 							return OverridableManager.GetOverriddenPropertyOperation(ValueSubObject, NextPropertyIterator, bOutInheritedOperation);
 						}
 					}
+					// This is not a subobject property or an null subobject, return none at this point
+					return EOverriddenPropertyOperation::None;
 				}
-				else if(CurrentOverriddenPropertyNode)
+
+				if(CurrentOverriddenPropertyNode)
 				{
 					// Caller wants to know about any override state on the reference of the map pair itself
 					checkf(MapProperty->KeyProp, TEXT("Expecting a key type for Maps"));
@@ -720,18 +720,16 @@ bool FOverriddenPropertySet::ClearOverriddenProperty(FOverriddenPropertyNode& Pa
 		FPropertyVisitorPath::Iterator NextPropertyIterator = PropertyIterator+1;
 		if (const FObjectPropertyBase* ObjectProperty = CastField<FObjectPropertyBase>(CurrentProperty))
 		{
-			if (UObject* SubObject = TryGetInstancedSubObjectValue(ObjectProperty, SubValuePtr))
+			UObject* SubObject = TryGetInstancedSubObjectValue(ObjectProperty, SubValuePtr);
+			if (NextPropertyIterator)
 			{
-				if (NextPropertyIterator)
-				{
-					return OverridableManager.ClearOverriddenProperty(SubObject, NextPropertyIterator);
-				}
-				else
-				{
-					OverridableManager.ClearOverrides(SubObject);
-					bClearedOverrides = true;
-				}
+				return SubObject ? OverridableManager.ClearOverriddenProperty(SubObject, NextPropertyIterator) : false;
+			}
+			if (SubObject)
+			{
+				OverridableManager.ClearOverrides(SubObject);
 			}
+			bClearedOverrides = true;
 		}
 		// Special handling for array of instanced subobjects 
 		else if (const FArrayProperty* ArrayProperty = CastField<FArrayProperty>(CurrentProperty))
@@ -766,20 +764,18 @@ bool FOverriddenPropertySet::ClearOverriddenProperty(FOverriddenPropertyNode& Pa
 				}
 				else if(ArrayHelper.IsValidIndex(ArrayIndex))
 				{
-					if (UObject* SubObject = TryGetInstancedSubObjectValue(InnerObjectProperty, ArrayHelper.GetElementPtr(ArrayIndex)))
+					UObject* SubObject = TryGetInstancedSubObjectValue(InnerObjectProperty, ArrayHelper.GetElementPtr(ArrayIndex));
+					if (NextPropertyIterator)
 					{
-						if (NextPropertyIterator)
-						{
-							return OverridableManager.ClearOverriddenProperty(SubObject, NextPropertyIterator);
-						}
-						if (CurrentOverriddenPropertyNode)
+						return SubObject ? OverridableManager.ClearOverriddenProperty(SubObject, NextPropertyIterator) : false;
+					}
+					if (CurrentOverriddenPropertyNode && SubObject)
+					{
+						const FOverriddenPropertyNodeID  SubObjectID(SubObject);
+						if (CurrentOverriddenPropertyNode->SubPropertyNodes.Remove(SubObjectID))
 						{
-							const FOverriddenPropertyNodeID  SubObjectID(SubObject);
-							if (CurrentOverriddenPropertyNode->SubPropertyNodes.Remove(SubObjectID))
-							{
-								OverridableManager.ClearInstancedSubObjectOverrides(Owner, SubObject);
-								return true;
-							}
+							OverridableManager.ClearInstancedSubObjectOverrides(Owner, SubObject);
+							return true;
 						}
 					}
 				}
@@ -813,16 +809,18 @@ bool FOverriddenPropertySet::ClearOverriddenProperty(FOverriddenPropertyNode& Pa
 			// So forward the call to the instanced subobject
 			if (NextPropertyIterator)
 			{
-				if(MapHelper.IsValidIndex(InternalMapIndex))
+				if (MapHelper.IsValidIndex(InternalMapIndex))
 				{
 					checkf(MapProperty->ValueProp, TEXT("Expecting a value type for Maps"));
-					if (UObject* ValueSubObject = TryGetInstancedSubObjectValue(ValueObjectProperty, MapHelper.GetValuePtr(InternalMapIndex)))
-					{
-						return OverridableManager.ClearOverriddenProperty(ValueSubObject, NextPropertyIterator);
-					}
+					UObject* ValueSubObject = TryGetInstancedSubObjectValue(ValueObjectProperty, MapHelper.GetValuePtr(InternalMapIndex));
+					return ValueSubObject ? OverridableManager.ClearOverriddenProperty(ValueSubObject, NextPropertyIterator) : false;
 				}
+
+				// Unable the reconcile what we need to clear.
+				return false;
 			}
-			else if(InternalMapIndex == INDEX_NONE)
+
+			if(InternalMapIndex == INDEX_NONE)
 			{
 				// Users want to clear all of the overrides on the array, but in the case of instanced subobject, we need to clear the overrides on them as well.
 				if (ValueObjectProperty)
@@ -1733,9 +1731,8 @@ FOverriddenPropertyNode* FOverriddenPropertySet::SetOverriddenPropertyOperation(
 	while (PropertyIterator && OverriddenPropertyNode->Operation != EOverriddenPropertyOperation::Replace)
 	{
 		const FProperty* CurrentProperty = (*PropertyIterator);
-		// While digging down the path, if the operation is replace and one of the property is always overridden
-		// then there isn't anything to do
-		if (Operation == EOverriddenPropertyOperation::Replace && CurrentProperty->HasAnyPropertyFlags(CPF_ExperimentalAlwaysOverriden))
+		// While digging down the path, if the property is always overridden then there should not be anything to do more
+		if (CurrentProperty->HasAnyPropertyFlags(CPF_ExperimentalAlwaysOverriden))
 		{
 			return nullptr;
 		}
@@ -1746,6 +1743,11 @@ FOverriddenPropertyNode* FOverriddenPropertySet::SetOverriddenPropertyOperation(
 	// Might have stop before as one of the parent property was completely replaced.
 	if (!PropertyIterator)
 	{
+		// If everything is replaced from here, remove all subnode operations
+		if (Operation == EOverriddenPropertyOperation::Replace)
+		{
+			OverriddenPropertyNode->Reset();
+		}
 		OverriddenPropertyNode->Operation = Operation;
 		return OverriddenPropertyNode;
 	}
