Renamed CVarVerifyGCAssumptionsChance to CVarVerifyGCChance to avoid naming conflict in Unity builds.

#rnx

diff --git a/Engine/Source/Runtime/CoreUObject/Private/UObject/Obj.cpp b/Engine/Source/Runtime/CoreUObject/Private/UObject/Obj.cpp
--- a/Engine/Source/Runtime/CoreUObject/Private/UObject/Obj.cpp
+++ b/Engine/Source/Runtime/CoreUObject/Private/UObject/Obj.cpp
@@ -5716,11 +5716,11 @@ void StaticUObjectInit()
 			CVarVerifyGCAssumptionsOnFullPurge->Set(false, ECVF_SetByCommandline);
 		}
 	}
-	if (IConsoleVariable* CVarVerifyGCAssumptionsChance = IConsoleManager::Get().FindConsoleVariable(TEXT("gc.VerifyAssumptionsChance")))
+	if (IConsoleVariable* CVarVerifyGCChance = IConsoleManager::Get().FindConsoleVariable(TEXT("gc.VerifyAssumptionsChance")))
 	{
 		if (bCommandLineForcesNoVerifyGC)
 		{
-			CVarVerifyGCAssumptionsChance->Set(0.0f, ECVF_SetByCommandline);
+			CVarVerifyGCChance->Set(0.0f, ECVF_SetByCommandline);
 		}
 	}
 
