Prevent root objects from being added to the internal GRoots list when they are also in the disregard for GC pool.

We already prevented that from happening when the disregard for GC pool was open but it was possible to unset and set the root flag on a disregard for GC object after the the pool was closed and the engine would add that object to the GRoots list.

diff --git a/Engine/Source/Runtime/CoreUObject/Private/UObject/GarbageCollection.cpp b/Engine/Source/Runtime/CoreUObject/Private/UObject/GarbageCollection.cpp
--- a/Engine/Source/Runtime/CoreUObject/Private/UObject/GarbageCollection.cpp
+++ b/Engine/Source/Runtime/CoreUObject/Private/UObject/GarbageCollection.cpp
@@ -655,7 +655,11 @@ namespace UE::GC::Private
 
 		if (ObjectItem->HasAnyFlags(EInternalObjectFlags_RootFlags) || ObjectItem->GetRefCount() != 0)
 		{
-			GRoots.Add(Index);
+			// Disregard for GC objects should never be added to the roots list because they are meant to be entirely skipped by GC.
+			if (!GUObjectArray.IsIndexDisregardForGC(Index))
+			{
+				GRoots.Add(Index);
+			}
 		}
 		else
 		{
@@ -721,10 +725,7 @@ bool FUObjectItem::SetRootFlags(EInternalObjectFlags FlagsToSet)
 				// RootFlags can only be added while GC is not running, otherwise GC might not see the change
 				checkSlow(IsInGameThread() || IsInParallelGameThread() || IsGarbageCollectionLocked());
 
-				if (!(GUObjectArray.IsOpenForDisregardForGC() & GUObjectArray.DisregardForGCEnabled())) //-V792
-				{
-					MarkRootAsDirty();
-				}
+				MarkRootAsDirty();
 			}
 		}
 	}
@@ -751,10 +752,7 @@ bool FUObjectItem::ClearRootFlags(EInternalObjectFlags FlagsToClear)
 	{
 		if ((OldFlags & RootFlags) != 0 && ((OldFlags & ~(int32)FlagsToClear) & RootFlags) == 0 && OldRefCount == 0)
 		{
-			if (!(GUObjectArray.IsOpenForDisregardForGC() & GUObjectArray.DisregardForGCEnabled())) //-V792
-			{
-				MarkRootAsDirty();
-			}
+			MarkRootAsDirty();
 		}
 	}
 
diff --git a/Engine/Source/Runtime/CoreUObject/Private/UObject/UObjectArray.cpp b/Engine/Source/Runtime/CoreUObject/Private/UObject/UObjectArray.cpp
--- a/Engine/Source/Runtime/CoreUObject/Private/UObject/UObjectArray.cpp
+++ b/Engine/Source/Runtime/CoreUObject/Private/UObject/UObjectArray.cpp
@@ -105,7 +105,7 @@ void FUObjectItem::CreateStatID() const
 #endif
 FUObjectArray::FUObjectArray()
 : ObjFirstGCIndex(0)
-, ObjLastNonGCIndex(INDEX_NONE)
+, ObjLastNonGCIndex(-1)
 , MaxObjectsNotConsideredByGC(0)
 , OpenForDisregardForGC(true)
 , PrimarySerialNumber(START_SERIAL_NUMBER)
@@ -205,19 +205,10 @@ void FUObjectArray::CloseDisregardForGC()
 
 void FUObjectArray::DisableDisregardForGC()
 {
-	if (!GExitPurge && (ObjFirstGCIndex > 0 || DisregardForGCEnabled()))
-	{
-		checkf(!IsAsyncLoading(), TEXT("Disregard for GC Set can't be safely disabled while async loading. Consider calling FlushAsyncLoading() first or using gc.MaxObjectsNotConsideredByGC=0 ini setting instead."));
-
-		void OnDisregardForGCSetDisabled(int32 NumObjects);
-		// If disregard for GC was already closed then ObjFirstGCIndex is the number of objects we need to scan, otherwise disregard for GC is still open and we need to scan all objects
-		int32 NumDisregardForGCObjects = ObjFirstGCIndex > 0 ? ObjFirstGCIndex : GetObjectArrayNum();
-		if (NumDisregardForGCObjects > 0)
-		{
-			OnDisregardForGCSetDisabled(NumDisregardForGCObjects);
-		}
-	}
+	// ObjLastNonGCIndex defaults to -1 and is being incremented with each UObject allocation when the disregard for GC set is open and enabled
+	const int32 NumDisregardForGCObjects = ObjLastNonGCIndex + 1;
 
+	// Reset the diregard for GC state before calling OnDisregardForGCSetDisabled() in case it changes any of the root flags
 	MaxObjectsNotConsideredByGC = 0;
 	GUObjectAllocator.DisablePersistentAllocator();
 	ObjFirstGCIndex = 0;
@@ -226,6 +217,14 @@ void FUObjectArray::DisableDisregardForGC()
 	{
 		CloseDisregardForGC();
 	}
+
+	if (!GExitPurge && NumDisregardForGCObjects > 0)
+	{
+		checkf(!IsAsyncLoading(), TEXT("Disregard for GC Set can't be safely disabled while async loading. Consider calling FlushAsyncLoading() first or using gc.MaxObjectsNotConsideredByGC=0 ini setting instead."));
+
+		void OnDisregardForGCSetDisabled(int32 NumObjects);
+		OnDisregardForGCSetDisabled(NumDisregardForGCObjects);
+	}
 }
 
 void FUObjectArray::AllocateUObjectIndex(UObjectBase* Object, EInternalObjectFlags InitialFlags, int32 AlreadyAllocatedIndex, int32 SerialNumber, FRemoteObjectId RemoteId)
@@ -255,12 +254,14 @@ void FUObjectArray::AllocateUObjectIndex(UObjectBase* Object, EInternalObjectFla
 
 	LockInternalArray();
 
+	const bool bDisregardForGCOpenAndEnabled = (IsOpenForDisregardForGC() & GUObjectArray.DisregardForGCEnabled()); //-V792 
+
 	if (AlreadyAllocatedIndex >= 0)
 	{
 		Index = AlreadyAllocatedIndex;
 	}
 	// Special non- garbage collectable range.
-	else if (OpenForDisregardForGC && DisregardForGCEnabled())
+	else if (bDisregardForGCOpenAndEnabled)
 	{
 		Index = ++ObjLastNonGCIndex;
 		// Check if we're not out of bounds, unless there hasn't been any gc objects yet
@@ -296,7 +297,7 @@ void FUObjectArray::AllocateUObjectIndex(UObjectBase* Object, EInternalObjectFla
 	UE_CLOG(ObjectItem->GetObject() != nullptr, LogUObjectArray, Fatal, TEXT("Attempting to add %s at index %d but another object (0x%016llx) exists at that index!"), *Object->GetFName().ToString(), Index, (int64)(PTRINT)ObjectItem->GetObject());
 	// At this point all not-compiled-in objects are not fully constructed yet and this is the earliest we can mark them as such
 	ObjectItem->FlagsAndRefCount = (int64)((uint64)EInternalObjectFlags::PendingConstruction << 32);
-	if (!(IsOpenForDisregardForGC() & GUObjectArray.DisregardForGCEnabled())) //-V792
+	if (!bDisregardForGCOpenAndEnabled)
 	{
 		// It's safe to access FGCFlags::GetReachableFlagValue_ForGC() here because creating new objects is being performed 
 		// under the same UObjectArray lock as swapping reachability flags inside of GC, see FGCFlags::SwapReachableAndMaybeUnreachable()
diff --git a/Engine/Source/Runtime/CoreUObject/Public/UObject/UObjectArray.h b/Engine/Source/Runtime/CoreUObject/Public/UObject/UObjectArray.h
--- a/Engine/Source/Runtime/CoreUObject/Public/UObject/UObjectArray.h
+++ b/Engine/Source/Runtime/CoreUObject/Public/UObject/UObjectArray.h
@@ -1191,15 +1191,26 @@ public:
 		return ObjObjects.IsValidIndex(Object->InternalIndex);
 	}
 
+	/**
+	 * Returns true if this object index is "disregard for GC"
+	 *
+	 * @param Index object index to get for disregard for GC
+	 * @return true if this object index is disregard for GC
+	 */
+	UE_FORCEINLINE_HINT bool IsIndexDisregardForGC(int32 ObjectIndex) const
+	{
+		return ObjectIndex <= ObjLastNonGCIndex;
+	}
+
 	/**
 	 * Returns true if this object is "disregard for GC"...same results as the legacy RF_DisregardForGC flag
 	 *
 	 * @param Object object to get for disregard for GC
-	 * @return true if this object si disregard for GC
+	 * @return true if this object is disregard for GC
 	 */
-	UE_FORCEINLINE_HINT bool IsDisregardForGC(const class UObjectBase* Object)
+	UE_FORCEINLINE_HINT bool IsDisregardForGC(const class UObjectBase* Object) const
 	{
-		return Object->InternalIndex <= ObjLastNonGCIndex;
+		return IsIndexDisregardForGC(Object->InternalIndex);
 	}
 	/**
 	 * Returns the size of the global UObject array, some of these might be unused
