[Serialization]

Fixes a bug where when loading cooked content and swapping to uncooked content makes it impossible to then swap back to cooked content as the linkerload packagestore backend cache would maintain entries for the now uncooked package ids and report them as existing. We used to cache the linkerload state for packages within zenloader itself but once the caching was moved to the PackageStore, the cache cleanup was not ported over. Previously zenload would clean up cache entries as UPackages are destroyed during GC. Now we defer cache cleanup until assets are removed from the asset registry (asset delete or content unmounted).

- Makes the LinkerLoadPackageStoreBackend respond to content dismount events so that the internal packageid cache it maintains can be emptied when content is unmounted.
- Extends the IAssetRegistryInterface to allow registring delegates with the AssetRegistry from CoreUObject for AssetAdded/Removed.

#rnx

diff --git a/Engine/Source/Runtime/AssetRegistry/Private/AssetRegistry.cpp b/Engine/Source/Runtime/AssetRegistry/Private/AssetRegistry.cpp
--- a/Engine/Source/Runtime/AssetRegistry/Private/AssetRegistry.cpp
+++ b/Engine/Source/Runtime/AssetRegistry/Private/AssetRegistry.cpp
@@ -264,10 +264,52 @@ public:
 		IAssetRegistry* AssetRegistry = IAssetRegistry::Get();
 		if (!AssetRegistry)
 		{
-			return false;;
+			return false;
 		}
 		return AssetRegistry->EnumerateAssets(Filter, Callback, InEnumerateFlags);
 	}
+
+	virtual bool RegisterOnAssetsAddedDelegate(const TFunction<void(TConstArrayView<FAssetData>)>& Function, FDelegateHandle& OutHandle) override
+	{
+		IAssetRegistry* AssetRegistry = IAssetRegistry::Get();
+		if (!AssetRegistry)
+		{
+			return false;
+		}
+		OutHandle = AssetRegistry->OnAssetsAdded().AddLambda(Function);
+		return true;
+	}
+
+	virtual bool UnregisterOnAssetsAddedDelegate(const FDelegateHandle& Handle) override
+	{
+		IAssetRegistry* AssetRegistry = IAssetRegistry::Get();
+		if (!AssetRegistry)
+		{
+			return false;
+		}
+		return AssetRegistry->OnAssetsAdded().Remove(Handle);
+	}
+
+	virtual bool RegisterOnAssetsRemovedDelegate(const TFunction<void(TConstArrayView<FAssetData>)>& Function, FDelegateHandle& OutHandle) override
+	{
+		IAssetRegistry* AssetRegistry = IAssetRegistry::Get();
+		if (!AssetRegistry)
+		{
+			return false;
+		}
+		OutHandle = AssetRegistry->OnAssetsRemoved().AddLambda(Function);
+		return true;
+	}
+
+	virtual bool UnregisterOnAssetsRemovedDelegate(const FDelegateHandle& Handle) override
+	{
+		IAssetRegistry* AssetRegistry = IAssetRegistry::Get();
+		if (!AssetRegistry)
+		{
+			return false;
+		}
+		return AssetRegistry->OnAssetsRemoved().Remove(Handle);
+	}
 };
 FAssetRegistryInterface GAssetRegistryInterface;
 
diff --git a/Engine/Source/Runtime/AssetRegistry/Public/AssetRegistry/AssetRegistryModule.h b/Engine/Source/Runtime/AssetRegistry/Public/AssetRegistry/AssetRegistryModule.h
--- a/Engine/Source/Runtime/AssetRegistry/Public/AssetRegistry/AssetRegistryModule.h
+++ b/Engine/Source/Runtime/AssetRegistry/Public/AssetRegistry/AssetRegistryModule.h
@@ -125,4 +125,46 @@ public:
 		}
 		return AssetRegistry->EnumerateAssets(Filter, Callback, InEnumerateFlags);
 	}
+
+	virtual bool RegisterOnAssetsAddedDelegate(const TFunction<void(TConstArrayView<FAssetData>)>& Function, FDelegateHandle& OutHandle) override
+	{
+		IAssetRegistry* AssetRegistry = IAssetRegistry::Get();
+		if (!AssetRegistry)
+		{
+			return false;
+		}
+		OutHandle = AssetRegistry->OnAssetsAdded().AddLambda(Function);
+		return true;
+	}
+
+	virtual bool UnregisterOnAssetsAddedDelegate(const FDelegateHandle& Handle) override
+	{
+		IAssetRegistry* AssetRegistry = IAssetRegistry::Get();
+		if (!AssetRegistry)
+		{
+			return false;
+		}
+		return AssetRegistry->OnAssetsAdded().Remove(Handle);
+	}
+
+	virtual bool RegisterOnAssetsRemovedDelegate(const TFunction<void(TConstArrayView<FAssetData>)>& Function, FDelegateHandle& OutHandle) override
+	{
+		IAssetRegistry* AssetRegistry = IAssetRegistry::Get();
+		if (!AssetRegistry)
+		{
+			return false;
+		}
+		OutHandle = AssetRegistry->OnAssetsRemoved().AddLambda(Function);
+		return true;
+	}
+
+	virtual bool UnregisterOnAssetsRemovedDelegate(const FDelegateHandle& Handle) override
+	{
+		IAssetRegistry* AssetRegistry = IAssetRegistry::Get();
+		if (!AssetRegistry)
+		{
+			return false;
+		}
+		return AssetRegistry->OnAssetsRemoved().Remove(Handle);
+	}
 };
diff --git a/Engine/Source/Runtime/CoreUObject/Private/Serialization/PackageStore.cpp b/Engine/Source/Runtime/CoreUObject/Private/Serialization/PackageStore.cpp
--- a/Engine/Source/Runtime/CoreUObject/Private/Serialization/PackageStore.cpp
+++ b/Engine/Source/Runtime/CoreUObject/Private/Serialization/PackageStore.cpp
@@ -3,8 +3,10 @@
 #include "Serialization/PackageStore.h"
 
 #include "AssetRegistry/AssetData.h"
+#include "Async/ParallelFor.h"
 #include "Containers/StringConv.h"
 #include "Containers/StringView.h"
+#include "Containers/StripedMap.h"
 #include "Containers/UnrealString.h"
 #include "CoreGlobals.h"
 #include "CoreGlobalsInternal.h"
@@ -33,13 +35,48 @@ DEFINE_LOG_CATEGORY_STATIC(LogPackageStore, Log, All);
 class FLinkerLoadPackageStoreBackend final : public IPackageStoreBackend
 {
 public:
-	virtual EPackageLoader GetSupportedLoaders() override
+	~FLinkerLoadPackageStoreBackend()
 	{
-		return EPackageLoader::LinkerLoad;
+		if (OnAssetRemovedHandle.IsValid())
+		{
+			if (IAssetRegistryInterface* AssetRegistry = IAssetRegistryInterface::GetPtr())
+			{
+				AssetRegistry->UnregisterOnAssetsRemovedDelegate(OnAssetRemovedHandle);
+			}
+		}
+	}
+
+	void OnAssetsRemoved(TConstArrayView<FAssetData> RemovedAssets)
+	{
+		TRACE_CPUPROFILER_EVENT_SCOPE(FLinkerLoadPackageStoreBackend::OnAssetsRemoved);
+		const int32 MinBatchSize = 64;
+		ParallelFor(TEXT("FLinkerLoadPackageStoreBackend::OnAssetsRemoved_PF"), RemovedAssets.Num(), MinBatchSize, [this, &RemovedAssets](int Index)
+			{
+				const FAssetData& AssetData = RemovedAssets[Index];
+				FPackageId PackageId = FPackageId::FromName(AssetData.PackageName);
+				PackageStoreCache.Remove(PackageId);
+			});
 	}
 
 	virtual void OnMounted(TSharedRef<const FPackageStoreBackendContext>) override
 	{
+		// We need to defer registration for the AssetRegistry delegate since the registry might not be ready yet
+		FCoreDelegates::OnFEngineLoopInitComplete.AddLambda([this]()
+			{
+				IAssetRegistryInterface* AssetRegistry = IAssetRegistryInterface::GetPtr();
+				if (ensureMsgf(AssetRegistry, TEXT("It's expected we always have an AssetRegistry after Engine init.")))
+				{
+					AssetRegistry->RegisterOnAssetsRemovedDelegate([this](TConstArrayView<FAssetData> RemovedAssets) 
+						{
+							OnAssetsRemoved(RemovedAssets);
+						} , OnAssetRemovedHandle);
+				}
+			});
+	}
+
+	virtual EPackageLoader GetSupportedLoaders() override
+	{
+		return EPackageLoader::LinkerLoad;
 	}
 
 	virtual void BeginRead() override
@@ -52,20 +89,28 @@ public:
 
 	EPackageStoreEntryStatus GetPackageStoreEntryEx(FPackageId PackageId, FName PackageName, FPackageStoreEntry& OutPackageStoreEntry, FPackagePath* OutPackagePath)
 	{
-		FPackageStoreEntry* Entry = PackageStoreCache.Find(PackageId);
-		if (Entry)
-		{
-			if (Entry->LoaderType == EPackageLoader::LinkerLoad)
+		EPackageStoreEntryStatus Status = EPackageStoreEntryStatus::Missing;
+		bool bCacheHit = PackageStoreCache.FindAndApply(PackageId,
+			[&](const FPackageStoreEntry& Entry)
 			{
-				OutPackageStoreEntry = *Entry;
-				if (OutPackagePath)
+				if (Entry.LoaderType == EPackageLoader::LinkerLoad)
 				{
-					*OutPackagePath = FPackagePath::FromPackageNameUnchecked(PackageName);
-					OutPackagePath->SetHeaderExtension(OutPackageStoreEntry.PackageExtension);
+					OutPackageStoreEntry = Entry;
+					if (OutPackagePath)
+					{
+						*OutPackagePath = FPackagePath::FromPackageNameUnchecked(PackageName);
+						OutPackagePath->SetHeaderExtension(OutPackageStoreEntry.PackageExtension);
+					}
+					Status = EPackageStoreEntryStatus::Ok;
 				}
-				return EPackageStoreEntryStatus::Ok;
-			}
-			return EPackageStoreEntryStatus::Missing;
+				else
+				{
+					Status = EPackageStoreEntryStatus::Missing;
+				}
+			});
+		if (bCacheHit)
+		{
+			return Status;
 		}
 
 		if (PackageName.IsNone())
@@ -121,7 +166,8 @@ public:
 	}
 
 private:
-	TMap<FPackageId, FPackageStoreEntry> PackageStoreCache;
+	TStripedMap<32, FPackageId, FPackageStoreEntry> PackageStoreCache;
+	FDelegateHandle OnAssetRemovedHandle;
 };
 
 
diff --git a/Engine/Source/Runtime/CoreUObject/Public/Misc/AssetRegistryInterface.h b/Engine/Source/Runtime/CoreUObject/Public/Misc/AssetRegistryInterface.h
--- a/Engine/Source/Runtime/CoreUObject/Public/Misc/AssetRegistryInterface.h
+++ b/Engine/Source/Runtime/CoreUObject/Public/Misc/AssetRegistryInterface.h
@@ -505,6 +505,18 @@ public:
 	virtual bool EnumerateAssets(const FARFilter& Filter, TFunctionRef<bool(const FAssetData&)> Callback,
 		UE::AssetRegistry::EEnumerateAssetsFlags InEnumerateFlags = UE::AssetRegistry::EEnumerateAssetsFlags::None) const = 0;
 
+	/** Register delegate for when assets are added to the registry */
+	virtual bool RegisterOnAssetsAddedDelegate(const TFunction<void(TConstArrayView<FAssetData>)>& Function, FDelegateHandle& OutHandle) = 0;
+
+	/** Removes registered delegate for when assets are added to the registry */
+	virtual bool UnregisterOnAssetsAddedDelegate(const FDelegateHandle& Handle) = 0;
+
+	/** Register delegate for when assets are removed from the registry */
+	virtual bool RegisterOnAssetsRemovedDelegate(const TFunction<void(TConstArrayView<FAssetData>)>& Function, FDelegateHandle& OutHandle) = 0;
+
+	/** Removes registered delegate for when assets are removed from the registry */
+	virtual bool UnregisterOnAssetsRemovedDelegate(const FDelegateHandle& Handle) = 0;
+
 protected:
 
 	static COREUOBJECT_API IAssetRegistryInterface* Default;
