[CoreUObjectTests]

Add missing implementations to MockAssetRegistryInterface from 44599700

#rnx

diff --git a/Engine/Source/Runtime/CoreUObject/Tests/ObjectPropertyTest.cpp b/Engine/Source/Runtime/CoreUObject/Tests/ObjectPropertyTest.cpp
--- a/Engine/Source/Runtime/CoreUObject/Tests/ObjectPropertyTest.cpp
+++ b/Engine/Source/Runtime/CoreUObject/Tests/ObjectPropertyTest.cpp
@@ -339,6 +339,26 @@ public:
 		return true;
 	}
 
+	virtual bool RegisterOnAssetsAddedDelegate(const TFunction<void(TConstArrayView<FAssetData>)>& Function, FDelegateHandle& OutHandle) override
+	{
+		return false;
+	};
+
+	virtual bool UnregisterOnAssetsAddedDelegate(const FDelegateHandle& Handle) override
+	{
+		return false;
+	};
+
+	virtual bool RegisterOnAssetsRemovedDelegate(const TFunction<void(TConstArrayView<FAssetData>)>& Function, FDelegateHandle& OutHandle) override
+	{
+		return false;
+	};
+
+	virtual bool UnregisterOnAssetsRemovedDelegate(const FDelegateHandle& Handle) override
+	{
+		return false;
+	};
+
 	IAssetRegistryInterface* Old;
 	TMap<FSoftObjectPath, FAssetData> AssetData;
 };
