Fix DeferCppStructOps operations with merged modules

- Add RemoveAllDeferredCppStructOps to call on merged module unload
- Update Verse DeferCppStructOps to stop filtering duplicate calls over the program lifetime
- Update UScriptStruct DeferCppStructOps to skip over duplicate calls now that Verse is no longer guarding against that

#rnx
#tests Tested with merged modules & editor

diff --git a/Engine/Source/Runtime/CoreUObject/Private/UObject/Class.cpp b/Engine/Source/Runtime/CoreUObject/Private/UObject/Class.cpp
--- a/Engine/Source/Runtime/CoreUObject/Private/UObject/Class.cpp
+++ b/Engine/Source/Runtime/CoreUObject/Private/UObject/Class.cpp
@@ -3069,6 +3069,15 @@ void UScriptStruct::DeferCppStructOps(FTopLevelAssetPath Target, ICppStructOps*
 
 	if (UScriptStruct::ICppStructOps* ExistingOps = DeferredStructOps.FindRef(Target))
 	{
+#if UE_MERGED_MODULES
+
+		// When using merged modules, we'll manually unload deferred struct ops at module teardown.
+		// Whenever we get to this, it means the using code is doing multiple initialization, which we'll prevent here.
+		delete InCppStructOps;
+		return;
+
+#else
+
 		IReload* Reload = GetActiveReloadInterface();
 		if (Reload == nullptr)
 		{
@@ -3081,6 +3090,8 @@ void UScriptStruct::DeferCppStructOps(FTopLevelAssetPath Target, ICppStructOps*
 			return;
 		}
 		// in reload, we will just leak these...they may be in use.
+
+#endif // UE_MERGED_MODULES
 	}
 	DeferredStructOps.Add(Target, InCppStructOps);
 }
@@ -3100,6 +3111,32 @@ void UScriptStruct::RemoveDeferredCppStructOps(FTopLevelAssetPath Target)
 	}
 }
 
+/**
+ * On module unload, remove all entries in a package from the stash
+ * The package name should be the name component itself, e.g. "Engine" rather than "/Script/Engine"
+ * @param PackageName Name of the package
+ */
+void UScriptStruct::RemoveAllDeferredCppStructOps(FName PackageName)
+{
+	UE_LOG(LogClass, Verbose, TEXT("RemoveAllDeferredCppStructOps: %s"), *PackageName.ToString());
+
+	TMap<FTopLevelAssetPath, UScriptStruct::ICppStructOps*>& DeferredStructOps = GetDeferredCppStructOps();
+
+	// We're looking for paths that match "/PackageName" - possible full names include "/Script/PackageName" and "/Plugin/_Verse/VNI/PackageName"
+	const FString PackageNamePattern = TEXT("/") + PackageName.ToString();
+
+	for (auto DeferredStructOpsIt = DeferredStructOps.CreateIterator(); DeferredStructOpsIt; ++DeferredStructOpsIt)
+	{
+		if (WriteToString<FName::StringBufferSize>(DeferredStructOpsIt->Key.GetPackageName()).ToView().EndsWith(PackageNamePattern))
+		{
+			UScriptStruct::ICppStructOps* ExistingOps = DeferredStructOpsIt->Value;
+			check(ExistingOps);
+			delete ExistingOps;
+			DeferredStructOpsIt.RemoveCurrent();
+		}
+	}
+}
+
 FTopLevelAssetPath UScriptStruct::GetFlattenedStructPathName() const
 {
 	return FTopLevelAssetPath(GetPackage()->GetFName(), GetFName());
diff --git a/Engine/Source/Runtime/CoreUObject/Private/UObject/UObjectBase.cpp b/Engine/Source/Runtime/CoreUObject/Private/UObject/UObjectBase.cpp
--- a/Engine/Source/Runtime/CoreUObject/Private/UObject/UObjectBase.cpp
+++ b/Engine/Source/Runtime/CoreUObject/Private/UObject/UObjectBase.cpp
@@ -1123,7 +1123,7 @@ static void RemoveObjectsInModules(TConstArrayView<FName> InModuleNames)
 
 	for (const FName PackageName : InModuleNames)
 	{
-		UPackage* Package = GetScriptPackageFromModuleName (PackageName);
+		UPackage* Package = GetScriptPackageFromModuleName(PackageName);
 		if (Package)
 		{
 			UE_LOG(LogUObjectBase, Log, TEXT("RemoveObjectsInModules: removing package %s"), *Package->GetName());
@@ -1140,6 +1140,8 @@ static void RemoveObjectsInModules(TConstArrayView<FName> InModuleNames)
 
 				return true;
 			}, true);
+
+			UScriptStruct::RemoveAllDeferredCppStructOps(PackageName);
 		}
 	}
 }
diff --git a/Engine/Source/Runtime/CoreUObject/Public/UObject/Class.h b/Engine/Source/Runtime/CoreUObject/Public/UObject/Class.h
--- a/Engine/Source/Runtime/CoreUObject/Public/UObject/Class.h
+++ b/Engine/Source/Runtime/CoreUObject/Public/UObject/Class.h
@@ -2145,6 +2145,13 @@ public:
 	 */
 	static COREUOBJECT_API void RemoveDeferredCppStructOps(FTopLevelAssetPath Target);
 
+	/**
+	 * On module unload, remove all entries in a package from the stash
+	 * The package name should be the name component itself, e.g. "Engine" rather than "/Script/Engine"
+	 * @param PackageName Name of the package
+	 */
+	static COREUOBJECT_API void RemoveAllDeferredCppStructOps(FName PackageName);
+
 	template<class CPPSTRUCT>
 	static void DeferCppStructOps(FTopLevelAssetPath Target)
 	{
