Handle dirtying of the object from the overridable operation

[RN] Overridding commands are now dirtying object

diff --git a/Engine/Plugins/Experimental/InstanceDataObjectFixupTool/Source/Private/InstanceDataObjectFixupPanel.cpp b/Engine/Plugins/Experimental/InstanceDataObjectFixupTool/Source/Private/InstanceDataObjectFixupPanel.cpp
--- a/Engine/Plugins/Experimental/InstanceDataObjectFixupTool/Source/Private/InstanceDataObjectFixupPanel.cpp
+++ b/Engine/Plugins/Experimental/InstanceDataObjectFixupTool/Source/Private/InstanceDataObjectFixupPanel.cpp
@@ -844,7 +844,7 @@ void FInstanceDataObjectFixupPanel::RedirectProperty(const FPropertyPath& From,
 		if (SourceProperty->SameType(DestinationProperty))
 		{
 			SourceProperty->CopyCompleteValue(Destination, Source);
-			FOverridableManager::Get().GetOverriddenProperties(Instance)->SetOverriddenPropertyOperation(EOverriddenPropertyOperation::Modified, nullptr, DestinationProperty);
+			FOverridableManager::Get().GetOverriddenProperties(Instance)->RestoreOverriddenPropertyOperation(EOverriddenPropertyOperation::Modified, nullptr, DestinationProperty);
 		}
 		else
 		{
diff --git a/Engine/Source/Runtime/CoreUObject/Private/UObject/Class.cpp b/Engine/Source/Runtime/CoreUObject/Private/UObject/Class.cpp
--- a/Engine/Source/Runtime/CoreUObject/Private/UObject/Class.cpp
+++ b/Engine/Source/Runtime/CoreUObject/Private/UObject/Class.cpp
@@ -1918,7 +1918,7 @@ void UStruct::SerializeVersionedTaggedProperties(FStructuredArchive::FSlot Slot,
 								{
 									if (FOverriddenPropertySet* OverriddenProperties = FOverridableSerializationLogic::GetOverriddenProperties())
 									{
-										OverriddenProperties->RestoreOverriddenPropertyOperation(Tag.OverrideOperation, UnderlyingArchive.GetSerializedPropertyChain(), TempProperty);
+										OverriddenProperties->ConditionallyRestoreOverriddenPropertyOperation(Tag.OverrideOperation, UnderlyingArchive.GetSerializedPropertyChain(), TempProperty);
 									}
 								}
 							}
@@ -1942,7 +1942,7 @@ void UStruct::SerializeVersionedTaggedProperties(FStructuredArchive::FSlot Slot,
 					{
 						if (FOverriddenPropertySet* OverriddenProperties = FOverridableSerializationLogic::GetOverriddenProperties())
 						{
-							OverriddenProperties->RestoreOverriddenPropertyOperation(Tag.OverrideOperation, UnderlyingArchive.GetSerializedPropertyChain(), Property);
+							OverriddenProperties->ConditionallyRestoreOverriddenPropertyOperation(Tag.OverrideOperation, UnderlyingArchive.GetSerializedPropertyChain(), Property);
 						}
 					}
 
diff --git a/Engine/Source/Runtime/CoreUObject/Private/UObject/OverridableManager.cpp b/Engine/Source/Runtime/CoreUObject/Private/UObject/OverridableManager.cpp
--- a/Engine/Source/Runtime/CoreUObject/Private/UObject/OverridableManager.cpp
+++ b/Engine/Source/Runtime/CoreUObject/Private/UObject/OverridableManager.cpp
@@ -131,12 +131,12 @@ FOverriddenPropertySet* FOverridableManager::RestoreOverrideOperation(TNotNull<U
 	// property operations to None.
 	if (!bDontClearOverrides)
 	{
-		ObjectOverriddenProperties.Reset();
+		ObjectOverriddenProperties.Reset(/*bShouldDirtyObject*/false);
 	}
 
 	if (!bDontClearOverrides || ObjectOverriddenProperties.GetOverriddenPropertyOperation(nullptr, nullptr) == EOverriddenPropertyOperation::None)
 	{
-		ObjectOverriddenProperties.SetOverriddenPropertyOperation(Operation, /*CurrentPropertyChain*/nullptr, /*Property*/nullptr);
+		ObjectOverriddenProperties.RestoreOverriddenPropertyOperation(Operation, /*CurrentPropertyChain*/nullptr, /*Property*/nullptr);
 		ObjectOverriddenProperties.bNeedsSubobjectTemplateInstantiation = bNeedsSubobjectTemplateInstantiation;
 	}
 	return &ObjectOverriddenProperties;
@@ -369,7 +369,7 @@ void FOverridableManager::ClearOverrides(TNotNull<UObject*> Object)
 #if WITH_EDITORONLY_DATA
 	if (FOverriddenPropertySet* ThisObjectOverriddenProperties = Find(Object))
 	{
-		ThisObjectOverriddenProperties->Reset();
+		ThisObjectOverriddenProperties->Reset(/*bShouldDirtyObject*/true);
 	}
 	PropagateClearOverridesToInstancedSubObjects(Object);
 #endif // WITH_EDITORONLY_DATA
diff --git a/Engine/Source/Runtime/CoreUObject/Private/UObject/OverriddenPropertySet.cpp b/Engine/Source/Runtime/CoreUObject/Private/UObject/OverriddenPropertySet.cpp
--- a/Engine/Source/Runtime/CoreUObject/Private/UObject/OverriddenPropertySet.cpp
+++ b/Engine/Source/Runtime/CoreUObject/Private/UObject/OverriddenPropertySet.cpp
@@ -478,6 +478,68 @@ bool FOverriddenPropertyNode::Serialize(FArchive& Ar)
 	return true;
 }
 
+void FOverriddenPropertyNode::SetOperation(EOverriddenPropertyOperation InOperation, bool* bWasModified /*= nullptr*/)
+{
+	if (Operation != InOperation)
+	{
+		Operation = InOperation;
+		if (bWasModified)
+		{
+			*bWasModified = true;
+		}
+	}
+}
+
+FOverriddenPropertyNode& FOverriddenPropertyNode::FindOrAddNode(const FOverriddenPropertyNodeID& InNodeID, bool* bWasModified /*= nullptr*/)
+{
+	if (FOverriddenPropertyNode* SubNode = SubPropertyNodes.FindByKey(InNodeID))
+	{
+		return *SubNode;
+	}
+
+	if (bWasModified)
+	{
+		*bWasModified = true;
+	}
+
+
+	// We can safely assume that this node is at least modified from now on
+	if (Operation == EOverriddenPropertyOperation::None)
+	{
+		Operation = EOverriddenPropertyOperation::Modified;
+	}
+
+	// Not found add the node
+	return SubPropertyNodes.Emplace_GetRef(InNodeID);
+}
+
+bool FOverriddenPropertyNode::RemoveSubPropertyNode(const FOverriddenPropertyNodeID& InNodeID, bool* bWasModified /*= nullptr*/)
+{
+	if (SubPropertyNodes.Remove(InNodeID))
+	{
+		if (bWasModified)
+		{
+			*bWasModified = true;
+		}
+		return true;
+	}
+	return false;
+}
+
+void FOverriddenPropertyNode::Reset(bool* bWasModified /*= nullptr*/)
+{
+	SetOperation(EOverriddenPropertyOperation::None, bWasModified);
+
+	if (!SubPropertyNodes.IsEmpty())
+	{
+		SubPropertyNodes.Empty();
+		if (bWasModified)
+		{
+			*bWasModified = true;
+		}
+	}
+}
+
 void FOverriddenPropertyNode::HandleObjectsReInstantiated(const TMap<UObject*, UObject*>& Map)
 {
 	NodeID.HandleObjectsReInstantiated(Map);
@@ -505,23 +567,6 @@ void FOverriddenPropertySet::RestoreOverriddenState(const FOverriddenPropertySet
 	bWasAdded = FromOverriddenProperties.bWasAdded;
 }
 
-FOverriddenPropertyNode& FOverriddenPropertySet::FindOrAddNode(FOverriddenPropertyNode& ParentNode, FOverriddenPropertyNodeID NodeID)
-{
-	if (FOverriddenPropertyNode* SubNode = ParentNode.SubPropertyNodes.FindByKey(NodeID))
-	{
-		return *SubNode;
-	}
-
-	// We can safely assume that the parent node is at least modified from now on
-	if (ParentNode.Operation == EOverriddenPropertyOperation::None)
-	{
-		ParentNode.Operation = EOverriddenPropertyOperation::Modified;
-	}
-
-	// Not found add the node
-	return ParentNode.SubPropertyNodes.Emplace_GetRef(NodeID);
-}
-
 EOverriddenPropertyOperation FOverriddenPropertySet::GetOverriddenPropertyOperation(const FOverriddenPropertyNode* ParentPropertyNode, FPropertyVisitorPath::Iterator PropertyIterator, bool* bOutInheritedOperation, const void* Data) const
 {
 	FOverridableManager& OverridableManager = FOverridableManager::Get();
@@ -529,7 +574,7 @@ EOverriddenPropertyOperation FOverriddenPropertySet::GetOverriddenPropertyOperat
 	const void* SubValuePtr = Data;
 	const FOverriddenPropertyNode* OverriddenPropertyNode = ParentPropertyNode;
 	int32 ArrayIndex = INDEX_NONE;
-	while (PropertyIterator && (!OverriddenPropertyNode || OverriddenPropertyNode->Operation != EOverriddenPropertyOperation::Replace))
+	while (PropertyIterator && (!OverriddenPropertyNode || OverriddenPropertyNode->GetOperation() != EOverriddenPropertyOperation::Replace))
 	{
 		ArrayIndex = INDEX_NONE;
 
@@ -540,7 +585,7 @@ EOverriddenPropertyOperation FOverriddenPropertySet::GetOverriddenPropertyOperat
 		if (OverriddenPropertyNode)
 		{
 			FOverriddenPropertyNodeID NodeID(CurrentProperty);
-			if (const FOverriddenPropertyNode* FoundNode = OverriddenPropertyNode->SubPropertyNodes.FindByKey(NodeID))
+			if (const FOverriddenPropertyNode* FoundNode = OverriddenPropertyNode->GetSubPropertyNodes().FindByKey(NodeID))
 			{
 				CurrentOverriddenPropertyNode = FoundNode;
 			}
@@ -589,13 +634,13 @@ EOverriddenPropertyOperation FOverriddenPropertySet::GetOverriddenPropertyOperat
 					{
 						// Caller wants to know about any override state on the reference of the subobject itself
 						const FOverriddenPropertyNodeID  SubObjectID(SubObject);
-						if (const FOverriddenPropertyNode* SubObjectOverriddenPropertyNode = CurrentOverriddenPropertyNode->SubPropertyNodes.FindByKey(SubObjectID))
+						if (const FOverriddenPropertyNode* SubObjectOverriddenPropertyNode = CurrentOverriddenPropertyNode->GetSubPropertyNodes().FindByKey(SubObjectID))
 						{
 							if (bOutInheritedOperation)
 							{
 								*bOutInheritedOperation = false;
 							}
-							return SubObjectOverriddenPropertyNode->Operation;
+							return SubObjectOverriddenPropertyNode->GetOperation();
 						}
 					}
 				}
@@ -646,13 +691,13 @@ EOverriddenPropertyOperation FOverriddenPropertySet::GetOverriddenPropertyOperat
 					checkf(MapProperty->KeyProp, TEXT("Expecting a key type for Maps"));
 					FOverriddenPropertyNodeID OverriddenKeyID = FOverriddenPropertyNodeID::FromMapKey(MapProperty->KeyProp, MapHelper.GetKeyPtr(InternalMapIndex));
 
-					if (const FOverriddenPropertyNode* SubObjectOverriddenPropertyNode = CurrentOverriddenPropertyNode->SubPropertyNodes.FindByKey(OverriddenKeyID))
+					if (const FOverriddenPropertyNode* SubObjectOverriddenPropertyNode = CurrentOverriddenPropertyNode->GetSubPropertyNodes().FindByKey(OverriddenKeyID))
 					{
 						if (bOutInheritedOperation)
 						{
 							*bOutInheritedOperation = false;
 						}
-						return SubObjectOverriddenPropertyNode->Operation;
+						return SubObjectOverriddenPropertyNode->GetOperation();
 					}
 				}
 			}
@@ -677,11 +722,20 @@ EOverriddenPropertyOperation FOverriddenPropertySet::GetOverriddenPropertyOperat
 	{
 		*bOutInheritedOperation = PropertyIterator || ArrayIndex != INDEX_NONE;
 	}
-	return OverriddenPropertyNode ? OverriddenPropertyNode->Operation : EOverriddenPropertyOperation::None;
+	return OverriddenPropertyNode ? OverriddenPropertyNode->GetOperation() : EOverriddenPropertyOperation::None;
 }
 
 bool FOverriddenPropertySet::ClearOverriddenProperty(FOverriddenPropertyNode& ParentPropertyNode, FPropertyVisitorPath::Iterator PropertyIterator, const void* Data)
 {
+	bool bWasModified = false;
+	ON_SCOPE_EXIT
+	{
+		if (bWasModified)
+		{
+			Owner->Modify();
+		}
+	};
+
 	FOverridableManager& OverridableManager = FOverridableManager::Get();
 	if (!PropertyIterator)
 	{
@@ -696,7 +750,7 @@ bool FOverriddenPropertySet::ClearOverriddenProperty(FOverriddenPropertyNode& Pa
 	int32 ArrayIndex = INDEX_NONE;
 	TArray<FOverriddenPropertyNode*> TraversedNodes;
 	TraversedNodes.Push(OverriddenPropertyNode);
-	while (PropertyIterator && (!OverriddenPropertyNode || OverriddenPropertyNode->Operation != EOverriddenPropertyOperation::Replace))
+	while (PropertyIterator && (!OverriddenPropertyNode || OverriddenPropertyNode->GetOperation() != EOverriddenPropertyOperation::Replace))
 	{
 		ArrayIndex = INDEX_NONE;
 
@@ -707,7 +761,7 @@ bool FOverriddenPropertySet::ClearOverriddenProperty(FOverriddenPropertyNode& Pa
 		if (OverriddenPropertyNode)
 		{
 			FOverriddenPropertyNodeID NodeID(CurrentProperty);
-			if (FOverriddenPropertyNode* FoundPropertyNode = OverriddenPropertyNode->SubPropertyNodes.FindByKey(NodeID))
+			if (FOverriddenPropertyNode* FoundPropertyNode = OverriddenPropertyNode->GetSubPropertyNodes().FindByKey(NodeID))
 			{
 				CurrentOverriddenPropertyNode = FoundPropertyNode;
 				TraversedNodes.Push(CurrentOverriddenPropertyNode);
@@ -770,7 +824,7 @@ bool FOverriddenPropertySet::ClearOverriddenProperty(FOverriddenPropertyNode& Pa
 					if (CurrentOverriddenPropertyNode && SubObject)
 					{
 						const FOverriddenPropertyNodeID  SubObjectID(SubObject);
-						if (CurrentOverriddenPropertyNode->SubPropertyNodes.Remove(SubObjectID))
+						if (CurrentOverriddenPropertyNode->RemoveSubPropertyNode(SubObjectID, &bWasModified))
 						{
 							OverridableManager.ClearInstancedSubObjectOverrides(Owner, SubObject);
 							return true;
@@ -841,7 +895,7 @@ bool FOverriddenPropertySet::ClearOverriddenProperty(FOverriddenPropertyNode& Pa
 				FOverriddenPropertyNodeID OverriddenKeyID = FOverriddenPropertyNodeID::FromMapKey(MapProperty->KeyProp, MapHelper.GetKeyPtr(InternalMapIndex));
 
 				FOverriddenPropertyNodeID CurrentPropKey;
-				if (CurrentOverriddenPropertyNode->SubPropertyNodes.Remove(OverriddenKeyID))
+				if (CurrentOverriddenPropertyNode->RemoveSubPropertyNode(OverriddenKeyID, &bWasModified))
 				{
 					if (UObject* ValueSubObject = TryGetInstancedSubObjectValue(ValueObjectProperty, MapHelper.GetValuePtr(InternalMapIndex)))
 					{
@@ -858,7 +912,7 @@ bool FOverriddenPropertySet::ClearOverriddenProperty(FOverriddenPropertyNode& Pa
 		++PropertyIterator;
 	}
 
-	auto CleanupClearedNodes = [this, &TraversedNodes]()
+	auto CleanupClearedNodes = [this, &TraversedNodes, &bWasModified]()
 	{
 		// Go through each traversed property in reversed order to do cleanup
 		// We need to continue the cleanup until there is more overrides than just the one we are removing
@@ -868,15 +922,15 @@ bool FOverriddenPropertySet::ClearOverriddenProperty(FOverriddenPropertyNode& Pa
 			if (LastCleanedNode)
 			{
 				// In the case there are other overrides, just cleanup that node and stop.
-				if (CurrentNode->SubPropertyNodes.Num() > 1)
+				if (CurrentNode->GetSubPropertyNodes().Num() > 1)
 				{
-					verifyf(CurrentNode->SubPropertyNodes.Remove(LastCleanedNode->NodeID), TEXT("Expecting the node to always be removed"));
+					verifyf(CurrentNode->RemoveSubPropertyNode(LastCleanedNode->GetNodeID(), &bWasModified), TEXT("Expecting the node to always be removed"));
 					LastCleanedNode = nullptr;
 					break;
 				}
 			}
 
-			CurrentNode->Reset();
+			CurrentNode->Reset(&bWasModified);
 			LastCleanedNode = CurrentNode;
 		}
 	};
@@ -902,6 +956,15 @@ bool FOverriddenPropertySet::ClearOverriddenProperty(FOverriddenPropertyNode& Pa
 
 void FOverriddenPropertySet::NotifyPropertyChange(FOverriddenPropertyNode* ParentPropertyNode, const EPropertyNotificationType Notification, FPropertyVisitorPath::Iterator PropertyIterator, const EPropertyChangeType::Type ChangeType, const void* Data, bool& bNeedsCleanup)
 {
+	bool bWasModified = false;
+	ON_SCOPE_EXIT
+	{
+		if (bWasModified)
+		{
+			Owner->Modify();
+		}
+	};
+
 	checkf(IsValid(Owner), TEXT("Expecting a valid overridable owner"));
 
 	if (ChangeType & EPropertyChangeType::ResetToDefault)
@@ -930,8 +993,8 @@ void FOverriddenPropertySet::NotifyPropertyChange(FOverriddenPropertyNode* Paren
 	FOverriddenPropertyNode* SubPropertyNode = nullptr;
 	if (ParentPropertyNode)
 	{
-		FOverriddenPropertyNode& SubPropertyNodeRef = FindOrAddNode(*ParentPropertyNode, FOverriddenPropertyNodeID(Property));
-		SubPropertyNode = SubPropertyNodeRef.Operation != EOverriddenPropertyOperation::Replace ? &SubPropertyNodeRef : nullptr;
+		FOverriddenPropertyNode& SubPropertyNodeRef = ParentPropertyNode->FindOrAddNode(FOverriddenPropertyNodeID(Property), &bWasModified);
+		SubPropertyNode = SubPropertyNodeRef.GetOperation() != EOverriddenPropertyOperation::Replace ? &SubPropertyNodeRef : nullptr;
 	}
 
 	ON_SCOPE_EXIT
@@ -941,16 +1004,16 @@ void FOverriddenPropertySet::NotifyPropertyChange(FOverriddenPropertyNode* Paren
 			return;
 		}
 
-		if (SubPropertyNode && SubPropertyNode->SubPropertyNodes.IsEmpty() &&
+		if (SubPropertyNode && SubPropertyNode->GetSubPropertyNodes().IsEmpty() &&
 			(bNeedsCleanup || 
-			 SubPropertyNode->Operation == EOverriddenPropertyOperation::None || 
-			 SubPropertyNode->Operation == EOverriddenPropertyOperation::Modified))
+			 SubPropertyNode->GetOperation() == EOverriddenPropertyOperation::None || 
+			 SubPropertyNode->GetOperation() == EOverriddenPropertyOperation::Modified))
 		{
 			FOverriddenPropertyNodeID RemoveNodeID(Property);
-			ParentPropertyNode->SubPropertyNodes.Remove(RemoveNodeID);
-			if (ParentPropertyNode->Operation == EOverriddenPropertyOperation::Modified && ParentPropertyNode->SubPropertyNodes.IsEmpty())
+			ParentPropertyNode->RemoveSubPropertyNode(RemoveNodeID, &bWasModified);
+			if (ParentPropertyNode->GetOperation() == EOverriddenPropertyOperation::Modified && ParentPropertyNode->GetSubPropertyNodes().IsEmpty())
 			{
-				ParentPropertyNode->Operation = EOverriddenPropertyOperation::None;
+				ParentPropertyNode->SetOperation(EOverriddenPropertyOperation::None, &bWasModified);
 			}
 		}
 	};
@@ -1057,7 +1120,7 @@ void FOverriddenPropertySet::NotifyPropertyChange(FOverriddenPropertyNode* Paren
 					if (SubPropertyNode)
 					{
 						// Overriding all entry in the array
-						SubPropertyNode->Operation = EOverriddenPropertyOperation::Replace;
+						SubPropertyNode->SetOperation(EOverriddenPropertyOperation::Replace, &bWasModified);
 					}
 
 					// This is a case of the entire array is overridden
@@ -1079,8 +1142,8 @@ void FOverriddenPropertySet::NotifyPropertyChange(FOverriddenPropertyNode* Paren
 						if(SubPropertyNode)
 						{
 							const FOverriddenPropertyNodeID  AddedSubObjectID(AddedSubObject);
-							FOverriddenPropertyNode& AddedSubObjectNode = FindOrAddNode(*SubPropertyNode, AddedSubObjectID);
-							AddedSubObjectNode.Operation = EOverriddenPropertyOperation::Add;
+							FOverriddenPropertyNode& AddedSubObjectNode = SubPropertyNode->FindOrAddNode(AddedSubObjectID, &bWasModified);
+							AddedSubObjectNode.SetOperation(EOverriddenPropertyOperation::Add, &bWasModified);
 
 							// Notify the subobject that it was added
 							if (FOverriddenPropertySet* AddedSubObjectOverriddenProperties = OverridableManager.GetOverriddenProperties(AddedSubObject))
@@ -1103,33 +1166,33 @@ void FOverriddenPropertySet::NotifyPropertyChange(FOverriddenPropertyNode* Paren
 							if (RemovedSubObjectArchetype && !RemovedSubObjectArchetype->HasAnyFlags(RF_ClassDefaultObject))
 							{
 								const FOverriddenPropertyNodeID RemovedSubObjectID (RemovedSubObjectArchetype);
-								FOverriddenPropertyNode& RemovedSubObjectNode = FindOrAddNode(*SubPropertyNode, RemovedSubObjectID);
-								if (RemovedSubObjectNode.Operation == EOverriddenPropertyOperation::Add)
+								FOverriddenPropertyNode& RemovedSubObjectNode = SubPropertyNode->FindOrAddNode(RemovedSubObjectID, &bWasModified);
+								if (RemovedSubObjectNode.GetOperation() == EOverriddenPropertyOperation::Add)
 								{
 									// An add then a remove becomes no opt
-									if (SubPropertyNode->SubPropertyNodes.Remove(RemovedSubObjectID))
+									if (SubPropertyNode->RemoveSubPropertyNode(RemovedSubObjectID, &bWasModified))
 									{
 										bNeedsCleanup = true;
 									}
 								}
 								else
 								{
-									RemovedSubObjectNode.Operation = EOverriddenPropertyOperation::Remove;
+									RemovedSubObjectNode.SetOperation(EOverriddenPropertyOperation::Remove, &bWasModified);
 								}
 							}
 							else
 							{
 								// Figure out if it is a remove of a previously added element
 								const FOverriddenPropertyNodeID RemovedSubObjectID (RemovedSubObject);
-								if (const FOverriddenPropertyNode* AddedSubObjectNode = SubPropertyNode->SubPropertyNodes.FindByKey(RemovedSubObjectID))
+								if (const FOverriddenPropertyNode* AddedSubObjectNode = SubPropertyNode->GetSubPropertyNodes().FindByKey(RemovedSubObjectID))
 								{
-									if (AddedSubObjectNode->Operation != EOverriddenPropertyOperation::Add)
+									if (AddedSubObjectNode->GetOperation() != EOverriddenPropertyOperation::Add)
 									{
 										UE_LOG(LogOverridableObject, Warning, TEXT("This removed object:%s(0x%p) was not tracked as an add in the overridden properties"), *GetNameSafe(RemovedSubObject), RemovedSubObject);
 									}
 
 									// An add then a remove becomes no opt
-									if (SubPropertyNode->SubPropertyNodes.Remove(RemovedSubObjectID))
+									if (SubPropertyNode->RemoveSubPropertyNode(RemovedSubObjectID, &bWasModified))
 									{
 										bNeedsCleanup = true;
 									}
@@ -1154,7 +1217,7 @@ void FOverriddenPropertySet::NotifyPropertyChange(FOverriddenPropertyNode* Paren
 					else if (SubPropertyNode)
 					{
 						// Overriding all entry in the array
-						SubPropertyNode->Operation = EOverriddenPropertyOperation::Replace;
+						SubPropertyNode->SetOperation(EOverriddenPropertyOperation::Replace, &bWasModified);
 					}
 					return;
 				}
@@ -1292,7 +1355,7 @@ void FOverriddenPropertySet::NotifyPropertyChange(FOverriddenPropertyNode* Paren
 				// Overriding a all entries in the map
 				if (SubPropertyNode)
 				{
-					SubPropertyNode->Operation = EOverriddenPropertyOperation::Replace;
+					SubPropertyNode->SetOperation(EOverriddenPropertyOperation::Replace, &bWasModified);
 				}
 
 				// This is a case of the entire array is overridden
@@ -1302,8 +1365,8 @@ void FOverriddenPropertySet::NotifyPropertyChange(FOverriddenPropertyNode* Paren
 					if(SubPropertyNode)
 					{
 						FOverriddenPropertyNodeID OverriddenKeyID = FOverriddenPropertyNodeID::FromMapKey(MapProperty->KeyProp, MapHelper.GetKeyPtr(It.GetInternalIndex()));
-						FOverriddenPropertyNode& OverriddenKeyNode = FindOrAddNode(*SubPropertyNode, OverriddenKeyID);
-						OverriddenKeyNode.Operation = EOverriddenPropertyOperation::Replace;
+						FOverriddenPropertyNode& OverriddenKeyNode = SubPropertyNode->FindOrAddNode(OverriddenKeyID, &bWasModified);
+						OverriddenKeyNode.SetOperation(EOverriddenPropertyOperation::Replace, &bWasModified);
 					}
 
 					// @todo support instanced object as a key in maps
@@ -1326,8 +1389,8 @@ void FOverriddenPropertySet::NotifyPropertyChange(FOverriddenPropertyNode* Paren
 				if(SubPropertyNode)
 				{
 					FOverriddenPropertyNodeID AddedKeyID = FOverriddenPropertyNodeID::FromMapKey(MapProperty->KeyProp, MapHelper.GetKeyPtr(InternalMapIndex));
-					FOverriddenPropertyNode& AddedKeyNode = FindOrAddNode(*SubPropertyNode, AddedKeyID);
-					AddedKeyNode.Operation = EOverriddenPropertyOperation::Add;
+					FOverriddenPropertyNode& AddedKeyNode = SubPropertyNode->FindOrAddNode(AddedKeyID, &bWasModified);
+					AddedKeyNode.SetOperation(EOverriddenPropertyOperation::Add, &bWasModified);
 				}
 			};
 
@@ -1338,18 +1401,18 @@ void FOverriddenPropertySet::NotifyPropertyChange(FOverriddenPropertyNode* Paren
 				if(SubPropertyNode)
 				{
 					FOverriddenPropertyNodeID RemovedKeyID = FOverriddenPropertyNodeID::FromMapKey(MapProperty->KeyProp, PreEditMapHelper.GetKeyPtr(InternalPreEditMapIndex));
-					FOverriddenPropertyNode& RemovedKeyNode = FindOrAddNode(*SubPropertyNode, RemovedKeyID);
-					if (RemovedKeyNode.Operation == EOverriddenPropertyOperation::Add)
+					FOverriddenPropertyNode& RemovedKeyNode = SubPropertyNode->FindOrAddNode(RemovedKeyID, &bWasModified);
+					if (RemovedKeyNode.GetOperation() == EOverriddenPropertyOperation::Add)
 					{
 						// @Todo support remove/add/remove
-						if (SubPropertyNode->SubPropertyNodes.Remove(RemovedKeyID))
+						if (SubPropertyNode->RemoveSubPropertyNode(RemovedKeyID, &bWasModified))
 						{
 							bNeedsCleanup = true;
 						}
 					}
 					else
 					{
-						RemovedKeyNode.Operation = EOverriddenPropertyOperation::Remove;
+						RemovedKeyNode.SetOperation(EOverriddenPropertyOperation::Remove, &bWasModified);
 					}
 				}
 			};
@@ -1365,7 +1428,7 @@ void FOverriddenPropertySet::NotifyPropertyChange(FOverriddenPropertyNode* Paren
 				else if(SubPropertyNode)
 				{
 					// Overriding all entry in the array
-					SubPropertyNode->Operation = EOverriddenPropertyOperation::Replace;
+					SubPropertyNode->SetOperation(EOverriddenPropertyOperation::Replace, &bWasModified);
 				}
 				return;
 			}
@@ -1460,7 +1523,7 @@ void FOverriddenPropertySet::NotifyPropertyChange(FOverriddenPropertyNode* Paren
 		{
 			if (Notification == EPropertyNotificationType::PostEdit && SubPropertyNode)
 			{
-				SubPropertyNode->Operation = EOverriddenPropertyOperation::Replace;
+				SubPropertyNode->SetOperation(EOverriddenPropertyOperation::Replace, &bWasModified);
 			}
 		}
 		else if (!NextPropertyIterator)
@@ -1496,7 +1559,7 @@ void FOverriddenPropertySet::NotifyPropertyChange(FOverriddenPropertyNode* Paren
 			{
 				if (SubPropertyNode)
 				{
-					SubPropertyNode->Operation = EOverriddenPropertyOperation::Replace;
+					SubPropertyNode->SetOperation(EOverriddenPropertyOperation::Replace, &bWasModified);
 				}
 
 				uint8* PreEditValuePtr = SavedPreEditContainers.Find(Owner, ObjectProperty);
@@ -1528,7 +1591,7 @@ void FOverriddenPropertySet::NotifyPropertyChange(FOverriddenPropertyNode* Paren
 		{
 			if (Notification == EPropertyNotificationType::PostEdit && SubPropertyNode)
 			{
-				SubPropertyNode->Operation = EOverriddenPropertyOperation::Replace;
+				SubPropertyNode->SetOperation(EOverriddenPropertyOperation::Replace, &bWasModified);
 			}
 		}
 		else if (OptionalProperty->IsSet(Data))
@@ -1544,7 +1607,7 @@ void FOverriddenPropertySet::NotifyPropertyChange(FOverriddenPropertyNode* Paren
 		if (SubPropertyNode)
 		{
 			// Replacing this entire property
-			SubPropertyNode->Operation = EOverriddenPropertyOperation::Replace;
+			SubPropertyNode->SetOperation(EOverriddenPropertyOperation::Replace, &bWasModified);
 		}
 	}
 }
@@ -1622,12 +1685,12 @@ EOverriddenPropertyOperation FOverriddenPropertySet::GetOverriddenPropertyOperat
 	return GetOverriddenPropertyOperation(&RootNode, CurrentPropertyChain, Property);
 }
 
-FOverriddenPropertyNode* FOverriddenPropertySet::SetOverriddenPropertyOperation(EOverriddenPropertyOperation Operation, const FArchiveSerializedPropertyChain* CurrentPropertyChain, FProperty* Property)
+FOverriddenPropertyNode* FOverriddenPropertySet::RestoreOverriddenPropertyOperation(EOverriddenPropertyOperation Operation, const FArchiveSerializedPropertyChain* CurrentPropertyChain, FProperty* Property)
 {
-	return SetOverriddenPropertyOperation(Operation, RootNode, CurrentPropertyChain, Property);
+	return RestoreOverriddenPropertyOperation(Operation, RootNode, CurrentPropertyChain, Property);
 }
 
-FOverriddenPropertyNode* FOverriddenPropertySet::RestoreOverriddenPropertyOperation(EOverriddenPropertyOperation Operation, const FArchiveSerializedPropertyChain* CurrentPropertyChain, FProperty* Property)
+FOverriddenPropertyNode* FOverriddenPropertySet::ConditionallyRestoreOverriddenPropertyOperation(EOverriddenPropertyOperation Operation, const FArchiveSerializedPropertyChain* CurrentPropertyChain, FProperty* Property)
 {
 	// 'None', 'Modified' and 'SubObjectShadowing' operations are not needed to be restored on a property,
 	// because 'None' is equal to the node not existing and
@@ -1639,7 +1702,7 @@ FOverriddenPropertyNode* FOverriddenPropertySet::RestoreOverriddenPropertyOperat
 		// Prevent marking as replaced the properties that are always overridden
 		if (!Property || Operation != EOverriddenPropertyOperation::Replace || !Property->HasAnyPropertyFlags(CPF_ExperimentalAlwaysOverriden))
 		{
-			return SetOverriddenPropertyOperation(Operation, CurrentPropertyChain, Property);
+			return RestoreOverriddenPropertyOperation(Operation, CurrentPropertyChain, Property);
 		}
 	}
 
@@ -1651,9 +1714,9 @@ const FOverriddenPropertyNode* FOverriddenPropertySet::GetOverriddenPropertyNode
 	return GetOverriddenPropertyNode(RootNode, CurrentPropertyChain);
 }
 
-const FOverriddenPropertyNode* FOverriddenPropertySet::GetOverriddenPropertyNode(FOverriddenPropertyNodeID NodeID) const
+const FOverriddenPropertyNode* FOverriddenPropertySet::GetOverriddenPropertyNode(const FOverriddenPropertyNodeID& NodeID) const
 {
-	if (const FOverriddenPropertyNode* FoundNode = RootNode.SubPropertyNodes.FindByKey(NodeID))
+	if (const FOverriddenPropertyNode* FoundNode = RootNode.GetSubPropertyNodes().FindByKey(NodeID))
 	{
 		return FoundNode;
 	}
@@ -1665,7 +1728,7 @@ EOverriddenPropertyOperation FOverriddenPropertySet::GetOverriddenPropertyOperat
 	// No need to look further
 	// if it is the entire property is replaced or
 	// if it is the FOverriddenPropertySet struct which is always Overridden
-	if (ParentPropertyNode && ParentPropertyNode->Operation == EOverriddenPropertyOperation::Replace)
+	if (ParentPropertyNode && ParentPropertyNode->GetOperation() == EOverriddenPropertyOperation::Replace)
 	{
 		return EOverriddenPropertyOperation::Replace;
 	}
@@ -1679,13 +1742,13 @@ EOverriddenPropertyOperation FOverriddenPropertySet::GetOverriddenPropertyOperat
 
 	TArray<class FProperty*, TInlineAllocator<8>>::TConstIterator PropertyIterator = PropertyChain.GetRootIterator();
 	const FOverriddenPropertyNode* OverriddenPropertyNode = ParentPropertyNode;
-	while (PropertyIterator && (!OverriddenPropertyNode || (OverriddenPropertyNode->Operation != EOverriddenPropertyOperation::Replace)))
+	while (PropertyIterator && (!OverriddenPropertyNode || (OverriddenPropertyNode->GetOperation() != EOverriddenPropertyOperation::Replace)))
 	{
 		const FProperty* CurrentProperty = (*PropertyIterator);
 		if (OverriddenPropertyNode)
 		{
 			FOverriddenPropertyNodeID NodeID(CurrentProperty);
-			if (const FOverriddenPropertyNode* FoundNode = OverriddenPropertyNode->SubPropertyNodes.FindByKey(NodeID))
+			if (const FOverriddenPropertyNode* FoundNode = OverriddenPropertyNode->GetSubPropertyNodes().FindByKey(NodeID))
 			{
 				OverriddenPropertyNode = FoundNode;
 				checkf(OverriddenPropertyNode, TEXT("Expecting a node"));
@@ -1704,15 +1767,15 @@ EOverriddenPropertyOperation FOverriddenPropertySet::GetOverriddenPropertyOperat
 		++PropertyIterator;
 	}
 
-	return OverriddenPropertyNode ? OverriddenPropertyNode->Operation : EOverriddenPropertyOperation::None;
+	return OverriddenPropertyNode ? OverriddenPropertyNode->GetOperation() : EOverriddenPropertyOperation::None;
 }
 
-FOverriddenPropertyNode* FOverriddenPropertySet::SetOverriddenPropertyOperation(EOverriddenPropertyOperation Operation, FOverriddenPropertyNode& ParentPropertyNode, const FArchiveSerializedPropertyChain* CurrentPropertyChain, FProperty* Property)
+FOverriddenPropertyNode* FOverriddenPropertySet::RestoreOverriddenPropertyOperation(EOverriddenPropertyOperation Operation, FOverriddenPropertyNode& ParentPropertyNode, const FArchiveSerializedPropertyChain* CurrentPropertyChain, FProperty* Property)
 {
 	// No need to look further
 	// if it is the entire property is replaced or
 	// if it is the FOverriddenPropertySet struct which is always Overridden
-	if (ParentPropertyNode.Operation == EOverriddenPropertyOperation::Replace && Property)
+	if (ParentPropertyNode.GetOperation() == EOverriddenPropertyOperation::Replace && Property)
 	{
 		return nullptr;
 	}
@@ -1726,7 +1789,7 @@ FOverriddenPropertyNode* FOverriddenPropertySet::SetOverriddenPropertyOperation(
 
 	TArray<class FProperty*, TInlineAllocator<8>>::TConstIterator PropertyIterator = PropertyChain.GetRootIterator();
 	FOverriddenPropertyNode* OverriddenPropertyNode = &ParentPropertyNode;
-	while (PropertyIterator && OverriddenPropertyNode->Operation != EOverriddenPropertyOperation::Replace)
+	while (PropertyIterator && OverriddenPropertyNode->GetOperation() != EOverriddenPropertyOperation::Replace)
 	{
 		const FProperty* CurrentProperty = (*PropertyIterator);
 		// While digging down the path, if the property is always overridden then there should not be anything to do more
@@ -1734,7 +1797,7 @@ FOverriddenPropertyNode* FOverriddenPropertySet::SetOverriddenPropertyOperation(
 		{
 			return nullptr;
 		}
-		OverriddenPropertyNode = &FindOrAddNode(*OverriddenPropertyNode, FOverriddenPropertyNodeID(CurrentProperty));
+		OverriddenPropertyNode = &OverriddenPropertyNode->FindOrAddNode(FOverriddenPropertyNodeID(CurrentProperty));
 		++PropertyIterator;
 	}
 
@@ -1746,33 +1809,33 @@ FOverriddenPropertyNode* FOverriddenPropertySet::SetOverriddenPropertyOperation(
 		{
 			OverriddenPropertyNode->Reset();
 		}
-		OverriddenPropertyNode->Operation = Operation;
+		OverriddenPropertyNode->SetOperation(Operation);
 		return OverriddenPropertyNode;
 	}
 
 	return nullptr;
 }
 
-EOverriddenPropertyOperation FOverriddenPropertySet::GetSubPropertyOperation(FOverriddenPropertyNodeID NodeID) const
+EOverriddenPropertyOperation FOverriddenPropertySet::GetSubPropertyOperation(const FOverriddenPropertyNodeID& NodeID) const
 {
 	if (const FOverriddenPropertyNode* SubPropertyNode = GetOverriddenPropertyNode(NodeID))
 	{
-		return SubPropertyNode->Operation;
+		return SubPropertyNode->GetOperation();
 	}
 	return EOverriddenPropertyOperation::None;
 }
 
-FOverriddenPropertyNode* FOverriddenPropertySet::SetSubPropertyOperation(EOverriddenPropertyOperation Operation, FOverriddenPropertyNode& Node, FOverriddenPropertyNodeID NodeID)
+FOverriddenPropertyNode* FOverriddenPropertySet::RestoreSubPropertyOperation(EOverriddenPropertyOperation Operation, FOverriddenPropertyNode& Node, const FOverriddenPropertyNodeID& NodeID)
 {
-	FOverriddenPropertyNode& OverriddenPropertyNode = FindOrAddNode(Node, NodeID);
-	OverriddenPropertyNode.Operation = Operation;
+	FOverriddenPropertyNode& OverriddenPropertyNode = Node.FindOrAddNode(NodeID);
+	OverriddenPropertyNode.SetOperation(Operation);
 	return &OverriddenPropertyNode;
 }
 
-FOverriddenPropertyNode* FOverriddenPropertySet::SetSubObjectOperation(EOverriddenPropertyOperation Operation, FOverriddenPropertyNode& Node, TNotNull<UObject*> SubObject)
+FOverriddenPropertyNode* FOverriddenPropertySet::RestoreSubObjectOperation(EOverriddenPropertyOperation Operation, FOverriddenPropertyNode& Node, TNotNull<UObject*> SubObject)
 {
 	const FOverriddenPropertyNodeID SubObjectID(SubObject);
-	FOverriddenPropertyNode* SubObjectNode = SetSubPropertyOperation(Operation, Node, SubObjectID);
+	FOverriddenPropertyNode* SubObjectNode = RestoreSubPropertyOperation(Operation, Node, SubObjectID);
 
 	if (Operation == EOverriddenPropertyOperation::Add)
 	{
@@ -1800,9 +1863,14 @@ bool FOverriddenPropertySet::IsCDOOwningProperty(const FProperty& Property) cons
 	return Property.GetOwnerClass() == Owner->GetClass();
 }
 
-void FOverriddenPropertySet::Reset()
+void FOverriddenPropertySet::Reset(bool bShouldDirtyObject)
 {
-	RootNode.Reset();
+	bool bWasModified = false;
+	RootNode.Reset(&bWasModified);
+	if (bWasModified && bShouldDirtyObject)
+	{
+		Owner->Modify();
+	}
 }
 
 void FOverriddenPropertySet::HandleObjectsReInstantiated(const TMap<UObject*, UObject*>& Map)
@@ -1837,7 +1905,7 @@ const FOverriddenPropertyNode* FOverriddenPropertySet::GetOverriddenPropertyNode
 	{
 		const FProperty* CurrentProperty = (*PropertyIterator);
 		FOverriddenPropertyNodeID NodeID(CurrentProperty);
-		if (const FOverriddenPropertyNode* FoundNode = OverriddenPropertyNode->SubPropertyNodes.FindByKey(NodeID))
+		if (const FOverriddenPropertyNode* FoundNode = OverriddenPropertyNode->GetSubPropertyNodes().FindByKey(NodeID))
 		{
 			OverriddenPropertyNode = FoundNode;
 		}
diff --git a/Engine/Source/Runtime/CoreUObject/Private/UObject/Property.cpp b/Engine/Source/Runtime/CoreUObject/Private/UObject/Property.cpp
--- a/Engine/Source/Runtime/CoreUObject/Private/UObject/Property.cpp
+++ b/Engine/Source/Runtime/CoreUObject/Private/UObject/Property.cpp
@@ -1701,7 +1701,7 @@ const TCHAR* FProperty::ImportSingleProperty( const TCHAR* Str, void* DestData,
 			{
 				checkf(Path, TEXT("Expecting a path"));
 				FArchiveSerializedPropertyChain Chain = Path->ToSerializedPropertyChain();
-				OverriddenProperties->RestoreOverriddenPropertyOperation(Operation.GetValue(), &Chain, /*Property*/nullptr);
+				OverriddenProperties->ConditionallyRestoreOverriddenPropertyOperation(Operation.GetValue(), &Chain, /*Property*/nullptr);
 			}
 		}
 
@@ -1976,10 +1976,10 @@ const TCHAR* FProperty::ImportSingleProperty( const TCHAR* Str, void* DestData,
 								{
 									// Need to fetch the ArrayOverriddenPropertyNode every loop as the previous iteration might have reallocated the node.
 									FArchiveSerializedPropertyChain Chain = Path->ToSerializedPropertyChain();
-									if (FOverriddenPropertyNode* ArrayOverriddenPropertyNode = OverriddenProperties ? OverriddenProperties->SetOverriddenPropertyOperation(EOverriddenPropertyOperation::Modified, &Chain, /*Property*/nullptr) : nullptr)
+									if (FOverriddenPropertyNode* ArrayOverriddenPropertyNode = OverriddenProperties ? OverriddenProperties->RestoreOverriddenPropertyOperation(EOverriddenPropertyOperation::Modified, &Chain, /*Property*/nullptr) : nullptr)
 									{
 										// Rebuild the overridden info
-										OverriddenProperties->SetSubObjectOperation(EOverriddenPropertyOperation::Remove, *ArrayOverriddenPropertyNode, RemovedSubObject);
+										OverriddenProperties->RestoreSubObjectOperation(EOverriddenPropertyOperation::Remove, *ArrayOverriddenPropertyNode, RemovedSubObject);
 									}
 								}
 							}
@@ -2022,10 +2022,10 @@ const TCHAR* FProperty::ImportSingleProperty( const TCHAR* Str, void* DestData,
 									{
 										// Need to fetch the ArrayOverriddenPropertyNode every loop as the previous iteration might have reallocated the node.
 										FArchiveSerializedPropertyChain Chain = Path->ToSerializedPropertyChain();
-										if (FOverriddenPropertyNode* ArrayOverriddenPropertyNode = OverriddenProperties ? OverriddenProperties->SetOverriddenPropertyOperation(EOverriddenPropertyOperation::Modified, &Chain, /*Property*/nullptr) : nullptr)
+										if (FOverriddenPropertyNode* ArrayOverriddenPropertyNode = OverriddenProperties ? OverriddenProperties->RestoreOverriddenPropertyOperation(EOverriddenPropertyOperation::Modified, &Chain, /*Property*/nullptr) : nullptr)
 										{
 											// Rebuild the overridden info
-											OverriddenProperties->SetSubObjectOperation(EOverriddenPropertyOperation::Add, *ArrayOverriddenPropertyNode, AddedSubObject);
+											OverriddenProperties->RestoreSubObjectOperation(EOverriddenPropertyOperation::Add, *ArrayOverriddenPropertyNode, AddedSubObject);
 										}
 									}
 								}
@@ -2068,11 +2068,11 @@ const TCHAR* FProperty::ImportSingleProperty( const TCHAR* Str, void* DestData,
 							{
 								// Need to fetch the MapOverriddenPropertyNode every loop as the previous might have reallocated the node.
 								FArchiveSerializedPropertyChain Chain = Path->ToSerializedPropertyChain();
-								if (FOverriddenPropertyNode* MapOverriddenPropertyNode = OverriddenProperties ? OverriddenProperties->SetOverriddenPropertyOperation(EOverriddenPropertyOperation::Modified, &Chain, /*Property*/nullptr) : nullptr)
+								if (FOverriddenPropertyNode* MapOverriddenPropertyNode = OverriddenProperties ? OverriddenProperties->RestoreOverriddenPropertyOperation(EOverriddenPropertyOperation::Modified, &Chain, /*Property*/nullptr) : nullptr)
 								{
 									// Rebuild the overridden info
 									FOverriddenPropertyNodeID RemovedKeyID = FOverriddenPropertyNodeID::FromMapKey(MapProperty->KeyProp, TempKeyValueStorage);
-									OverriddenProperties->SetSubPropertyOperation(EOverriddenPropertyOperation::Remove, *MapOverriddenPropertyNode, RemovedKeyID);
+									OverriddenProperties->RestoreSubPropertyOperation(EOverriddenPropertyOperation::Remove, *MapOverriddenPropertyNode, RemovedKeyID);
 								}
 							}
 
@@ -2111,11 +2111,11 @@ const TCHAR* FProperty::ImportSingleProperty( const TCHAR* Str, void* DestData,
 							{
 								// Need to fetch the MapOverriddenPropertyNode every loop as the previous might have reallocated the node.
 								FArchiveSerializedPropertyChain Chain = Path->ToSerializedPropertyChain();
-								if (FOverriddenPropertyNode* MapOverriddenPropertyNode = OverriddenProperties ? OverriddenProperties->SetOverriddenPropertyOperation(EOverriddenPropertyOperation::Modified, &Chain, /*Property*/nullptr) : nullptr)
+								if (FOverriddenPropertyNode* MapOverriddenPropertyNode = OverriddenProperties ? OverriddenProperties->RestoreOverriddenPropertyOperation(EOverriddenPropertyOperation::Modified, &Chain, /*Property*/nullptr) : nullptr)
 								{
 									// Rebuild the overridden info
 									FOverriddenPropertyNodeID ModifiedKeyID = FOverriddenPropertyNodeID::FromMapKey(MapProperty->KeyProp, TempKeyValueStorage);
-									OverriddenProperties->SetSubPropertyOperation(EOverriddenPropertyOperation::Modified, *MapOverriddenPropertyNode, ModifiedKeyID);
+									OverriddenProperties->RestoreSubPropertyOperation(EOverriddenPropertyOperation::Modified, *MapOverriddenPropertyNode, ModifiedKeyID);
 								}
 							}
 							return Result;
@@ -2152,11 +2152,11 @@ const TCHAR* FProperty::ImportSingleProperty( const TCHAR* Str, void* DestData,
 							{
 								// Need to fetch the MapOverriddenPropertyNode every loop as the previous might have reallocated the node.
 								FArchiveSerializedPropertyChain Chain = Path->ToSerializedPropertyChain();
-								if (FOverriddenPropertyNode* MapOverriddenPropertyNode = OverriddenProperties ? OverriddenProperties->SetOverriddenPropertyOperation(EOverriddenPropertyOperation::Modified, &Chain, /*Property*/nullptr) : nullptr)
+								if (FOverriddenPropertyNode* MapOverriddenPropertyNode = OverriddenProperties ? OverriddenProperties->RestoreOverriddenPropertyOperation(EOverriddenPropertyOperation::Modified, &Chain, /*Property*/nullptr) : nullptr)
 								{
 									// Rebuild the overridden info
 									FOverriddenPropertyNodeID AddedKeyID = FOverriddenPropertyNodeID::FromMapKey(MapProperty->KeyProp, TempKeyValueStorage);
-									OverriddenProperties->SetSubPropertyOperation(EOverriddenPropertyOperation::Add, *MapOverriddenPropertyNode, AddedKeyID);
+									OverriddenProperties->RestoreSubPropertyOperation(EOverriddenPropertyOperation::Add, *MapOverriddenPropertyNode, AddedKeyID);
 								}
 							}
 							return Result;
diff --git a/Engine/Source/Runtime/CoreUObject/Private/UObject/PropertyArray.cpp b/Engine/Source/Runtime/CoreUObject/Private/UObject/PropertyArray.cpp
--- a/Engine/Source/Runtime/CoreUObject/Private/UObject/PropertyArray.cpp
+++ b/Engine/Source/Runtime/CoreUObject/Private/UObject/PropertyArray.cpp
@@ -332,10 +332,10 @@ void FArrayProperty::SerializeItem(FStructuredArchive::FSlot Slot, void* Value,
 							}
 
 							// Need to fetch the ArrayOverriddenPropertyNode every loop as the previous iteration might have reallocated the node.
-							if (FOverriddenPropertyNode* ArrayOverriddenPropertyNode = OverriddenProperties ? OverriddenProperties->SetOverriddenPropertyOperation(EOverriddenPropertyOperation::Modified, UnderlyingArchive.GetSerializedPropertyChain(), /*Property*/nullptr) : nullptr)
+							if (FOverriddenPropertyNode* ArrayOverriddenPropertyNode = OverriddenProperties ? OverriddenProperties->RestoreOverriddenPropertyOperation(EOverriddenPropertyOperation::Modified, UnderlyingArchive.GetSerializedPropertyChain(), /*Property*/nullptr) : nullptr)
 							{
 								// Rebuild the overridden info
-								OverriddenProperties->SetSubObjectOperation(EOverriddenPropertyOperation::Remove, *ArrayOverriddenPropertyNode, RemovedSubObject);
+								OverriddenProperties->RestoreSubObjectOperation(EOverriddenPropertyOperation::Remove, *ArrayOverriddenPropertyNode, RemovedSubObject);
 							}
 						}
 					}
@@ -448,10 +448,10 @@ void FArrayProperty::SerializeItem(FStructuredArchive::FSlot Slot, void* Value,
 							InnerObjectProperty->SetObjectPropertyValue(ArrayHelper.GetRawPtr(Index), AddedSubObject);
 
 							// Need to fetch the ArrayOverriddenPropertyNode every loop as the previous iteration might have reallocated the node.
-							if (FOverriddenPropertyNode* ArrayOverriddenPropertyNode = OverriddenProperties ? OverriddenProperties->SetOverriddenPropertyOperation(EOverriddenPropertyOperation::Modified, UnderlyingArchive.GetSerializedPropertyChain(), /*Property*/nullptr) : nullptr)
+							if (FOverriddenPropertyNode* ArrayOverriddenPropertyNode = OverriddenProperties ? OverriddenProperties->RestoreOverriddenPropertyOperation(EOverriddenPropertyOperation::Modified, UnderlyingArchive.GetSerializedPropertyChain(), /*Property*/nullptr) : nullptr)
 							{
 								// Rebuild the overridden info
-								OverriddenProperties->SetSubObjectOperation(EOverriddenPropertyOperation::Add, *ArrayOverriddenPropertyNode, AddedSubObject);
+								OverriddenProperties->RestoreSubObjectOperation(EOverriddenPropertyOperation::Add, *ArrayOverriddenPropertyNode, AddedSubObject);
 							}
 						}
 					}
@@ -507,7 +507,7 @@ void FArrayProperty::SerializeItem(FStructuredArchive::FSlot Slot, void* Value,
 
 					if (OverriddenProperties && ArrayOverrideOp != EOverriddenPropertyOperation::None)
 					{
-						auto FindObject = [InnerObjectProperty, bFilterExternalObjects](const FOverriddenPropertyNodeID ObjectToFind, FScriptArrayHelper& ArrayHelper) -> int32
+						auto FindObject = [InnerObjectProperty, bFilterExternalObjects](const FOverriddenPropertyNodeID& ObjectToFind, FScriptArrayHelper& ArrayHelper) -> int32
 						{
 							const int32 ArrayNum = ArrayHelper.Num();
 							for (int i = 0; i < ArrayNum; ++i)
@@ -528,26 +528,26 @@ void FArrayProperty::SerializeItem(FStructuredArchive::FSlot Slot, void* Value,
 
 						if (const FOverriddenPropertyNode* ArrayOverriddenPropertyNode = OverriddenProperties->GetOverriddenPropertyNode(UnderlyingArchive.GetSerializedPropertyChain()))
 						{
-							for (const FOverriddenPropertyNode& SubNode : ArrayOverriddenPropertyNode->SubPropertyNodes)
+							for (const FOverriddenPropertyNode& SubNode : ArrayOverriddenPropertyNode->GetSubPropertyNodes())
 							{
-								switch (SubNode.Operation)
+								switch (SubNode.GetOperation())
 								{
 								case EOverriddenPropertyOperation::Remove:
 									{
-										const int32 DefaultIndex = FindObject(SubNode.NodeID, DefaultsArrayHelper);
+										const int32 DefaultIndex = FindObject(SubNode.GetNodeID(), DefaultsArrayHelper);
 										if (DefaultIndex != INDEX_NONE)
 										{
 											RemovedIndices.Add(DefaultIndex);
 										}
 										else
 										{
-											UE_LOG(LogOverridableObject, VeryVerbose, TEXT("Unable to save deleted item %s"), *SubNode.NodeID.ToDebugString());
+											UE_LOG(LogOverridableObject, VeryVerbose, TEXT("Unable to save deleted item %s"), *SubNode.GetNodeID().ToDebugString());
 										}
 										break;
 									}
 								case EOverriddenPropertyOperation::Add:
 									{
-										const int32 Index = FindObject(SubNode.NodeID, ArrayHelper);
+										const int32 Index = FindObject(SubNode.GetNodeID(), ArrayHelper);
  										if (Index != INDEX_NONE)
 										{
 											ShadowIndices.Remove(Index);
@@ -559,7 +559,7 @@ void FArrayProperty::SerializeItem(FStructuredArchive::FSlot Slot, void* Value,
 										}
 										else
 										{
-											UE_LOG(LogOverridableObject, VeryVerbose, TEXT("Unable to save added item %s"), *SubNode.NodeID.ToDebugString());
+											UE_LOG(LogOverridableObject, VeryVerbose, TEXT("Unable to save added item %s"), *SubNode.GetNodeID().ToDebugString());
 										}
 										break;
 									}
diff --git a/Engine/Source/Runtime/CoreUObject/Private/UObject/PropertyMap.cpp b/Engine/Source/Runtime/CoreUObject/Private/UObject/PropertyMap.cpp
--- a/Engine/Source/Runtime/CoreUObject/Private/UObject/PropertyMap.cpp
+++ b/Engine/Source/Runtime/CoreUObject/Private/UObject/PropertyMap.cpp
@@ -369,11 +369,11 @@ void FMapProperty::SerializeItem(FStructuredArchive::FSlot Slot, void* Value, co
 							if (MapHelper.RemovePair(TempKeyValueStorage))
 							{
 								// Need to fetch the MapOverriddenPropertyNode every loop as the previous might have reallocated the node.
-								if (FOverriddenPropertyNode* MapOverriddenPropertyNode = OverriddenProperties ? OverriddenProperties->SetOverriddenPropertyOperation(EOverriddenPropertyOperation::Modified, UnderlyingArchive.GetSerializedPropertyChain(), /*Property*/nullptr) : nullptr)
+								if (FOverriddenPropertyNode* MapOverriddenPropertyNode = OverriddenProperties ? OverriddenProperties->RestoreOverriddenPropertyOperation(EOverriddenPropertyOperation::Modified, UnderlyingArchive.GetSerializedPropertyChain(), /*Property*/nullptr) : nullptr)
 								{
 									// Rebuild the overridden info
 									FOverriddenPropertyNodeID RemovedKeyID = FOverriddenPropertyNodeID::FromMapKey(KeyProp, TempKeyValueStorage);
-									OverriddenProperties->SetSubPropertyOperation(EOverriddenPropertyOperation::Remove, *MapOverriddenPropertyNode, RemovedKeyID);
+									OverriddenProperties->RestoreSubPropertyOperation(EOverriddenPropertyOperation::Remove, *MapOverriddenPropertyNode, RemovedKeyID);
 								}
 							}
 						}
@@ -414,11 +414,11 @@ void FMapProperty::SerializeItem(FStructuredArchive::FSlot Slot, void* Value, co
 							if (Index != INDEX_NONE)
 							{
 								// Need to fetch the MapOverriddenPropertyNode every loop as the previous might have reallocated the node.
-								if (FOverriddenPropertyNode* MapOverriddenPropertyNode = OverriddenProperties ? OverriddenProperties->SetOverriddenPropertyOperation(EOverriddenPropertyOperation::Modified, UnderlyingArchive.GetSerializedPropertyChain(), /*Property*/nullptr) : nullptr)
+								if (FOverriddenPropertyNode* MapOverriddenPropertyNode = OverriddenProperties ? OverriddenProperties->RestoreOverriddenPropertyOperation(EOverriddenPropertyOperation::Modified, UnderlyingArchive.GetSerializedPropertyChain(), /*Property*/nullptr) : nullptr)
 								{
 									// Rebuild the overridden info
 									FOverriddenPropertyNodeID ModifiedKeyID = FOverriddenPropertyNodeID::FromMapKey(KeyProp, TempKeyValueStorage);
-									OverriddenProperties->SetSubPropertyOperation(EOverriddenPropertyOperation::Modified, *MapOverriddenPropertyNode, ModifiedKeyID);
+									OverriddenProperties->RestoreSubPropertyOperation(EOverriddenPropertyOperation::Modified, *MapOverriddenPropertyNode, ModifiedKeyID);
 								}
 							}
 						}
@@ -497,11 +497,11 @@ void FMapProperty::SerializeItem(FStructuredArchive::FSlot Slot, void* Value, co
 							}
 
 							// Need to fetch the MapOverriddenPropertyNode every loop as the previous might have reallocated the node.
-							if (FOverriddenPropertyNode* MapOverriddenPropertyNode = OverriddenProperties ? OverriddenProperties->SetOverriddenPropertyOperation(EOverriddenPropertyOperation::Modified, UnderlyingArchive.GetSerializedPropertyChain(), /*Property*/nullptr) : nullptr)
+							if (FOverriddenPropertyNode* MapOverriddenPropertyNode = OverriddenProperties ? OverriddenProperties->RestoreOverriddenPropertyOperation(EOverriddenPropertyOperation::Modified, UnderlyingArchive.GetSerializedPropertyChain(), /*Property*/nullptr) : nullptr)
 							{
 								// Rebuild the overridden info
 								FOverriddenPropertyNodeID AddedKeyID = FOverriddenPropertyNodeID::FromMapKey(KeyProp, TempKeyValueStorage);
-								OverriddenProperties->SetSubPropertyOperation(EOverriddenPropertyOperation::Add, *MapOverriddenPropertyNode, AddedKeyID);
+								OverriddenProperties->RestoreSubPropertyOperation(EOverriddenPropertyOperation::Add, *MapOverriddenPropertyNode, AddedKeyID);
 							}
 						}
 					}
@@ -555,13 +555,13 @@ void FMapProperty::SerializeItem(FStructuredArchive::FSlot Slot, void* Value, co
 							if (const FOverriddenPropertyNode* MapOverriddenPropertyNode = OverriddenProperties->GetOverriddenPropertyNode(UnderlyingArchive.GetSerializedPropertyChain()))
 							{
 								// Figure out the modifications of the map
-								for (const FOverriddenPropertyNode& SubNode : MapOverriddenPropertyNode->SubPropertyNodes)
+								for (const FOverriddenPropertyNode& SubNode : MapOverriddenPropertyNode->GetSubPropertyNodes())
 								{
-									switch (SubNode.Operation)
+									switch (SubNode.GetOperation())
 									{
 										case EOverriddenPropertyOperation::Remove:
 										{
-											const int32 InternalIndex = SubNode.NodeID.ToMapInternalIndex(DefaultsMapHelper);
+											const int32 InternalIndex = SubNode.GetNodeID().ToMapInternalIndex(DefaultsMapHelper);
 											if (InternalIndex != INDEX_NONE)
 											{
 												RemovedIndices.Add(InternalIndex);
@@ -570,7 +570,7 @@ void FMapProperty::SerializeItem(FStructuredArchive::FSlot Slot, void* Value, co
 										}
 										case EOverriddenPropertyOperation::Add:
 										{
-											const int32 InternalIndex = SubNode.NodeID.ToMapInternalIndex(MapHelper);
+											const int32 InternalIndex = SubNode.GetNodeID().ToMapInternalIndex(MapHelper);
 											if (InternalIndex != INDEX_NONE)
 											{
 												AddedIndices.Add(InternalIndex);
@@ -580,7 +580,7 @@ void FMapProperty::SerializeItem(FStructuredArchive::FSlot Slot, void* Value, co
 										}
 										case EOverriddenPropertyOperation::Modified:
 										{
-											const int32 InternalIndex = SubNode.NodeID.ToMapInternalIndex(MapHelper);
+											const int32 InternalIndex = SubNode.GetNodeID().ToMapInternalIndex(MapHelper);
 											if (InternalIndex != INDEX_NONE)
 											{
 												ModifiedIndices.Add(InternalIndex);
diff --git a/Engine/Source/Runtime/CoreUObject/Public/UObject/OverriddenPropertySet.h b/Engine/Source/Runtime/CoreUObject/Public/UObject/OverriddenPropertySet.h
--- a/Engine/Source/Runtime/CoreUObject/Public/UObject/OverriddenPropertySet.h
+++ b/Engine/Source/Runtime/CoreUObject/Public/UObject/OverriddenPropertySet.h
@@ -329,6 +329,7 @@ struct FOverriddenPropertyNode
 		: NodeID(MoveTemp(InNodeID))
 	{}
 
+private:
 	UPROPERTY()
 	FOverriddenPropertyNodeID NodeID;
 
@@ -340,17 +341,41 @@ struct FOverriddenPropertyNode
 	// Instead will be serializing by hand
 	TArray<FOverriddenPropertyNode> SubPropertyNodes;
 
-	bool operator==(const FOverriddenPropertyNode& Other) const
+public:
+
+	const FOverriddenPropertyNodeID& GetNodeID() const
 	{
-		return NodeID == Other.NodeID;
+		return NodeID;
+	}
+
+	EOverriddenPropertyOperation GetOperation() const
+	{
+		return Operation;
+	}
+
+	void SetOperation(EOverriddenPropertyOperation InOperation, bool* bWasModified = nullptr);
+
+	TArrayView<FOverriddenPropertyNode> GetSubPropertyNodes()
+	{
+		return SubPropertyNodes;
 	}
 
-	void Reset()
+	const TArray<FOverriddenPropertyNode>& GetSubPropertyNodes() const
 	{
-		Operation = EOverriddenPropertyOperation::None;
-		SubPropertyNodes.Empty();
+		return SubPropertyNodes;
 	}
 
+	FOverriddenPropertyNode& FindOrAddNode(const FOverriddenPropertyNodeID& NodeID, bool* bWasModified = nullptr);
+
+	bool RemoveSubPropertyNode(const FOverriddenPropertyNodeID& NodeID, bool* bWasModified = nullptr);
+
+	bool operator==(const FOverriddenPropertyNode& Other) const
+	{
+		return NodeID == Other.NodeID;
+	}
+
+	void Reset(bool* bWasModified = nullptr);
+
 	void HandleObjectsReInstantiated(const TMap<UObject*, UObject*>& Map);
 	void HandleDeadObjectReferences(const TSet<UObject*>& ActiveInstances, const TSet<UObject*>& TemplateInstances);
 
@@ -434,50 +459,50 @@ public:
 	UE_API EOverriddenPropertyOperation GetOverriddenPropertyOperation(const FArchiveSerializedPropertyChain* CurrentPropertyChain, FProperty* Property) const;
 
 	/**
-	 * Setup the overridable operation of the current property from the serialized property chain and the specified property
+	 * Restores the overridable operation of the current property from the serialized property chain and the specified property
 	 * @param Operation to set for this property
 	 * @param CurrentPropertyChain leading to the property being serialized if any, null or empty it will return the node of the object itself
 	 * @param Property being serialized if any, if null it will fallback on the last property of the chain
 	 * @return the node containing the information of the overridden property */
-	UE_API FOverriddenPropertyNode* SetOverriddenPropertyOperation(EOverriddenPropertyOperation Operation, const FArchiveSerializedPropertyChain* CurrentPropertyChain, FProperty* Property);
+	UE_API FOverriddenPropertyNode* RestoreOverriddenPropertyOperation(EOverriddenPropertyOperation Operation, const FArchiveSerializedPropertyChain* CurrentPropertyChain, FProperty* Property);
 
 	/**
-	 * Restore the overridable operation of the current property from the serialized property chain and the specified property
+	 * Conditionally restores the overridable operation of the current property from the serialized property chain and the specified property
 	 * @note: This will not restore modified state, has restoring sub properties will do it anyway. 
 	 * @param Operation to set for this property
 	 * @param CurrentPropertyChain leading to the property being serialized if any, null or empty it will return the node of the object itself
 	 * @param Property being serialized if any, if null it will fallback on the last property of the chain
 	 * @return the node containing the information of the overridden property */
-	UE_API FOverriddenPropertyNode* RestoreOverriddenPropertyOperation(EOverriddenPropertyOperation Operation, const FArchiveSerializedPropertyChain* CurrentPropertyChain, FProperty* Property);
+	UE_API FOverriddenPropertyNode* ConditionallyRestoreOverriddenPropertyOperation(EOverriddenPropertyOperation Operation, const FArchiveSerializedPropertyChain* CurrentPropertyChain, FProperty* Property);
 
 	/**
 	 * Retrieve the overridden property node from the serialized property chain
 	 * @param CurrentPropertyChain leading to the property being serialized if any, null or empty it will return the node of the object itself
 	 * @return the node containing the information of the overridden property */
 	UE_API const FOverriddenPropertyNode* GetOverriddenPropertyNode(const FArchiveSerializedPropertyChain* CurrentPropertyChain) const;
-	UE_API const FOverriddenPropertyNode* GetOverriddenPropertyNode(FOverriddenPropertyNodeID NodeID) const;
+	UE_API const FOverriddenPropertyNode* GetOverriddenPropertyNode(const FOverriddenPropertyNodeID& NodeID) const;
 
 	/**
 	 * Retrieve the overridable operation given the property key
 	 * @param NodeID that uniquely identify the property within the object 
 	 * @return the current type of override operation on the property */
-	UE_API EOverriddenPropertyOperation GetSubPropertyOperation(FOverriddenPropertyNodeID NodeID) const;
+	UE_API EOverriddenPropertyOperation GetSubPropertyOperation(const FOverriddenPropertyNodeID& NodeID) const;
 
 	/**
-	 * Set the overridable operation of a sub property of the specified node.
+	 * Restores the overridable operation of a sub property of the specified node.
 	 * @param Operation to set for this property
 	 * @param Node from where the sub property is owned by
 	 * @param NodeID the ID of the sub property 
 	 * @return the node to the sub property */
-	UE_API FOverriddenPropertyNode* SetSubPropertyOperation(EOverriddenPropertyOperation Operation, FOverriddenPropertyNode& Node, FOverriddenPropertyNodeID NodeID);
+	UE_API FOverriddenPropertyNode* RestoreSubPropertyOperation(EOverriddenPropertyOperation Operation, FOverriddenPropertyNode& Node, const FOverriddenPropertyNodeID& NodeID);
 
 	/**
-	 * Set the overridable operation of a sub object of the specified node.
+	 * Restores the overridable operation of a sub object of the specified node.
 	 * @param Operation to set for this property
 	 * @param Node from where the sub property is owned by
 	 * @param SubObject of the operation
 	 * @return the node to the sub property */
-	UE_API FOverriddenPropertyNode* SetSubObjectOperation(EOverriddenPropertyOperation Operation, FOverriddenPropertyNode& Node, TNotNull<UObject*> SubObject);
+	UE_API FOverriddenPropertyNode* RestoreSubObjectOperation(EOverriddenPropertyOperation Operation, FOverriddenPropertyNode& Node, TNotNull<UObject*> SubObject);
 	/**
 	 * Check if this is an overridden property set of a CDO and that this property is owned by the class of this CDO
 	 * NOTE: this is used to know if a property should be serialized to keep its default CDO value.
@@ -488,7 +513,7 @@ public:
 
 	/**
 	 * Resets all overrides of the object */
-	UE_API void Reset();
+	UE_API void Reset(bool bShouldDirtyObject = true);
 	UE_API void HandleObjectsReInstantiated(const TMap<UObject*, UObject*>& Map);
 	UE_API void HandleDeadObjectReferences(const TSet<UObject*>& ActiveInstances, const TSet<UObject*>& TemplateInstances);
 
@@ -511,15 +536,13 @@ public:
 
 protected:
 
-	UE_API FOverriddenPropertyNode& FindOrAddNode(FOverriddenPropertyNode& ParentPropertyNode, FOverriddenPropertyNodeID NodeID);
-
 	UE_API EOverriddenPropertyOperation GetOverriddenPropertyOperation(const FOverriddenPropertyNode* ParentPropertyNode, FPropertyVisitorPath::Iterator PropertyIterator, bool* bOutInheritedOperation, const void* Data) const;
 	UE_API bool ClearOverriddenProperty(FOverriddenPropertyNode& ParentPropertyNode, FPropertyVisitorPath::Iterator PropertyIterator, const void* Data);
 	UE_API void NotifyPropertyChange(FOverriddenPropertyNode* ParentPropertyNode, const EPropertyNotificationType Notification, FPropertyVisitorPath::Iterator PropertyIterator, const EPropertyChangeType::Type ChangeType, const void* Data, bool& bNeedsCleanup);
 	UE_API void NotifyAllPropertyChange(FOverriddenPropertyNode* ParentPropertyNode, const EPropertyNotificationType Notification, const EPropertyChangeType::Type ChangeType, TNotNull<const UStruct*> Struct, const void* Data);
 
 	UE_API EOverriddenPropertyOperation GetOverriddenPropertyOperation(const FOverriddenPropertyNode* ParentPropertyNode, const FArchiveSerializedPropertyChain* CurrentPropertyChain, FProperty* Property) const;
-	UE_API FOverriddenPropertyNode* SetOverriddenPropertyOperation(EOverriddenPropertyOperation Operation, FOverriddenPropertyNode& ParentPropertyNode, const FArchiveSerializedPropertyChain* CurrentPropertyChain, FProperty* Property);
+	UE_API FOverriddenPropertyNode* RestoreOverriddenPropertyOperation(EOverriddenPropertyOperation Operation, FOverriddenPropertyNode& ParentPropertyNode, const FArchiveSerializedPropertyChain* CurrentPropertyChain, FProperty* Property);
 	UE_API const FOverriddenPropertyNode* GetOverriddenPropertyNode(const FOverriddenPropertyNode& ParentPropertyNode, const FArchiveSerializedPropertyChain* CurrentPropertyChain) const;
 
 	UE_API UObject* TryGetInstancedSubObjectValue(const FObjectPropertyBase* FromProperty, const void* ValuePtr) const;
diff --git a/Engine/Source/Runtime/Engine/Private/UnrealExporter.cpp b/Engine/Source/Runtime/Engine/Private/UnrealExporter.cpp
--- a/Engine/Source/Runtime/Engine/Private/UnrealExporter.cpp
+++ b/Engine/Source/Runtime/Engine/Private/UnrealExporter.cpp
@@ -864,13 +864,13 @@ void ExportProperties_Overrides
 						TArray<int32> RemovedIndices;
 						TArray<int32> AddedIndices;
 
-						for (const FOverriddenPropertyNode& SubNode: ArrayOverriddenPropertyNode->SubPropertyNodes)
+						for (const FOverriddenPropertyNode& SubNode: ArrayOverriddenPropertyNode->GetSubPropertyNodes())
 						{
-							switch (SubNode.Operation)
+							switch (SubNode.GetOperation())
 							{
 							case EOverriddenPropertyOperation::Remove:
 								{
-									const int32 DefaultIndex = FindObject(SubNode.NodeID, DefaultsArrayHelper);
+									const int32 DefaultIndex = FindObject(SubNode.GetNodeID(), DefaultsArrayHelper);
 									if (DefaultIndex != INDEX_NONE)
 									{
 										RemovedIndices.Add(DefaultIndex);
@@ -879,7 +879,7 @@ void ExportProperties_Overrides
 								}
 							case EOverriddenPropertyOperation::Add:
 								{
-									const int32 Index = FindObject(SubNode.NodeID, ArrayHelper);
+									const int32 Index = FindObject(SubNode.GetNodeID(), ArrayHelper);
 									if (Index != INDEX_NONE)
 									{
 										AddedIndices.Add(Index);
@@ -953,13 +953,13 @@ void ExportProperties_Overrides
 
 
 						// Figure out the modifications of the map
-						for (const FOverriddenPropertyNode& SubNode : MapOverriddenPropertyNode->SubPropertyNodes)
+						for (const FOverriddenPropertyNode& SubNode : MapOverriddenPropertyNode->GetSubPropertyNodes())
 						{
-							switch (SubNode.Operation)
+							switch (SubNode.GetOperation())
 							{
 							case EOverriddenPropertyOperation::Remove:
 								{
-									const int32 InternalIndex = SubNode.NodeID.ToMapInternalIndex(DiffMapHelper);
+									const int32 InternalIndex = SubNode.GetNodeID().ToMapInternalIndex(DiffMapHelper);
 									if (InternalIndex != INDEX_NONE)
 									{
 										RemovedIndices.Add(InternalIndex);
@@ -968,7 +968,7 @@ void ExportProperties_Overrides
 								}
 							case EOverriddenPropertyOperation::Add:
 								{
-									const int32 InternalIndex = SubNode.NodeID.ToMapInternalIndex(MapHelper);
+									const int32 InternalIndex = SubNode.GetNodeID().ToMapInternalIndex(MapHelper);
 									if (InternalIndex != INDEX_NONE)
 									{
 										AddedIndices.Add(InternalIndex);
@@ -977,7 +977,7 @@ void ExportProperties_Overrides
 								}
 							case EOverriddenPropertyOperation::Modified:
 								{
-									const int32 InternalIndex = SubNode.NodeID.ToMapInternalIndex(MapHelper);
+									const int32 InternalIndex = SubNode.GetNodeID().ToMapInternalIndex(MapHelper);
 									if (InternalIndex != INDEX_NONE)
 									{
 										ModifiedIndices.Add(InternalIndex);
