Constinit GNatives and GCasts

diff --git a/Engine/Source/Runtime/CoreUObject/Private/UObject/ScriptCore.cpp b/Engine/Source/Runtime/CoreUObject/Private/UObject/ScriptCore.cpp
--- a/Engine/Source/Runtime/CoreUObject/Private/UObject/ScriptCore.cpp
+++ b/Engine/Source/Runtime/CoreUObject/Private/UObject/ScriptCore.cpp
@@ -94,7 +94,7 @@ static FAutoConsoleVariableRef CVarMaxFunctionStatDepth(
 //
 // Native function table.
 //
-COREUOBJECT_API FNativeFuncPtr GNatives[EX_Max];
+COREUOBJECT_API constinit TStaticArray<FNativeFuncPtr, EX_Max> GNatives(InPlace, &UObject::execUndefined);
 COREUOBJECT_API int32 GNativeDuplicate=0;
 
 COREUOBJECT_API int32 GMaximumScriptLoopIterations = 1000000;
@@ -854,19 +854,9 @@ int32 FScriptInstrumentationSignal::GetScriptCodeOffset() const
 //
 COREUOBJECT_API uint8 GRegisterNative(int32 NativeBytecodeIndex, const FNativeFuncPtr& Func)
 {
-	static bool bInitialized = false;
-	if (!bInitialized)
-	{
-		bInitialized = true;
-		for (uint32 i = 0; i < UE_ARRAY_COUNT(GNatives); i++)
-		{
-			GNatives[i] = &UObject::execUndefined;
-		}
-	}
-
 	if( NativeBytecodeIndex != INDEX_NONE )
 	{
-		if( NativeBytecodeIndex<0 || (uint32)NativeBytecodeIndex > UE_ARRAY_COUNT(GNatives) || GNatives[NativeBytecodeIndex] != &UObject::execUndefined) 
+		if( NativeBytecodeIndex<0 || NativeBytecodeIndex > GNatives.Num() || GNatives[NativeBytecodeIndex] != &UObject::execUndefined) 
 		{
 			CA_SUPPRESS(6385)
 			if (!ReloadNotifyFunctionRemap(Func, GNatives[NativeBytecodeIndex]))
@@ -881,20 +871,10 @@ COREUOBJECT_API uint8 GRegisterNative(int32 NativeBytecodeIndex, const FNativeFu
 	return 0;
 }
 
-static FNativeFuncPtr GCasts[CST_Max];
+static constinit TStaticArray<FNativeFuncPtr, CST_Max> GCasts(InPlace, &UObject::execUndefined);
 
 static uint8 GRegisterCast(ECastToken CastCode, const FNativeFuncPtr& Func)
 {
-	static int32 bInitialized = false;
-	if (!bInitialized)
-	{
-		bInitialized = true;
-		for (uint32 i = 0; i < UE_ARRAY_COUNT(GCasts); i++)
-		{
-			GCasts[i] = &UObject::execUndefined;
-		}
-	}
-
 	if (CastCode != CST_Max)
 	{
 		GCasts[CastCode] = Func;
