Fix big regression introduced

diff --git a/Engine/Source/Runtime/CoreUObject/Private/UObject/OverriddenPropertySet.cpp b/Engine/Source/Runtime/CoreUObject/Private/UObject/OverriddenPropertySet.cpp
--- a/Engine/Source/Runtime/CoreUObject/Private/UObject/OverriddenPropertySet.cpp
+++ b/Engine/Source/Runtime/CoreUObject/Private/UObject/OverriddenPropertySet.cpp
@@ -437,7 +437,7 @@ bool FOverriddenPropertyNodeID::operator==(const FOverriddenPropertyNodeID& Othe
 	// so are only way if to compare ptr.
 	if (Object.IsValid(/*bEvenIfGarbage*/true) && 
 		Other.Object.IsValid(/*bEvenIfGarbage*/true) && 
-		Object == Other.Object)
+		Object.Get(/*bEvenIfPendingKill*/true) == Other.Object.Get(/*bEvenIfPendingKill*/true))
 	{
 		return true;
 	}
