UObject ConstructForTests is only ever called from FScriptStructTestWrapper which is not in test/shipping builds.

Don't include it in the FakeVTable in builds where we're never going to call it.

diff --git a/Engine/Source/Runtime/CoreUObject/Public/UObject/Class.h b/Engine/Source/Runtime/CoreUObject/Public/UObject/Class.h
--- a/Engine/Source/Runtime/CoreUObject/Public/UObject/Class.h
+++ b/Engine/Source/Runtime/CoreUObject/Public/UObject/Class.h
@@ -1126,6 +1126,7 @@ namespace UE::CoreUObject::Private
 		}
 	}
 
+#if !(UE_BUILD_TEST || UE_BUILD_SHIPPING)
 	template <typename CppStruct>
 		requires (!TStructOpsTypeTraits<CppStruct>::WithZeroConstructor)
 	void ConstructForTests(void* Dest)
@@ -1142,10 +1143,12 @@ namespace UE::CoreUObject::Private
 			}
 			else
 			{
+				// NOTE: unlike regular "Construct", no value init here! (missing "()")
 				::new (Dest) CppStruct;
 			}
 		}
 	}
+#endif
 
 	template <typename CppStruct>
 		requires (!TStructOpsTypeTraits<CppStruct>::WithNoDestructor && !TIsPODType<CppStruct>::Value)
@@ -1379,6 +1382,12 @@ namespace UE::CoreUObject::Private
 	#define UE_FOREACH_FAKE_VTABLE_FUNC_EDITOR_ONLY(...)
 #endif // WITH_EDITOR
 
+#if !(UE_BUILD_TEST || UE_BUILD_SHIPPING)
+	#define UE_FOREACH_FAKE_VTABLE_FUNC_WITH_TESTS_ONLY(...) __VA_ARGS__
+#else
+	#define UE_FOREACH_FAKE_VTABLE_FUNC_WITH_TESTS_ONLY(...)
+#endif
+
 	template <typename CppStruct>
 		requires (!!TStructOpsTypeTraits<CppStruct>::WithVisitor)
 	EPropertyVisitorControlFlow Visit(FPropertyVisitorContext& Context, const TFunctionRef<EPropertyVisitorControlFlow(const FPropertyVisitorContext& Context)> InFunc)
@@ -1396,7 +1405,9 @@ namespace UE::CoreUObject::Private
 
 	#define UE_FOREACH_FAKE_VTABLE_FUNC(IndexNameFunctype) \
 		IndexNameFunctype(0,	Construct,								void(void* Address)) \
+		UE_FOREACH_FAKE_VTABLE_FUNC_WITH_TESTS_ONLY( \
 		IndexNameFunctype(1,	ConstructForTests,						void(void* Address)) \
+		) \
 		IndexNameFunctype(2,	Destruct,								void(void* Address)) \
 		IndexNameFunctype(3,	SerializeArchive,						bool(FArchive& Ar, void* Data, UStruct* DefaultsStruct, const void* Defaults)) \
 		IndexNameFunctype(4,	SerializeSlot,							bool(FStructuredArchive::FSlot Slot, void* Data, UStruct* DefaultsStruct, const void* Defaults)) \
@@ -1652,11 +1663,13 @@ public:
 			FakeVPtr->GetHandlingFunc<UE::CoreUObject::Private::EStructOpsFakeVTableFlags::Construct>()(Dest);
 		}
 
+#if !(UE_BUILD_TEST || UE_BUILD_SHIPPING)
 		/** Call the C++ constructor without value-init (new T instead of new T()) **/
 		void ConstructForTests(void* Dest)
 		{
 			FakeVPtr->GetHandlingFunc<UE::CoreUObject::Private::EStructOpsFakeVTableFlags::ConstructForTests>()(Dest);
 		}
+#endif
 
 		/** return false if this destructor can be skipped **/
 		bool HasDestructor() const
