Be more consistent with the naming of parameter passed to WeakObjPtr

#rnx

diff --git a/Engine/Source/Runtime/CoreUObject/Private/UObject/OverriddenPropertySet.cpp b/Engine/Source/Runtime/CoreUObject/Private/UObject/OverriddenPropertySet.cpp
--- a/Engine/Source/Runtime/CoreUObject/Private/UObject/OverriddenPropertySet.cpp
+++ b/Engine/Source/Runtime/CoreUObject/Private/UObject/OverriddenPropertySet.cpp
@@ -434,10 +434,11 @@ bool FOverriddenPropertyNodeID::operator==(const FOverriddenPropertyNodeID& Othe
 
 	// After reinstantiation, we do not change the path we only patch the ptr.
 	// Unfortunately we do not have a stable id that is stable through reinstantiation,
-	// so are only way if to compare ptr.
-	if (Object.IsValid(/*bEvenIfGarbage*/true) && 
-		Other.Object.IsValid(/*bEvenIfGarbage*/true) && 
-		Object.Get(/*bEvenIfPendingKill*/true) == Other.Object.Get(/*bEvenIfPendingKill*/true))
+	// so the only way is to compare ptr.
+	constexpr bool bEvenIfPendingKill = true;
+	if (Object.IsValid(bEvenIfPendingKill) && 
+		Other.Object.IsValid(bEvenIfPendingKill) && 
+		Object.Get(bEvenIfPendingKill) == Other.Object.Get(bEvenIfPendingKill))
 	{
 		return true;
 	}
