Optional Packages: Remove the restriction that Optional arguments cannot be the only referencer of non-optional objects in the package. Allow those non-optional objects to be pulled into the optional package. The restriction was originally put in place to provide a guarantee that the harvesting of objects into the optionalrealm of a package would not cause more objects to be added to the non-optional realm version of the package; add some assertions to the PackageHarvester to guarantee that narrower requirement remains met.
#rnx

diff --git a/Engine/Source/Runtime/CoreUObject/Private/UObject/SavePackage/PackageHarvester.cpp b/Engine/Source/Runtime/CoreUObject/Private/UObject/SavePackage/PackageHarvester.cpp
--- a/Engine/Source/Runtime/CoreUObject/Private/UObject/SavePackage/PackageHarvester.cpp
+++ b/Engine/Source/Runtime/CoreUObject/Private/UObject/SavePackage/PackageHarvester.cpp
@@ -239,6 +239,7 @@ FPackageHarvester::FExportScope::FExportScope(FPackageHarvester& InHarvester, co
 	check(!Harvester.HasAnyExportHarvestingRealms());
 	Harvester.CurrentExportDependencies = { InToProcess.Export, InToProcess.CellExport };
 	Harvester.CurrentExportHarvestingRealms = InToProcess.HarvestedFromRealms;
+	Harvester.RootExportHarvestingRealms = Harvester.CurrentExportHarvestingRealms;
 
 	if (Harvester.CurrentExportHarvestingRealms.Contains(ESaveRealm::Optional))
 	{
@@ -259,7 +260,8 @@ FPackageHarvester::FExportScope::FExportScope(FPackageHarvester& InHarvester, co
 		// serialize functions can change their behavior based on Ar.IsSavingOptionalObject, and the non-optional package 
 		// could pull in extra imports or FNames that client builds don't want. Currently we guarantee that the
 		// optional realm is serialized alone in FPackageHarvester::TryHarvestExport.
-		checkf(Harvester.CurrentExportHarvestingRealms.Num() == 1, TEXT("Expecting only 1 realm when it's optional - if this trips, we will have to manage marking 'AllowedInOptional' object/fname properties to not bloat things"));
+		checkf(Harvester.CurrentExportHarvestingRealms.Num() == 1,
+			TEXT("Expecting only 1 realm when it's optional - if this trips, we will have to manage marking 'AllowedInOptional' object/fname properties to not bloat things"));
 	}
 }
 
@@ -267,6 +269,7 @@ FPackageHarvester::FExportScope::~FExportScope()
 {
 	Harvester.AppendCurrentExportDependencies();
 	Harvester.CurrentExportHarvestingRealms = PreviousExportHarvestingRealms;
+	Harvester.RootExportHarvestingRealms.Empty();
 	Harvester.SetFilterEditorOnly(bPreviousFilterEditorOnly);
 	Harvester.SetSavingOptionalObject(bPreviousSavingOptionalObject);
 }
@@ -311,17 +314,22 @@ bool FPackageHarvester::FHarvestScope::IsEmpty() const
 	return Harvester.CurrentExportHarvestingRealms.IsEmpty();
 }
 
-FPackageHarvester::FHarvestScope FPackageHarvester::EnterRootReferencesScope()
+FPackageHarvester::FRootReferencesHarvestScope::FRootReferencesHarvestScope(FPackageHarvester& InHarvester)
+	: FHarvestScope(InHarvester)
 {
-	check(!HasAnyExportHarvestingRealms());
-	FHarvestScope Scope(*this);
+	check(!Harvester.HasAnyExportHarvestingRealms());
 
-	CurrentExportHarvestingRealms.Add(ESaveRealm::Game);
-	if (!SaveContext.IsCooking())
+	Harvester.CurrentExportHarvestingRealms.Add(ESaveRealm::Game);
+	if (!Harvester.SaveContext.IsCooking())
 	{
-		CurrentExportHarvestingRealms.Add(ESaveRealm::Editor);
+		Harvester.CurrentExportHarvestingRealms.Add(ESaveRealm::Editor);
 	}
-	return Scope;
+	Harvester.RootExportHarvestingRealms = Harvester.CurrentExportHarvestingRealms;
+}
+
+FPackageHarvester::FRootReferencesHarvestScope::~FRootReferencesHarvestScope()
+{
+	Harvester.RootExportHarvestingRealms.Empty();
 }
 
 FPackageHarvester::FHarvestScope FPackageHarvester::EnterRealmsArrayScope(FExportingRealmsArray& Array)
@@ -713,8 +721,11 @@ void FPackageHarvester::TryHarvestExport(UObject* InObject, UE::SavePackageUtili
 		}
 	}
 	CurrentExportHarvestingRealms = HarvestingRealms;
+	RootExportHarvestingRealms = CurrentExportHarvestingRealms;
 
 	TryHarvestExportInternal(InObject, ObjectStatus);
+
+	RootExportHarvestingRealms.Empty();
 }
 
 void FPackageHarvester::TryHarvestExportInternal(UObject* InObject, UE::SavePackageUtilities::FObjectStatus& ObjectStatus)
@@ -815,8 +826,11 @@ void FPackageHarvester::TryHarvestCellExport(Verse::VCell* InCell)
 	FExportingRealmsArray HarvestingRealms;
 	HarvestingRealms.Add(ESaveRealm::Game);
 	CurrentExportHarvestingRealms = HarvestingRealms;
+	RootExportHarvestingRealms = CurrentExportHarvestingRealms;
 
 	HarvestCellExport(InCell);
+
+	RootExportHarvestingRealms.Empty();
 }
 #endif
 
@@ -1310,9 +1324,8 @@ void FPackageHarvester::HarvestSearchableName(TObjectPtr<UObject> TypeObject, FN
 
 void FPackageHarvester::HarvestExport(UObject* InObject, UE::SavePackageUtilities::FObjectStatus& ObjectStatus)
 {
-	bool bFromOptionalRef = CurrentExportDependencies.CurrentExport &&
-		CurrentExportDependencies.CurrentExport->GetClass()->HasAnyClassFlags(CLASS_Optional);
-	FTaggedExport TaggedExport(InObject, !DoesObjectNeedLoadForEditorGame(InObject), bFromOptionalRef);
+	ValidateNoOptionalRealmContamination(InObject);
+	FTaggedExport TaggedExport(InObject, !DoesObjectNeedLoadForEditorGame(InObject));
 
 	ForEachExportHarvestingRealm([this, &TaggedExport, InObject, &ObjectStatus](ESaveRealm HarvestingRealm)
 		{
@@ -1324,6 +1337,51 @@ void FPackageHarvester::HarvestExport(UObject* InObject, UE::SavePackageUtilitie
 	ExportsToProcess.Enqueue({ InObject, nullptr, CurrentExportHarvestingRealms });
 }
 
+
+void FPackageHarvester::ValidateNoOptionalRealmContamination(UObject* HarvestedObject
+#if WITH_VERSE_VM || defined(__INTELLISENSE__)
+, Verse::VCell* HarvestedCell
+#endif
+)
+{
+	// This function is only called from Harvest functions, and callers of Harvest functions are required to
+	// not to call the Harvest function if there are not any root HarvestingRealms or current HarvestingRealms.
+	check(!RootExportHarvestingRealms.IsEmpty());
+	check(!CurrentExportHarvestingRealms.IsEmpty());
+
+	bool bRootHarvestInOptional = RootExportHarvestingRealms[0] == ESaveRealm::Optional;
+	bool bHarvestInOptional = CurrentExportHarvestingRealms[0] == ESaveRealm::Optional;
+	// The optional realm has to be harvested individually because we set flags separately for it
+	check(!bRootHarvestInOptional || RootExportHarvestingRealms.Num() == 1);
+	check(!bHarvestInOptional || CurrentExportHarvestingRealms.Num() == 1);
+
+	if (bRootHarvestInOptional && !bHarvestInOptional)
+	{
+		UObject* Export = CurrentExportDependencies.CurrentExport;
+		FString HarvestedName(TEXT("<UnknownObject>"));
+		if (HarvestedObject)
+		{
+			HarvestedName = HarvestedObject->GetFullName();
+		}
+#if WITH_VERSE_VM || defined(__INTELLISENSE__)
+		else if (HarvestedCell)
+		{
+			HarvestedName = TEXT("a VCell");
+		}
+#endif
+
+		UE_CALL_ONCE([&]
+			{
+				UE_LOG(LogSavePackage, Error,
+					TEXT("SavePackage logic error: During save of package %s, while processing export %s in the optional realm, %s was harvested in the game realm.")
+					TEXT(" This is invalid; references harvested from exports processed in the optional realm must not affect the non-optional package."),
+					*SaveContext.Package->GetPackage()->GetName(), Export ? *Export->GetPathName() : TEXT("<unknown>"),
+					*HarvestedName);
+				FDebug::DumpStackTraceToLog(ELogVerbosity::Display);
+			});
+	}
+}
+
 void FPackageHarvester::HarvestImport(TObjectPtr<UObject> InObject, UE::SavePackageUtilities::FObjectStatus& ObjectStatus)
 {
 	ForEachExportHarvestingRealm([this, InObject, &ObjectStatus](ESaveRealm HarvestingRealm)
@@ -1367,6 +1425,8 @@ void FPackageHarvester::SetInstigatorFromProcessingStack(UE::SavePackageUtilitie
 #if WITH_VERSE_VM || defined(__INTELLISENSE__)
 void FPackageHarvester::HarvestCellExport(Verse::VCell* InCell)
 {
+	ValidateNoOptionalRealmContamination(nullptr, InCell);
+
 	ForEachExportHarvestingRealm([this, InCell](ESaveRealm HarvestingRealm)
 		{
 			FHarvestedRealm& RealmData = SaveContext.GetHarvestedRealm(HarvestingRealm);
diff --git a/Engine/Source/Runtime/CoreUObject/Private/UObject/SavePackage/PackageHarvester.h b/Engine/Source/Runtime/CoreUObject/Private/UObject/SavePackage/PackageHarvester.h
--- a/Engine/Source/Runtime/CoreUObject/Private/UObject/SavePackage/PackageHarvester.h
+++ b/Engine/Source/Runtime/CoreUObject/Private/UObject/SavePackage/PackageHarvester.h
@@ -53,12 +53,20 @@ public:
 		FHarvestScope& operator=(const FHarvestScope& Other) = delete;
 		~FHarvestScope();
 		bool IsEmpty() const;
-	private:
+	protected:
 		FPackageHarvester& Harvester;
 		FExportingRealmsArray PreviousExportHarvestingRealms;
 		bool bActive;
 	};
 
+	/** Scope object constructed when harvesting references in the package's root context, outside of an Export. */
+	class FRootReferencesHarvestScope : public FHarvestScope
+	{
+	public:
+		FRootReferencesHarvestScope(FPackageHarvester& InHarvester);
+		~FRootReferencesHarvestScope();
+	};
+
 public:
 	FPackageHarvester(FSaveContext& InContext);
 
@@ -71,7 +79,6 @@ public:
 	void TryHarvestCellExport(Verse::VCell* InCell);
 	void ProcessCellExport(const FExportWithContext& InProcessContext);
 #endif
-	FHarvestScope EnterRootReferencesScope();
 
 	void HarvestExportDataName(FName Name);
 	void HarvestPackageHeaderName(FName Name);
@@ -158,6 +165,16 @@ private:
 	void ProcessImport(TObjectPtr<UObject> InObject, UE::SavePackageUtilities::FObjectStatus& ObjectStatus);
 	void SetInstigatorFromProcessingStack(UE::SavePackageUtilities::FObjectStatus& ObjectStatus,
 		ESaveRealm HarvestingRealm);
+	/**
+	 * Enforce the contract that Optional objects can not cause the harvesting of any new objects into the
+	 * non-optional realms of the package. When we find references while harvesting in the optional realm,
+	 * the references must remain in the optional realm.
+	 */
+	void ValidateNoOptionalRealmContamination(UObject* HarvestedObject
+#if WITH_VERSE_VM || defined(__INTELLISENSE__)
+		, Verse::VCell* HarvestedCell = nullptr
+#endif
+	);
 
 #if WITH_VERSE_VM || defined(__INTELLISENSE__)
 	void HarvestCellExport(Verse::VCell* InCell);
@@ -196,5 +213,17 @@ private:
 	TQueue<FExportWithContext> ExportsToProcess;
 	FExportDependencies CurrentExportDependencies;
 	TMap<UObject*, TSet<FProperty*>> TransientPropertyOverrides;
+	/**
+	 * Accumulated list of which HarvestingRealms we are collecting references into within the current scope.
+	 * This list starts at the default value for a root object or an export (e.g. Game+Editor when saving in editor)
+	 * but can be modified to drop some of those realms when following dependency edges that only apply to a single
+	 * realm, or can be modified to switch to a different set of realms (e.g. Optional) if specified by the
+	 * dependency edge or target object.
+	 */
 	FExportingRealmsArray CurrentExportHarvestingRealms;
+	/**
+	 * Snapshot of the CurrentExportHarvestingRealms that were present when we first started processing an export.
+         * These remain unchanged until we stop processing the export.
+	 */
+	FExportingRealmsArray RootExportHarvestingRealms;
 };
\ No newline at end of file
diff --git a/Engine/Source/Runtime/CoreUObject/Private/UObject/SavePackage/SaveContext.h b/Engine/Source/Runtime/CoreUObject/Private/UObject/SavePackage/SaveContext.h
--- a/Engine/Source/Runtime/CoreUObject/Private/UObject/SavePackage/SaveContext.h
+++ b/Engine/Source/Runtime/CoreUObject/Private/UObject/SavePackage/SaveContext.h
@@ -43,24 +43,17 @@ struct FTaggedExport
 	 * This will artificially mark the object RF_Public in the linker tables so the iostore generates the public hash
 	 */
 	uint32 bGeneratePublicHash : 1;
-	/**
-	 * Indicate if the object that directly referenced this export was optional
-	 * Used to determine mandatory objects in the game save realm
-	 */
-	uint32 bFromOptionalReference : 1;
 
 	FTaggedExport()
 		: Obj(nullptr)
 		, bNotAlwaysLoadedForEditorGame(false)
 		, bGeneratePublicHash(false)
-		, bFromOptionalReference(false)
 	{}
 
-	FTaggedExport(TObjectPtr<UObject> InObj, bool bInNotAlwaysLoadedForEditorGame = true, bool bInFromOptionalReference = false)
+	FTaggedExport(TObjectPtr<UObject> InObj, bool bInNotAlwaysLoadedForEditorGame = true)
 		: Obj(InObj)
 		, bNotAlwaysLoadedForEditorGame(bInNotAlwaysLoadedForEditorGame)
 		, bGeneratePublicHash(false)
-		, bFromOptionalReference(bInFromOptionalReference)
 	{}
 
 	inline bool operator == (const FTaggedExport& Other) const
@@ -92,7 +85,6 @@ enum class EIllegalRefReason : uint8
 {
 	None = 0,
 	ReferenceToOptional,
-	ReferenceFromOptionalToMissingGameExport,
 	UnsaveableClass,
 	UnsaveableOuter,
 	ExternalPackage,
diff --git a/Engine/Source/Runtime/CoreUObject/Private/UObject/SavePackage2.cpp b/Engine/Source/Runtime/CoreUObject/Private/UObject/SavePackage2.cpp
--- a/Engine/Source/Runtime/CoreUObject/Private/UObject/SavePackage2.cpp
+++ b/Engine/Source/Runtime/CoreUObject/Private/UObject/SavePackage2.cpp
@@ -369,7 +369,6 @@ ESavePackageResult HarvestPackage(FSaveContext& SaveContext)
 		// if we are automatically generating an optional package output, re-harvest objects with that realm as the default
 		if (SaveContext.IsSaveAutoOptional())
 		{
-			//@todo: FH: the top level might have to be altered to use the editor one when auto harvesting optional data
 			FSaveContext::FSetSaveRealmToSaveScope RealmScope(SaveContext, ESaveRealm::Optional);
 			Harvester.TryHarvestExport(InRoot, ObjectStatus);
 		}
@@ -429,41 +428,18 @@ ESavePackageResult HarvestPackage(FSaveContext& SaveContext)
 		FHarvestedRealm& OptionalContext = SaveContext.GetHarvestedRealm(ESaveRealm::Optional);
 		for (auto It = OptionalContext.GetExports().CreateIterator(); It; ++It)
 		{
-			if (!FPackageHarvester::ShouldObjectBeHarvestedInOptionalRealm(It->Obj, SaveContext))
+			// If the export is found in the game context as well then remove the export from the optional realm
+			// and import it from the game realm.
+			if (FTaggedExport* GameExport = SaveContext.GetHarvestedRealm(ESaveRealm::Game).GetExports().Find(It->Obj))
 			{
-				// Make sure the export is found in the game context as well
-				if (FTaggedExport* GameExport = SaveContext.GetHarvestedRealm(ESaveRealm::Game).GetExports().Find(It->Obj))
-				{
-					// Flagging the export in the game context to generate it's public hash isn't necessary anymore
-					//GameExport->bGeneratePublicHash = true; 
+				// Flagging the export in the game context to generate it's public hash isn't necessary anymore
+				//GameExport->bGeneratePublicHash = true; 
 
-					// Transform the export as an import
-					OptionalContext.AddImport(It->Obj);
-					// Flag the package itself to be an import
-					bHasNonOptionalSelfReference = true;
-					It.RemoveCurrent();
-				}
-				// if not found in the game context and the reference directly came from an optional object, record an illegal reference
-				else if (It->bFromOptionalReference)
-				{
-					// TODO: Suppress this error and leave the object in the optional realm on Server builds, because it does not work in
-					// some cases on multiplatform cooks. EditorOptional is a global commandline setting that applies to all platforms in
-					// a multiplatform cook, but on some projects the project cooks Client and Server together and wants editoroptional
-					// only on client, so it cooks both of them as editoroptional. But then server doesn't support editoroptional for
-					// objects that are in ClassesExcludedOnDedicatedServer but that are referenced both from optional and from gamerealm
-					// objects; we get this error because the object is stripped from the game realm but remains in the optional realm.
-					// 
-					// Change EditorOptional to be specifiable per platform, and remove this suppression.
-					if (!SaveContext.GetTargetPlatform()->IsServerOnly())
-					{
-						SaveContext.RecordIllegalReference(nullptr, It->Obj, EIllegalRefReason::ReferenceFromOptionalToMissingGameExport);
-						It.RemoveCurrent();
-					}
-				}
-				else
-				{
-					It.RemoveCurrent();
-				}
+				// Transform the export as an import
+				OptionalContext.AddImport(It->Obj);
+				// Flag the package itself to be an import
+				bHasNonOptionalSelfReference = true;
+				It.RemoveCurrent();
 			}
 		}
 		// Also add the current package itself as an import if we are referencing any non optional export
@@ -474,7 +450,7 @@ ESavePackageResult HarvestPackage(FSaveContext& SaveContext)
 	}
 
 	{
-		FPackageHarvester::FHarvestScope RootReferenceScope = Harvester.EnterRootReferencesScope();
+		FPackageHarvester::FRootReferencesHarvestScope RootReferenceScope(Harvester);
 		// Trim PrestreamPackage list 
 		TSet<TObjectPtr<UPackage>>& PrestreamPackages = SaveContext.GetPrestreamPackages();
 		TSet<TObjectPtr<UPackage>> KeptPrestreamPackages;
@@ -550,12 +526,6 @@ ESavePackageResult ValidateRealms(FSaveContext& SaveContext)
 					Reference.From ? *Reference.From->GetPathName() : TEXT("Unknown"),
 					Reference.To ? *Reference.To->GetPathName() : TEXT("Unknown"));
 				break;
-			case EIllegalRefReason::ReferenceFromOptionalToMissingGameExport:
-				ErrorMessage = FString::Printf(TEXT("Can't save %s: Optional object (%s) has a reference to cooked object (%s) which is missing. Non optional objects referenced by optional objects needs to be present in cooked data."),
-					SaveContext.GetFilename(),
-					Reference.From ? *Reference.From->GetPathName() : TEXT("Unknown"),
-					Reference.To ? *Reference.To->GetPathName() : TEXT("Unknown"));
-				break;
 			case EIllegalRefReason::UnsaveableClass:
 				ErrorMessage = FString::Printf(TEXT("Can't save %s: Object (%s) is an export but is an instance of class (%s) which is unsaveable: %s"),
 					SaveContext.GetFilename(),
