Fix for the infinite loop caused when trashing a package then synchronously loading its replacement in a transaction.

The TL;DR is that the problem we have is we've got a package `A` that we are trashing in a transaction, then immediately doing a sync-load. So what happens is:

* we find `A`
* trash it
* issue a new package load call
* it isn't in any cache so we load
* we abort the transaction (undoing the trash of the package)
* then issue the load, but zen loader will find the existing package because it is now *not* trashed
* we'll start again and do the same thing for infinity

Another issue I found was that if you had a package `B` that wasn't yet loaded:

* you'd find `B`, returning nullptr
* sync load `B`
* we abort, do the load for `B`, and retry
* the next iteration of the loop you'd find `B` returning the actual package even though logically the package hadn't been created yet

This meant two fixes were needed:

* when we load a package during retry, we rename it to the `AutoRTFMUnusedPackage` name, so that it won't appear when you do `FindObject`
* if we are retrying and we have issued a load for a package that is actually already cached, that can only happen if we've trashed the package and then issued a sync load for it, so we need to rename the original object, do the sync load, then rename it back

diff --git a/Engine/Source/Programs/AutoRTFMTests/Private/UPackageTests.cpp b/Engine/Source/Programs/AutoRTFMTests/Private/UPackageTests.cpp
--- a/Engine/Source/Programs/AutoRTFMTests/Private/UPackageTests.cpp
+++ b/Engine/Source/Programs/AutoRTFMTests/Private/UPackageTests.cpp
@@ -419,23 +419,126 @@ TEST_CASE("UPackage.AsyncLoading")
 			REQUIRE(ReloadedPackage->GetFName() == Name);
 
 			REQUIRE(Package != ReloadedPackage);
-			REQUIRE(!ReloadedPackage->GetFName().ToString().StartsWith(TEXT("TrashedPackage")));
+		});
+	}
+
+	SECTION("Find Package Loaded In Transaction")
+	{
+		FString Name(FString::Printf(TEXT("/Game/%dAutoRTFMTestPackage%d"), FPlatformProcess::GetCurrentProcessId(), __LINE__));
+		FPackageScopedMaker _(Name);
+
+		AutoRTFM::Testing::Commit([&]
+		{
+			// Make sure the package doesn't exist.
+			REQUIRE(FindObject<UObject>(nullptr, *Name) == nullptr);
+
+			UPackage* const Package = LoadPackage(nullptr, *Name, LOAD_None);
+			REQUIRE(Package);
+
+			TArray<UObject*> Objects;
+
+			REQUIRE(StaticFindAllObjects(Objects, UObject::StaticClass(), *Name));
+			REQUIRE(Objects.Num() == 1);
+			REQUIRE(Objects[0] == Package);
 		});
 	}
 
 	SECTION("Trash Package that was created in same transaction as reloaded")
 	{
-		// #jira SOL-8376: Loading a package, trashing it, then reloading, causes infinite loops because
-		// the transactionally safe async loader will abort and retry for each package load that actually
-		// needs to hit the async loader (because it will fight us for locks).
-		static bool bSOL_8376 = false;
+		FString Name(FString::Printf(TEXT("/Game/%dAutoRTFMTestPackage%d"), FPlatformProcess::GetCurrentProcessId(), __LINE__));
+		FPackageScopedMaker _(Name);
 
-		if (bSOL_8376)
+		AutoRTFM::Testing::Commit([&]
 		{
-			FString Name(FString::Printf(TEXT("/Game/%dAutoRTFMTestPackage%d"), FPlatformProcess::GetCurrentProcessId(), __LINE__));
-			FPackageScopedMaker _(Name);
+			UPackage* const Package = LoadPackage(nullptr, *Name, LOAD_None);
+			REQUIRE(Package);
+			
+			// TODO: If we could move the trashing logic into `UPackage` we could just call that here?
+			const FName NewName = MakeUniqueObjectName(nullptr, UPackage::StaticClass(), NAME_TrashedPackage);
+			Package->Rename(*NewName.ToString(), nullptr, REN_DontCreateRedirectors | REN_NonTransactional | REN_DoNotDirty);
+			Package->SetFlags(RF_Transient);
+
+			UPackage* const ReloadedPackage = LoadPackage(nullptr, *Name, LOAD_None);
+			REQUIRE(ReloadedPackage);
+			REQUIRE(ReloadedPackage->GetFName() == Name);
+
+			REQUIRE(Package != ReloadedPackage);
+		});
+	}
+
+	SECTION("Trash Package that was found in the same transaction")
+	{
+		FString Name(FString::Printf(TEXT("/Game/%dAutoRTFMTestPackage%d"), FPlatformProcess::GetCurrentProcessId(), __LINE__));
+		FPackageScopedMaker _(Name);
+		UPackage* const Package = LoadPackage(nullptr, *Name, LOAD_None);
+		REQUIRE(Package);
+
+		AutoRTFM::Testing::Commit([&]
+		{
+			UPackage* const FoundPackage = FindObject<UPackage>(nullptr, *Name);
+			REQUIRE(Package == FoundPackage);
+			
+			// TODO: If we could move the trashing logic into `UPackage` we could just call that here?
+			const FName NewName = MakeUniqueObjectName(nullptr, UPackage::StaticClass(), NAME_TrashedPackage);
+			Package->Rename(*NewName.ToString(), nullptr, REN_DontCreateRedirectors | REN_NonTransactional | REN_DoNotDirty);
+			Package->SetFlags(RF_Transient);
 
-			AutoRTFM::Testing::Commit([&]
+			UPackage* const ReloadedPackage = LoadPackage(nullptr, *Name, LOAD_None);
+			REQUIRE(ReloadedPackage);
+			REQUIRE(ReloadedPackage->GetFName() == Name);
+
+			REQUIRE(Package != ReloadedPackage);
+		});
+
+		UPackage* const FoundPackage = FindObject<UPackage>(nullptr, *Name);
+		REQUIRE(Package != FoundPackage);
+	}
+
+	SECTION("Multiple retries to load multiple packages")
+	{
+		FString Name1(FString::Printf(TEXT("/Game/%dAutoRTFMTestPackage%d"), FPlatformProcess::GetCurrentProcessId(), __LINE__));
+		FPackageScopedMaker _1(Name1);
+
+		FString Name2(FString::Printf(TEXT("/Game/%dAutoRTFMTestPackage%d"), FPlatformProcess::GetCurrentProcessId(), __LINE__));
+		FPackageScopedMaker _2(Name2);
+
+		FString Name3(FString::Printf(TEXT("/Game/%dAutoRTFMTestPackage%d"), FPlatformProcess::GetCurrentProcessId(), __LINE__));
+		FPackageScopedMaker _3(Name3);
+
+		FString Name4(FString::Printf(TEXT("/Game/%dAutoRTFMTestPackage%d"), FPlatformProcess::GetCurrentProcessId(), __LINE__));
+		FPackageScopedMaker _4(Name4);
+
+		TArray<FString> PackagesToLoad({Name1, Name2, Name3, Name4});
+
+		AutoRTFM::Testing::Commit([&]
+		{
+			for (const FString& Name : PackagesToLoad)
+			{
+				UPackage* const Package = LoadPackage(nullptr, *Name, LOAD_None);
+				REQUIRE(Package);
+			}
+		});
+	}
+
+	SECTION("Multiple retries to load multiple packages with trashing")
+	{
+		FString Name1(FString::Printf(TEXT("/Game/%dAutoRTFMTestPackage%d"), FPlatformProcess::GetCurrentProcessId(), __LINE__));
+		FPackageScopedMaker _1(Name1);
+
+		FString Name2(FString::Printf(TEXT("/Game/%dAutoRTFMTestPackage%d"), FPlatformProcess::GetCurrentProcessId(), __LINE__));
+		FPackageScopedMaker _2(Name2);
+
+		FString Name3(FString::Printf(TEXT("/Game/%dAutoRTFMTestPackage%d"), FPlatformProcess::GetCurrentProcessId(), __LINE__));
+		FPackageScopedMaker _3(Name3);
+
+		FString Name4(FString::Printf(TEXT("/Game/%dAutoRTFMTestPackage%d"), FPlatformProcess::GetCurrentProcessId(), __LINE__));
+		FPackageScopedMaker _4(Name4);
+
+		TArray<FString> PackagesToLoad({Name1, Name2, Name3, Name4});
+
+		AutoRTFM::Testing::Commit([&]
+		{
+			for (const FString& Name : PackagesToLoad)
 			{
 				UPackage* const Package = LoadPackage(nullptr, *Name, LOAD_None);
 				REQUIRE(Package);
@@ -450,8 +553,29 @@ TEST_CASE("UPackage.AsyncLoading")
 				REQUIRE(ReloadedPackage->GetFName() == Name);
 
 				REQUIRE(Package != ReloadedPackage);
-				REQUIRE(!ReloadedPackage->GetFName().ToString().StartsWith(TEXT("TrashedPackage")));
-			});
-		}
+			}
+		});
+	}
+
+	SECTION("Multiple async loads")
+	{
+		FString Name1(FString::Printf(TEXT("/Game/%dAutoRTFMTestPackage%d"), FPlatformProcess::GetCurrentProcessId(), __LINE__));
+		FPackageScopedMaker _1(Name1);
+
+		FString Name2(FString::Printf(TEXT("/Game/%dAutoRTFMTestPackage%d"), FPlatformProcess::GetCurrentProcessId(), __LINE__));
+		FPackageScopedMaker _2(Name2);
+
+		AutoRTFM::Testing::Commit([&]
+		{
+			REQUIRE(FindObject<UPackage>(nullptr, *Name1) == nullptr);
+			REQUIRE(FindObject<UPackage>(nullptr, *Name2) == nullptr);
+
+			FLoadPackageAsyncDelegate Delegate;
+			const int32 Id1 = LoadPackageAsync(Name1, Delegate);
+			const int32 Id2 = LoadPackageAsync(Name2, Delegate);
+			FlushAsyncLoading(Id2);
+			
+			REQUIRE(FindObject<UPackage>(nullptr, *Name2) != nullptr);
+		});
 	}
 }
diff --git a/Engine/Source/Runtime/Core/Public/UObject/UnrealNames.inl b/Engine/Source/Runtime/Core/Public/UObject/UnrealNames.inl
--- a/Engine/Source/Runtime/Core/Public/UObject/UnrealNames.inl
+++ b/Engine/Source/Runtime/Core/Public/UObject/UnrealNames.inl
@@ -210,6 +210,7 @@ REGISTER_NAME(260, Mobile)
 REGISTER_NAME(261, Oodle)
 REGISTER_NAME(262, TrashedPackage)
 REGISTER_NAME(263, MetaData)
+REGISTER_NAME(264, AutoRTFMUnusedPackage)
 
 // Online
 REGISTER_NAME(280,DGram)
diff --git a/Engine/Source/Runtime/CoreUObject/Private/Serialization/TransactionallySafeAsyncLoading.cpp b/Engine/Source/Runtime/CoreUObject/Private/Serialization/TransactionallySafeAsyncLoading.cpp
--- a/Engine/Source/Runtime/CoreUObject/Private/Serialization/TransactionallySafeAsyncLoading.cpp
+++ b/Engine/Source/Runtime/CoreUObject/Private/Serialization/TransactionallySafeAsyncLoading.cpp
@@ -45,6 +45,40 @@ static FAutoConsoleVariableRef CVarTransactionallySafeAsyncLoadingLegacyTrashedP
 // - When flushing a previous request-id, we check if the request-id was one
 //   that we know has already been flushed. If not, we need to abort and retry
 //   the transaction with the flush happening in between the abort and the retry.
+//
+// There is an additional issue we've had to address, dealing with trashing packages and
+// then issuing sync-loads for them in the same transaction. The problem is:
+// 
+// - `FindObject` returns `A`.
+// - We call `TrashPackage` on `A`.
+// - We do `LoadPackage` which does abort.
+// - We undo `TrashPackage` of `A`.
+// - We then issue the `LoadPackage`, which will return the pointer to `A` because there
+//   is already an object with the name.
+// - We retry.
+// - Repeat forever.
+// 
+// What we actually need to happen is:
+// 
+// - `FindObject` returns `A`.
+// - We call `TrashPackage` on `A`.
+// - We do `LoadPackage` which does abort.
+// - We undo `TrashPackage` of `A`.
+// - We then issue the `LoadPackage`, which needs to return a new package `B` that will
+//   eventually get the same name as `A`.
+// - We retry.
+// - `FindObject` returns `A`.
+// - We call `TrashPackage` on `A`.
+// - We do `LoadPackage` which returns `B`.
+//
+// To get this to work we detect the case where we are aborting-and-retrying to do a
+// sync-load for a package that already exists, and this case can *only* occur when
+// a package has been trashed and then a new sync-load has been issued in the same
+// transaction nest. We have to rename the old package temporarily, load a new package
+// that'll get a new object backing it, then rename the old package back to the original
+// package name. This ensures that we have both `A` (the original package we are going
+// to trash) and `B` (the new package ready to take over once `A` is trashed and we do
+// the sync-load) ready for the transactional retry.
 class FTransactionallySafeAsyncPackageLoader final : public IAsyncPackageLoader
 {
 public:
@@ -170,8 +204,23 @@ public:
 					(void)AutoRTFM::Close([&]
 					{
 						ensure(MapPayload->IsStillValid());
-						InCompletionDelegate.ExecuteIfBound(MapPayload->PackageName, MapPayload->LoadedPackage.Get(), MapPayload->Result);
+						
+						UPackage* const Package = MapPayload->GetPackage().Get();
+					
+						if (LIKELY(!CVar_TransactionallySafeAsyncLoadingLegacyTrashedPackages))
+						{
+							// Now that the package has come into actual use we rename it to the actual
+							// package name that it originally was conceived with, making it discoverable
+							// in `FindObject` calls.
+							if (Package)
+							{
+								Package->Rename(*(MapPayload->PackageName).ToString(), nullptr, REN_DontCreateRedirectors | REN_NonTransactional | REN_DoNotDirty);
+							}
+						}
+
+						InCompletionDelegate.ExecuteIfBound(MapPayload->PackageName, Package, MapPayload->Result);
 					});
+
 					return PackageCacheHit;
 				}
 			});
@@ -453,10 +502,10 @@ private:
 	{
 		FMapPayload() = default;
 		FMapPayload(FName PackageName, UPackage* const LoadedPackage, EAsyncLoadingResult::Type Result) :
-			PackageName(PackageName), LoadedPackage(MakeWeakObjectPtr(LoadedPackage)), Result(Result) {}
+			PackageName(PackageName), LoadedPackages({ MakeWeakObjectPtr(LoadedPackage) }), Result(Result) {}
 
 		FName PackageName;
-		TWeakObjectPtr<UPackage> LoadedPackage;
+		TArray<TWeakObjectPtr<UPackage>> LoadedPackages;
 		EAsyncLoadingResult::Type Result;
 
 		// Tells us whether a given map payload entry is still valid (EG. a cached loaded package didn't get GC'ed).
@@ -466,8 +515,10 @@ private:
 			// became null (meaning the GC did its thing).
 			if (EAsyncLoadingResult::Succeeded == Result)
 			{
+				TWeakObjectPtr<UPackage> Package = GetPackage();
+
 				// If the package has been GC'ed then it is not still valid.
-				if (!LoadedPackage.IsValid())
+				if (!Package.IsValid())
 				{
 					return false;
 				}
@@ -477,13 +528,16 @@ private:
 					return true;
 				}
 
-				if (LoadedPackage->GetFName() != PackageName)
+				static FString AutoRTFMUnusedPackage = FName(NAME_AutoRTFMUnusedPackage).ToString();
+				if (Package->GetFName().ToString().StartsWith(AutoRTFMUnusedPackage))
 				{
-					static FString TrashedPackage = FName(NAME_TrashedPackage).ToString();
-					check(LoadedPackage->GetFName().ToString().StartsWith(TrashedPackage));
-					return false;
+					// Unused packages are ones that have been created during a transaction
+					// retry, and thus are valid!
+					return true;
 				}
 
+				ensure(Package->GetFName() == PackageName);
+
 				return true;
 			}
 			else
@@ -491,6 +545,37 @@ private:
 				return true;
 			}
 		}
+
+		TWeakObjectPtr<UPackage> GetPackage() const
+		{
+			if (LoadedPackages.IsEmpty())
+			{
+				return nullptr;
+			}
+
+			if (UNLIKELY(CVar_TransactionallySafeAsyncLoadingLegacyTrashedPackages))
+			{
+				ensure(LoadedPackages.Num() <= 1);
+				return LoadedPackages[0].Get();
+			}
+
+			for (TWeakObjectPtr<UPackage> LoadedPackage : LoadedPackages)
+			{
+				if (LoadedPackage.IsValid())
+				{
+					// Skip trashed packages.
+					static FString TrashedPackage = FName(NAME_TrashedPackage).ToString();
+					if (LoadedPackage->GetFName().ToString().StartsWith(TrashedPackage))
+					{
+						continue;
+					}
+
+					return LoadedPackage;
+				}
+			}
+
+			return nullptr;
+		}
 	};
 
 	TMap<FString, FMapPayload> PackageCache;
@@ -525,16 +610,97 @@ private:
 	{
 		AutoRTFM::UnreachableIfTransactional(UnreachableMessage);
 
+		FString PackagePathStr = PackagePath.GetDebugNameWithExtension();
+
+		UPackage* OldPackage = nullptr;
+
+		if (LIKELY(!CVar_TransactionallySafeAsyncLoadingLegacyTrashedPackages))
+		{
+			if (AutoRTFM::IsRetrying())
+			{
+				UE::TScopeLock _(Mutex);
+
+				if (FMapPayload* const Payload = PackageCache.Find(PackagePathStr))
+				{
+					for (int32 Index = 0; Index < Payload->LoadedPackages.Num(); Index++)
+					{
+						UPackage* const Package = Payload->LoadedPackages[Index].Get();
+
+						if (Package == nullptr)
+						{
+							continue;
+						}
+
+						// We should only find a single old-package with the original package name.
+						ensure(OldPackage == nullptr);
+
+						static FString AutoRTFMUnusedPackage = FName(NAME_AutoRTFMUnusedPackage).ToString();
+						static FString TrashedPackage = FName(NAME_TrashedPackage).ToString();
+
+						// Unused packages (transactionally created during a retry, but not yet in use) and
+						// trashed packages are both handled elsewhere and won't affect us creating a brand
+						// new package like we want, so skip them.
+						if (Package->GetFName().ToString().StartsWith(AutoRTFMUnusedPackage) ||
+							Package->GetFName().ToString().StartsWith(TrashedPackage))
+						{
+							continue;
+						}
+
+						// We need to rename the old package because we need to make a brand new object with
+						// the same name and have that be a unique object. This case can only occur when we've
+						// trashed a package in a transaction and then issued a sync-load in the same
+						// transaction to reload the object. When we abort-and-retry (and get to the current
+						// code location) any previously trashed packages will have had that trashing undone
+						// during the transaction abort, but we *know* that when we retry the transaction we
+						// are going to immediately trash the object and thus need a new unique object.
+						const FName NewName = MakeUniqueObjectName(nullptr, UPackage::StaticClass(), NAME_AutoRTFMUnusedPackage);
+						Package->Rename(*(NewName.ToString()), nullptr, REN_DontCreateRedirectors | REN_NonTransactional | REN_DoNotDirty);
+
+						OldPackage = Package;
+					}
+				}
+			}
+		}
+
 		FLoadPackageAsyncDelegate WrapperDelegate;
-		WrapperDelegate.BindLambda([this, PackagePathStr = PackagePath.GetDebugNameWithExtension(), CompletionDelegate = ::MoveTemp(InCompletionDelegate)](
+		WrapperDelegate.BindLambda([this, PackagePathStr = PackagePathStr, CompletionDelegate = ::MoveTemp(InCompletionDelegate), OldPackage = OldPackage](
 			const FName& Name, UPackage* const Package, EAsyncLoadingResult::Type Type)
 		{
 			if (EAsyncLoadingResult::Canceled != Type)
 			{
-				FMapPayload Payload(Name, Package, Type);
-
 				UE::TScopeLock _(Mutex);
-				PackageCache.FindOrAdd(PackagePathStr) = Payload;
+
+				if (FMapPayload* const Payload = PackageCache.Find(PackagePathStr))
+				{
+					Payload->LoadedPackages.Add(Package);
+				}
+				else
+				{
+					PackageCache.FindOrAdd(PackagePathStr) = FMapPayload(Name, Package, Type);
+				}
+
+				if (LIKELY(!CVar_TransactionallySafeAsyncLoadingLegacyTrashedPackages))
+				{
+					if (AutoRTFM::IsRetrying() && Package)
+					{
+						// If we are retrying and managed to load a package, we need to rename the package from the actual given
+						// name to our special unused-package name. This is to prevent it appearing in `FindObject` calls *before*
+						// the sync-load is performed on the retry within the transaction.
+						const FString CreatedPackageName = Package->GetFName().ToString();
+
+						const FName NewName = MakeUniqueObjectName(nullptr, UPackage::StaticClass(), NAME_AutoRTFMUnusedPackage);
+						Package->Rename(*(NewName.ToString()), nullptr, REN_DontCreateRedirectors | REN_NonTransactional | REN_DoNotDirty);
+
+						// If `OldPackage` is a valid object it means we renamed it for the purposes of creating a new
+						// unused-as-of-yet package above. So we need to rename it back to the actual package name here
+						// so that when we begin the retry of the transaction the object has its correct name back and
+						// can be found in `FindObject` searches.
+						if (OldPackage)
+						{
+							OldPackage->Rename(*CreatedPackageName, nullptr, REN_DontCreateRedirectors | REN_NonTransactional | REN_DoNotDirty);
+						}
+					}
+				}
 			}
 
 			CompletionDelegate.ExecuteIfBound(Name, Package, Type);
