Allow FGCObject to be collected by the Verse VM garbage collector.

This check can be hit if you add a TSharedPtr<FGCObject> into an object which is managed by the new Verse VM's garbage collector.  (In this case, the FGCObject is an FContentScope.)

diff --git a/Engine/Source/Runtime/CoreUObject/Private/Misc/GCObjectReferencer.cpp b/Engine/Source/Runtime/CoreUObject/Private/Misc/GCObjectReferencer.cpp
--- a/Engine/Source/Runtime/CoreUObject/Private/Misc/GCObjectReferencer.cpp
+++ b/Engine/Source/Runtime/CoreUObject/Private/Misc/GCObjectReferencer.cpp
@@ -6,13 +6,15 @@
 
 #include "AutoRTFM.h"
 #include "CoreMinimal.h"
+#include "HAL/IConsoleManager.h"
 #include "Misc/ScopeLock.h"
-#include "UObject/ObjectMacros.h"
-#include "UObject/Object.h"
 #include "Templates/Casts.h"
-#include "UObject/GCObject.h"
+#include "UObject/GarbageCollection.h"
 #include "UObject/GarbageCollectionVerification.h"
-#include "HAL/IConsoleManager.h"
+#include "UObject/GCObject.h"
+#include "UObject/ObjectMacros.h"
+#include "UObject/Object.h"
+#include "VerseVM/VVMHeap.h"
 
 #include <atomic>
 
@@ -77,10 +79,29 @@ void UGCObjectReferencer::AddReferencedObjects(UObject* InThis, FReferenceCollec
 	Super::AddReferencedObjects(This, Collector);
 }
 
+static bool IsUnrealGCMarkingPassInProgress()
+{
+	return !GObjUnhashUnreachableIsInProgress && !GObjIncrementalPurgeIsInProgress && IsGarbageCollectingAndLockingUObjectHashTables();
+}
+
+static bool IsFrankenGCDestroyingPassInProgress()
+{
+#if WITH_VERSE_VM
+	// The FrankenGC object-destruction pass occurs immediately after the marking 
+	// pass of both Verse and Unreal has finished and reached a steady-state. 
+	// Unfortunately `IsUnrealGCMarkingPassInProgress` still returns true at that 
+	// point, even though all marking has actually concluded. 
+	// We detect this case here.
+	return UE::GC::GIsFrankenGCCollecting && Verse::FHeap::IsDestroying();
+#else
+	return false;
+#endif
+}
+
 void UGCObjectReferencer::AddObject(FGCObject* Object)
 {
 	check(Object);
-	check(GObjUnhashUnreachableIsInProgress || GObjIncrementalPurgeIsInProgress || !IsGarbageCollectingAndLockingUObjectHashTables());
+	check(!IsUnrealGCMarkingPassInProgress() || IsFrankenGCDestroyingPassInProgress());
 
 	TArray<FGCObject*>* ReferencedObjects = Object->bCanMakeInitialReferences ? &Impl->InitialReferencedObjects : &Impl->RemainingReferencedObjects;
 
@@ -108,8 +129,8 @@ void UGCObjectReferencer::AddObject(FGCObject* Object)
 void UGCObjectReferencer::RemoveObject(FGCObject* Object)
 {
 	check(Object);
-	check(GObjUnhashUnreachableIsInProgress || GObjIncrementalPurgeIsInProgress || !IsGarbageCollectingAndLockingUObjectHashTables());
-	
+	check(!IsUnrealGCMarkingPassInProgress() || IsFrankenGCDestroyingPassInProgress());
+
 	TArray<FGCObject*>* ReferencedObjects = Object->bCanMakeInitialReferences ? &Impl->InitialReferencedObjects : &Impl->RemainingReferencedObjects;
 
 	UE_AUTORTFM_OPEN
diff --git a/Engine/Source/Runtime/CoreUObject/Private/UObject/GarbageCollection.cpp b/Engine/Source/Runtime/CoreUObject/Private/UObject/GarbageCollection.cpp
--- a/Engine/Source/Runtime/CoreUObject/Private/UObject/GarbageCollection.cpp
+++ b/Engine/Source/Runtime/CoreUObject/Private/UObject/GarbageCollection.cpp
@@ -3924,7 +3924,7 @@ namespace UE::GC
 {
 #if WITH_VERSE_VM || defined(__INTELLISENSE__)
 bool GIsFrankenGCCollecting = false;
-static bool bInFrankenGCStartStop = false; // This just tracks if we are in a StartVersGC/StopVerseGC call pair
+static bool bInFrankenGCStartStop = false; // This just tracks if we are in a StartVerseGC/StopVerseGC call pair
 static bool bFrankenGCEnabled = false;
 static Verse::FCollectionCycleRequest VerseCycleRequest;
 
@@ -3996,11 +3996,11 @@ static FORCEINLINE void StopVerseGC()
 	bInFrankenGCStartStop = false;
 	if (GIsFrankenGCCollecting)
 	{
-		GIsFrankenGCCollecting = false;
-
 		Verse::FIOContext Context = Verse::FIOContextPromise{};
 		Verse::FHeap::ExternallySynchronouslyTerminateGC(Context);
 		VerseCycleRequest.Wait(Context);
+
+		GIsFrankenGCCollecting = false;
 	}
 }
 #else
@@ -6458,14 +6458,12 @@ void UE::GC::MarkNativeStructAsReachable(Verse::VNativeStruct* NativeStruct)
 }
 #endif
 
+UE_AUTORTFM_ALWAYS_OPEN
 void UObjectBase::MarkAsReachable() const
 {
 	// It is safe to perform mark as reachable in the open - the worst case is that we'll mark an object reachable that
 	// should/would be destroyed, and so in the next GC iteration it will be destroyed instead of in this iteration.
-	UE_AUTORTFM_OPEN
-	{
-		::MarkAsReachable<false>(this);
-	};
+	::MarkAsReachable</*bIsVerse*/ false>(this);
 }
 
 void UObjectBase::AddRef() const
@@ -6481,14 +6479,14 @@ void UObjectBase::ReleaseRef() const
 }
 
 #if WITH_VERSE_VM || defined(__INTELLISENSE__)
+
+UE_AUTORTFM_ALWAYS_OPEN
 void UObject::VerseMarkAsReachable() const
 {
-	UE_AUTORTFM_OPEN
-	{
-		::MarkAsReachable<true>(this);
-	};
+	::MarkAsReachable</*bIsVerse*/ true>(this);
 }
-#endif
+
+#endif // WITH_VERSE_VM || defined(__INTELLISENSE__)
 
 /*-----------------------------------------------------------------------------
 	Implementation of realtime garbage collection helper functions in 
