Removed call to IsThisNotNull from UStruct::IsChildOf and removed redundant calls to IsThisNotNull used along IsValidChecked(this).

#prefligfht 6895fa18593d93678d25226e

diff --git a/Engine/Plugins/FX/Niagara/Source/Niagara/Private/NiagaraComponent.cpp b/Engine/Plugins/FX/Niagara/Source/Niagara/Private/NiagaraComponent.cpp
--- a/Engine/Plugins/FX/Niagara/Source/Niagara/Private/NiagaraComponent.cpp
+++ b/Engine/Plugins/FX/Niagara/Source/Niagara/Private/NiagaraComponent.cpp
@@ -1922,9 +1922,8 @@ void UNiagaraComponent::DestroyInstanceNotComponent()
 
 void UNiagaraComponent::OnPooledReuse(UWorld* NewWorld)
 {
-	PRAGMA_DISABLE_DEPRECATION_WARNINGS
-	check(IsThisNotNull(this, "UNiagaraComponent::OnPooledReuse") && IsValidChecked(this));
-	PRAGMA_ENABLE_DEPRECATION_WARNINGS
+	check(IsValidChecked(this));
+
 	SetUserParametersToDefaultValues();
 
 	//Need to reset the component's visibility in case it's returned to the pool while marked invisible.
diff --git a/Engine/Source/Runtime/CoreUObject/Private/UObject/Class.cpp b/Engine/Source/Runtime/CoreUObject/Private/UObject/Class.cpp
--- a/Engine/Source/Runtime/CoreUObject/Private/UObject/Class.cpp
+++ b/Engine/Source/Runtime/CoreUObject/Private/UObject/Class.cpp
@@ -2751,12 +2751,6 @@ void UStruct::TagSubobjects(EObjectFlags NewFlags)
 #if USTRUCT_FAST_ISCHILDOF_COMPARE_WITH_OUTERWALK || USTRUCT_FAST_ISCHILDOF_IMPL == USTRUCT_ISCHILDOF_OUTERWALK
 bool UStruct::IsChildOf( const UStruct* SomeBase ) const
 {
-	// If you're looking at this check it is due to calling IsChildOf with a this nullptr. *MAKE* sure you do not call this function
-	// with a this nullptr. It is undefined behavior, and some compilers, clang13 have started to optimize out this == nullptr checks.
-	PRAGMA_DISABLE_DEPRECATION_WARNINGS
-	check(IsThisNotNull(this, "UStruct::IsChildOf"));
-	PRAGMA_ENABLE_DEPRECATION_WARNINGS
-
 	if (SomeBase == nullptr)
 	{
 		return false;
diff --git a/Engine/Source/Runtime/Engine/Private/ActorConstruction.cpp b/Engine/Source/Runtime/Engine/Private/ActorConstruction.cpp
--- a/Engine/Source/Runtime/Engine/Private/ActorConstruction.cpp
+++ b/Engine/Source/Runtime/Engine/Private/ActorConstruction.cpp
@@ -800,9 +800,7 @@ namespace
 
 bool AActor::ExecuteConstruction(const FTransform& Transform, const FRotationConversionCache* TransformRotationCache, const FComponentInstanceDataCache* InstanceDataCache, bool bIsDefaultTransform, ESpawnActorScaleMethod TransformScaleMethod)
 {
-	PRAGMA_DISABLE_DEPRECATION_WARNINGS
-	check(IsThisNotNull(this, "AActor::ExecuteConstruction") && IsValidChecked(this));
-	PRAGMA_ENABLE_DEPRECATION_WARNINGS
+	check(IsValidChecked(this));
 	check(!HasAnyFlags(RF_BeginDestroyed|RF_FinishDestroyed));
 
 #if WITH_EDITOR
