Fixed a crash when changing a transient property that was triggered by another recent localization fix

diff --git a/Engine/Source/Runtime/CoreUObject/Private/Serialization/TextReferenceCollector.cpp b/Engine/Source/Runtime/CoreUObject/Private/Serialization/TextReferenceCollector.cpp
--- a/Engine/Source/Runtime/CoreUObject/Private/Serialization/TextReferenceCollector.cpp
+++ b/Engine/Source/Runtime/CoreUObject/Private/Serialization/TextReferenceCollector.cpp
@@ -6,6 +6,7 @@
 #include "Serialization/ArchiveUObject.h"
 #include "UObject/Class.h"
 #include "UObject/Package.h"
+#include "UObject/UnrealType.h"
 
 class FTextReferencesArchive : public FArchiveUObject
 {
@@ -18,7 +19,7 @@ public:
 		, Count(&OutCount)
 	{
 		this->SetIsSaving(true);
-		this->SetIsPersistent(true); // Skips transient properties
+		this->SetIsPersistent(false);
 		ArShouldSkipBulkData = true; // Skips bulk data as we can't handle saving that!
 
 		// Build up the list of objects that are within our package - we won't follow object references to things outside of our package
@@ -43,6 +44,12 @@ public:
 		}
 	}
 
+	virtual bool ShouldSkipProperty(const FProperty* InProperty) const override
+	{
+		return FArchiveUObject::ShouldSkipProperty(InProperty)
+			|| (InProperty->HasAnyPropertyFlags(CPF_Transient) && !IsSerializingDefaults());
+	}
+
 	void ProcessObject(UObject* Obj)
 	{
 		if (!Obj || !AllObjectsInPackage.Contains(Obj) || ProcessedObjects.Contains(Obj))
