Improve performance of Cast<> and IsA<> for final classes.
(1M iter: ~10ms -> ~3ms in editor on failed casts)

diff --git a/Engine/Source/Runtime/CoreUObject/Public/UObject/UObjectBaseUtility.h b/Engine/Source/Runtime/CoreUObject/Public/UObject/UObjectBaseUtility.h
--- a/Engine/Source/Runtime/CoreUObject/Public/UObject/UObjectBaseUtility.h
+++ b/Engine/Source/Runtime/CoreUObject/Public/UObject/UObjectBaseUtility.h
@@ -642,7 +642,15 @@ public:
 	template<class T>
 	bool IsA() const
 	{
-		return IsA(T::StaticClass());
+		// If the class is final, then there's no point in checking for child classes
+		if constexpr (std::is_final_v<T>)
+		{
+			return GetClass() == T::StaticClass();
+		}
+		else
+		{
+			return IsA(T::StaticClass());
+		}
 	}
 
 	/**
