PropertyBag: Add a feature to be able to duplicate a property.

diff --git a/Engine/Source/Runtime/CoreUObject/Private/StructUtils/PropertyBag.cpp b/Engine/Source/Runtime/CoreUObject/Private/StructUtils/PropertyBag.cpp
--- a/Engine/Source/Runtime/CoreUObject/Private/StructUtils/PropertyBag.cpp
+++ b/Engine/Source/Runtime/CoreUObject/Private/StructUtils/PropertyBag.cpp
@@ -2016,6 +2016,52 @@ EPropertyBagAlterationResult FInstancedPropertyBag::AddProperty(const FName InNa
 	return AddProperties({ FPropertyBagPropertyDesc(InName, InSourceProperty ) }, bOverwrite);
 }
 
+EPropertyBagAlterationResult FInstancedPropertyBag::DuplicateProperty(const FName InName, FName* InOutNewName)
+{
+	if (const FPropertyBagPropertyDesc* Desc = FindPropertyDescByName(InName))
+	{
+		FName NewName = InOutNewName ? *InOutNewName : NAME_None;
+		if (NewName.IsNone())
+		{
+			NewName = InName;
+
+			do
+			{
+				// We start from _2 (so Number=3)
+				NewName.SetNumber(FMath::Max(NewName.GetNumber() + 1, 3));
+			}
+			while (FindPropertyDescByName(NewName));
+
+			if (InOutNewName)
+			{
+				*InOutNewName = NewName;
+			}
+		}
+
+		FPropertyBagPropertyDesc NewDesc = *Desc;
+		NewDesc.ID = FGuid();
+		NewDesc.Name = NewName;
+
+		const EPropertyBagAlterationResult Result = AddProperties({ NewDesc }, false);
+		if (Result != EPropertyBagAlterationResult::Success)
+		{
+			return Result;
+		}
+
+		const FPropertyBagPropertyDesc* OriginalDesc = FindPropertyDescByName(InName);
+		const FPropertyBagPropertyDesc* AddedDesc = FindPropertyDescByName(NewName);
+		check(OriginalDesc && AddedDesc && OriginalDesc->CachedProperty && AddedDesc->CachedProperty);
+
+		void* TargetAddress = Value.GetMutableMemory() + AddedDesc->CachedProperty->GetOffset_ForInternal();
+		void const* SourceAddress = Value.GetMemory() + OriginalDesc->CachedProperty->GetOffset_ForInternal();
+		AddedDesc->CachedProperty->CopyCompleteValue(TargetAddress, SourceAddress);
+
+		return EPropertyBagAlterationResult::Success;
+	}
+
+	return EPropertyBagAlterationResult::SourcePropertyNotFound;
+}
+
 EPropertyBagResult FInstancedPropertyBag::ReplaceAllPropertiesAndValues(const TConstArrayView<FPropertyBagPropertyDesc> InDescs, const TConstArrayView<TConstArrayView<uint8>> InValues)
 {
 	if (InDescs.Num() != InValues.Num())
diff --git a/Engine/Source/Runtime/CoreUObject/Public/StructUtils/PropertyBag.h b/Engine/Source/Runtime/CoreUObject/Public/StructUtils/PropertyBag.h
--- a/Engine/Source/Runtime/CoreUObject/Public/StructUtils/PropertyBag.h
+++ b/Engine/Source/Runtime/CoreUObject/Public/StructUtils/PropertyBag.h
@@ -506,6 +506,13 @@ struct FInstancedPropertyBag
 	 */
 	UE_API EPropertyBagAlterationResult AddProperty(const FName InName, const FProperty* InSourceProperty, bool bOverwrite = true);
 
+	/**
+	 * Duplicates the given property in the bag if it exists.
+	 * @param InName Name of the existing property
+	 * @param InOutNewName Name of the new property
+	 */
+	UE_API EPropertyBagAlterationResult DuplicateProperty(FName InName, FName* InOutNewName = nullptr);
+
 	/**
 	 * Clears all properties, then adds the supplied properties to the bag and sets their respective values.
 	 * Numeric property values will be converted if possible, when a property's type changes.
