Adding extra logging for non deterministic fiel errors due to FStructUtils::AttemptToFindUninitializedScriptStructMembers()

diff --git a/Engine/Source/Runtime/CoreUObject/Private/UObject/Class.cpp b/Engine/Source/Runtime/CoreUObject/Private/UObject/Class.cpp
--- a/Engine/Source/Runtime/CoreUObject/Private/UObject/Class.cpp
+++ b/Engine/Source/Runtime/CoreUObject/Private/UObject/Class.cpp
@@ -4378,7 +4378,10 @@ int32 FStructUtils::AttemptToFindUninitializedScriptStructMembers()
 
 		for (const FProperty* Property : UninitializedPropertiesZeroed)
 		{
-			LogUninitializedProperty(Property, TEXT("is not initialized properly even though its struct probably has a custom default constructor"));
+			LogUninitializedProperty(
+				Property,
+				TEXT("is not initialized properly even though its struct probably has a custom default constructor. "
+					"Non deterministic fields should use UPROPERTY(Meta = (IgnoreForMemberInitializationTest)) to avoid errors from this test"));
 		}
 
 		for (const FProperty* Property : UninitializedPropertiesNoInit)
