Fix BugHawk issue found from 44759575:
UStruct::IsChildOf had its explicit null 'this' check removed, which contradicts the documented guarantee that it "will not crash on null structs". Calling a non-static member function via a null pointer is undefined behavior in C++, and without the guard, the function now relies on a valid 'this' and immediately iterates super-structs, creating a mismatch between documentation and implementation and reducing diagnostic clarity.

#prefligfht skip

diff --git a/Engine/Source/Runtime/CoreUObject/Public/UObject/Class.h b/Engine/Source/Runtime/CoreUObject/Public/UObject/Class.h
--- a/Engine/Source/Runtime/CoreUObject/Public/UObject/Class.h
+++ b/Engine/Source/Runtime/CoreUObject/Public/UObject/Class.h
@@ -686,7 +686,7 @@ public:
 		PropertiesSize = NewSize;
 	}
 
-	/** Returns true if this struct either is class T, or is a child of class T. This will not crash on null structs */
+	/** Returns true if this struct either is class T, or is a child of class T. */
 	template<class T>
 	bool IsChildOf() const
 	{
@@ -705,7 +705,7 @@ public:
 		}
 	}
 
-	/** Returns true if this struct either is SomeBase, or is a child of SomeBase. This will not crash on null structs */
+	/** Returns true if this struct either is SomeBase, or is a child of SomeBase. */
 #if USTRUCT_FAST_ISCHILDOF_COMPARE_WITH_OUTERWALK || USTRUCT_FAST_ISCHILDOF_IMPL == USTRUCT_ISCHILDOF_OUTERWALK
 	COREUOBJECT_API bool IsChildOf( const UStruct* SomeBase ) const;
 #else
