RemoteObject:  Exposing GetClass that may skip migrations when used with FRemoteObjectId.

#rnx

diff --git a/Engine/Source/Runtime/CoreUObject/Private/UObject/RemoteObject.cpp b/Engine/Source/Runtime/CoreUObject/Private/UObject/RemoteObject.cpp
--- a/Engine/Source/Runtime/CoreUObject/Private/UObject/RemoteObject.cpp
+++ b/Engine/Source/Runtime/CoreUObject/Private/UObject/RemoteObject.cpp
@@ -1,6 +1,7 @@
 // Copyright Epic Games, Inc. All Rights Reserved.
 
 #include "UObject/RemoteObject.h"
+
 #include "UObject/RemoteExecutor.h"
 #include "UObject/ObjectPtr.h"
 #include "UObject/Class.h"
@@ -20,6 +21,7 @@
 #include "Misc/ConfigCacheIni.h"
 #include "Misc/CoreDelegates.h"
 #include "String/LexFromString.h"
+#include "Templates/Casts.h"
 #include "Hash/CityHash.h"
 #include "Modules/VisualizerDebuggingState.h"
 #include <atomic>
@@ -773,6 +775,29 @@ bool CanResolveObject(FRemoteObjectId ObjectId)
 	return false;
 }
 
+UClass* GetClass(FRemoteObjectId ObjectId, ERemoteObjectGetClassBehavior GetClassBehavior)
+{
+	if (FRemoteObjectStub* Stub = Private::FindRemoteObjectStub(ObjectId))
+	{
+		if (Stub->Class.IsValid())
+		{
+			return Stub->Class.GetClass();
+		}
+		else if (GetClassBehavior == ERemoteObjectGetClassBehavior::ReturnNullIfNeverLocal)
+		{
+			return nullptr;
+		}
+		else if (GetClassBehavior == ERemoteObjectGetClassBehavior::MigrateIfNeverLocal)
+		{
+			UObject* ResolvedObject = ResolveObject(Stub, ERemoteReferenceType::Strong);
+			return ResolvedObject->GetClass();
+		}
+	}
+
+	return nullptr;
+}
+
+
 } // namespace UE::RemoteObject::Handle
 
 namespace UE::CoreUObject::Private
diff --git a/Engine/Source/Runtime/CoreUObject/Public/UObject/RemoteObject.h b/Engine/Source/Runtime/CoreUObject/Public/UObject/RemoteObject.h
--- a/Engine/Source/Runtime/CoreUObject/Public/UObject/RemoteObject.h
+++ b/Engine/Source/Runtime/CoreUObject/Public/UObject/RemoteObject.h
@@ -143,6 +143,15 @@ namespace UE::RemoteObject::Handle
 		Weak = 1
 	};
 
+	enum class ERemoteObjectGetClassBehavior
+	{
+		/** If the object was never local, then simply return a null class */
+		ReturnNullIfNeverLocal,
+
+		/** If the object was never local, force a migration to enforce a correct return value */
+		MigrateIfNeverLocal,
+	};
+
 	/**
 	* Resolves a remote object given its stub, aborting the active transaction if the object is unavailable
 	* @param Stub Basic data required to migrate a remote object
@@ -161,6 +170,11 @@ namespace UE::RemoteObject::Handle
 	
 	COREUOBJECT_API void TouchResidentObject(UObject* Object);
 
+	/**
+	 * Attempts to Get the Class of a given ObjectId without performing a migration.  This will succeed if the ObjectId has ever been local.
+	 * If the RemoteObjectId has never been local, you can decide if a Migration will occur or Null should be returned.
+	 */
+	COREUOBJECT_API UClass* GetClass(FRemoteObjectId ObjectId, ERemoteObjectGetClassBehavior GetClassBehavior);
 
 	/**
 	* Resolves a remote object
