UE: Handle cases of invalid index in header data when resolving name of missing import. This will result in a missing name in logs but will no longer crash. Of course, we don't want this to happen but if it does happen, it shouldn't result in a crash due to attempting to emit a debugging log.

diff --git a/Engine/Source/Runtime/CoreUObject/Private/Serialization/AsyncLoading2.cpp b/Engine/Source/Runtime/CoreUObject/Private/Serialization/AsyncLoading2.cpp
--- a/Engine/Source/Runtime/CoreUObject/Private/Serialization/AsyncLoading2.cpp
+++ b/Engine/Source/Runtime/CoreUObject/Private/Serialization/AsyncLoading2.cpp
@@ -2338,7 +2338,8 @@ struct FPackageImportStore
 
 	FString GetImportName(const FAsyncPackageHeaderData& Header, FPackageObjectIndex GlobalIndex)
 	{
-		const UObject* Import = GetImportObject(Header, GlobalIndex);
+		bool bIsInvalidPackageImport = GlobalIndex.IsPackageImport() && !Header.ImportedPackageIds.IsValidIndex(GlobalIndex.ToPackageImportRef().GetImportedPackageIndex());
+		const UObject* Import = bIsInvalidPackageImport ? nullptr : GetImportObject(Header, GlobalIndex);
 		if (Import)
 		{
 			return Import->GetFName().ToString();
