Cooker: SkipOnlyEditorOnly: When an out-of-date package is upgraded at load time because it has a SoftObjectPath property that has changed to a hard import, we do a syncload of the SoftObjectPath. We need to tell the cooker this syncload is expected to suppress the warnings for unexpected loads.
#rnx

diff --git a/Engine/Source/Runtime/CoreUObject/Private/UObject/PropertyObject.cpp b/Engine/Source/Runtime/CoreUObject/Private/UObject/PropertyObject.cpp
--- a/Engine/Source/Runtime/CoreUObject/Private/UObject/PropertyObject.cpp
+++ b/Engine/Source/Runtime/CoreUObject/Private/UObject/PropertyObject.cpp
@@ -5,6 +5,7 @@
 #include "UObject/SoftObjectPtr.h"
 #include "UObject/UnrealType.h"
 #include "Blueprint/BlueprintSupport.h"
+#include "UObject/ICookInfo.h"
 #include "UObject/LinkerPlaceholderBase.h"
 #include "UObject/LinkerPlaceholderExportObject.h"
 #include "UObject/LinkerPlaceholderClass.h"
@@ -67,6 +68,7 @@ EConvertFromTypeResult FObjectProperty::ConvertFromType(const FPropertyTag& Tag,
 		// This property used to be a TSoftObjectPtr<Foo> but is now a raw FObjectProperty Foo*, we can convert without loss of data
 		FSoftObjectPtr PreviousValue;
 		Slot << PreviousValue;
+		ECookLoadType CookLoadType = this->IsEditorOnlyProperty() ? ECookLoadType::EditorOnly : ECookLoadType::UsedInGame;
 
 		UObject* PreviousValueObj = nullptr;
 
@@ -108,12 +110,14 @@ EConvertFromTypeResult FObjectProperty::ConvertFromType(const FPropertyTag& Tag,
 #endif
 				else
 				{
+					FCookLoadScope CookLoadScope(CookLoadType);
 					PreviousValueObj = PreviousValue.LoadSynchronous();
 				}
 			}
 		}
 		else
 		{
+			FCookLoadScope CookLoadScope(CookLoadType);
 			PreviousValueObj = PreviousValue.LoadSynchronous();
 		}
 
