Correct define usage (UE_BUILD_SHIPPING vs UE_SHIPPING)

diff --git a/Engine/Source/Runtime/CoreUObject/Private/Serialization/AsyncLoading2.cpp b/Engine/Source/Runtime/CoreUObject/Private/Serialization/AsyncLoading2.cpp
--- a/Engine/Source/Runtime/CoreUObject/Private/Serialization/AsyncLoading2.cpp
+++ b/Engine/Source/Runtime/CoreUObject/Private/Serialization/AsyncLoading2.cpp
@@ -323,12 +323,8 @@ void FindAllRuntimeScriptPackages(FRuntimeScriptPackages& OutPackages)
 #define ALT2_ENABLE_PACKAGE_DEPENDENCY_DEBUGGING 0
 #endif
 
-#ifndef UE_SHIPPING
-	#define UE_SHIPPING 0
-#endif
-
 #ifndef ALT2_ENABLE_SCRIPT_IMPORTS_DEBUGGING
-	#define ALT2_ENABLE_SCRIPT_IMPORTS_DEBUGGING !UE_SHIPPING && !IS_PROGRAM
+	#define ALT2_ENABLE_SCRIPT_IMPORTS_DEBUGGING !UE_BUILD_SHIPPING && !IS_PROGRAM
 #endif
 
 #if ALT2_ENABLE_SCRIPT_IMPORTS_DEBUGGING
