FOverriddenPropertyNodeID minor optimization: when creating the node path, use the object index as the name index instead of goind through Printf.

diff --git a/Engine/Source/Runtime/CoreUObject/Private/UObject/OverriddenPropertySet.cpp b/Engine/Source/Runtime/CoreUObject/Private/UObject/OverriddenPropertySet.cpp
--- a/Engine/Source/Runtime/CoreUObject/Private/UObject/OverriddenPropertySet.cpp
+++ b/Engine/Source/Runtime/CoreUObject/Private/UObject/OverriddenPropertySet.cpp
@@ -359,7 +359,7 @@ FOverriddenPropertyNodeID::FOverriddenPropertyNodeID(TNotNull<const UObject*> In
 	// Note: Using ObjectIndex by itself is not sufficient for an enduring unique identifier
 	// as re-instantiation can cause a reuse of the index for another object. Appending the serial solves this issue
 	const int32 ObjectIndex = GUObjectArray.ObjectToIndex(KeyObject);
-	Path = FName(FString::Printf(TEXT("ID_%d_%d"), ObjectIndex, GUObjectArray.AllocateSerialNumber(ObjectIndex)));
+	Path = FName(FString::Printf(TEXT("ID_%d"), ObjectIndex), GUObjectArray.AllocateSerialNumber(ObjectIndex));
 }
 
 FOverriddenPropertyNodeID FOverriddenPropertyNodeID::RootNodeId()
