Expose FOrientedBox to blueprint

diff --git a/Engine/Source/Runtime/CoreUObject/Public/UObject/NoExportTypes.h b/Engine/Source/Runtime/CoreUObject/Public/UObject/NoExportTypes.h
--- a/Engine/Source/Runtime/CoreUObject/Public/UObject/NoExportTypes.h
+++ b/Engine/Source/Runtime/CoreUObject/Public/UObject/NoExportTypes.h
@@ -1682,35 +1682,35 @@ struct FBoxSphereBounds
  * Structure for arbitrarily oriented boxes (i.e. not necessarily axis-aligned).
  * @note The full C++ class is located here: Engine\Source\Runtime\Core\Public\Math\OrientedBox.h
  */
-USTRUCT(immutable, noexport, IsAlwaysAccessible, HasDefaults)
+USTRUCT(immutable, noexport, BlueprintType, IsAlwaysAccessible, HasDefaults)
 struct FOrientedBox
 {
 	/** Holds the center of the box. */
-	UPROPERTY(EditAnywhere, Category=OrientedBox, SaveGame)
+	UPROPERTY(EditAnywhere, BlueprintReadWrite, Category=OrientedBox, SaveGame)
 	FVector Center;
 
 	/** Holds the x-axis vector of the box. Must be a unit vector. */
-	UPROPERTY(EditAnywhere, Category=OrientedBox, SaveGame)
+	UPROPERTY(EditAnywhere, BlueprintReadOnly, Category=OrientedBox, SaveGame)
 	FVector AxisX;
 	
 	/** Holds the y-axis vector of the box. Must be a unit vector. */
-	UPROPERTY(EditAnywhere, Category=OrientedBox, SaveGame)
+	UPROPERTY(EditAnywhere, BlueprintReadOnly, Category=OrientedBox, SaveGame)
 	FVector AxisY;
 	
 	/** Holds the z-axis vector of the box. Must be a unit vector. */
-	UPROPERTY(EditAnywhere, Category=OrientedBox, SaveGame)
+	UPROPERTY(EditAnywhere, BlueprintReadOnly, Category=OrientedBox, SaveGame)
 	FVector AxisZ;
 
 	/** Holds the extent of the box along its x-axis. */
-	UPROPERTY(EditAnywhere, Category=OrientedBox, SaveGame)
+	UPROPERTY(EditAnywhere, BlueprintReadWrite, Category=OrientedBox, SaveGame)
 	FLargeWorldCoordinatesReal ExtentX;
 	
 	/** Holds the extent of the box along its y-axis. */
-	UPROPERTY(EditAnywhere, Category=OrientedBox, SaveGame)
+	UPROPERTY(EditAnywhere, BlueprintReadWrite, Category=OrientedBox, SaveGame)
 	FLargeWorldCoordinatesReal ExtentY;
 
 	/** Holds the extent of the box along its z-axis. */
-	UPROPERTY(EditAnywhere, Category=OrientedBox, SaveGame)
+	UPROPERTY(EditAnywhere, BlueprintReadWrite, Category=OrientedBox, SaveGame)
 	FLargeWorldCoordinatesReal ExtentZ;
 };
 
