FObjectPtr UEOpEquals takes its parameter by reference to avoid copy-construction and potential calls to MarkAsReachable.

diff --git a/Engine/Source/Runtime/CoreUObject/Public/UObject/ObjectPtr.h b/Engine/Source/Runtime/CoreUObject/Public/UObject/ObjectPtr.h
--- a/Engine/Source/Runtime/CoreUObject/Public/UObject/ObjectPtr.h
+++ b/Engine/Source/Runtime/CoreUObject/Public/UObject/ObjectPtr.h
@@ -167,7 +167,7 @@ public:
 		return *this;
 	}
 
-	[[nodiscard]] FORCEINLINE bool UEOpEquals(FObjectPtr Other) const
+	[[nodiscard]] FORCEINLINE bool UEOpEquals(const FObjectPtr& Other) const
 	{
 		return (Handle == Other.Handle);
 	}
