Fix spelling mistake: BULKDATA_PayloadInSeperateFile -> BULKDATA_PayloadInSeparateFile

The original spelling is retained in the enum for API backwards compatibility.

diff --git a/Engine/Plugins/Mutable/Source/CustomizableObject/Private/MuCO/CustomizableObject.cpp b/Engine/Plugins/Mutable/Source/CustomizableObject/Private/MuCO/CustomizableObject.cpp
--- a/Engine/Plugins/Mutable/Source/CustomizableObject/Private/MuCO/CustomizableObject.cpp
+++ b/Engine/Plugins/Mutable/Source/CustomizableObject/Private/MuCO/CustomizableObject.cpp
@@ -213,7 +213,7 @@ void UCustomizableObject::PreSave(FObjectPreSaveContext ObjectSaveContext)
 						FMemory::Memcpy(Ptr, FileBulkData.GetData(), FileBulkData.Num());
 						ByteBulkData.Unlock();
 
-						uint32 BulkDataFlags = BULKDATA_PayloadInSeperateFile | BULKDATA_Force_NOT_InlinePayload;
+						uint32 BulkDataFlags = BULKDATA_PayloadInSeparateFile | BULKDATA_Force_NOT_InlinePayload;
 						if (File.Flags == uint16(EMutableFileFlags::HighRes))
 						{
 							BulkDataFlags |= BULKDATA_OptionalPayload;
diff --git a/Engine/Plugins/Mutable/Source/CustomizableObject/Private/MuCO/CustomizableObjectInstance.cpp b/Engine/Plugins/Mutable/Source/CustomizableObject/Private/MuCO/CustomizableObjectInstance.cpp
--- a/Engine/Plugins/Mutable/Source/CustomizableObject/Private/MuCO/CustomizableObjectInstance.cpp
+++ b/Engine/Plugins/Mutable/Source/CustomizableObject/Private/MuCO/CustomizableObjectInstance.cpp
@@ -3563,7 +3563,7 @@ FTexturePlatformData* MutableCreateImagePlatformData(TSharedPtr<const mu::FImage
 		}
 		else
 		{
-			Mip->BulkData.SetBulkDataFlags(BULKDATA_PayloadInSeperateFile);
+			Mip->BulkData.SetBulkDataFlags(BULKDATA_PayloadInSeparateFile);
 			Mip->BulkData.ClearBulkDataFlags(BULKDATA_PayloadAtEndOfFile);
 		}
 
@@ -6292,7 +6292,7 @@ void UCustomizableInstancePrivate::BuildMaterials(const TSharedRef<FUpdateContex
 										// To prevent LogTexture Error "Loading non-streamed mips from an external bulk file."
 										for (int32 i = 0; i < MutableTexture->GetPlatformData()->Mips.Num(); ++i)
 										{
-											MutableTexture->GetPlatformData()->Mips[i].BulkData.ClearBulkDataFlags(BULKDATA_PayloadInSeperateFile);
+											MutableTexture->GetPlatformData()->Mips[i].BulkData.ClearBulkDataFlags(BULKDATA_PayloadInSeparateFile);
 										}
 									}
 
diff --git a/Engine/Source/Runtime/CoreUObject/Private/Serialization/BulkData.cpp b/Engine/Source/Runtime/CoreUObject/Private/Serialization/BulkData.cpp
--- a/Engine/Source/Runtime/CoreUObject/Private/Serialization/BulkData.cpp
+++ b/Engine/Source/Runtime/CoreUObject/Private/Serialization/BulkData.cpp
@@ -60,7 +60,7 @@ FStringBuilderBase& LexToString(EBulkDataFlags Flags, FStringBuilderBase& Sb)
 	TEST_AND_ADD_FLAG(Sb, Flags, BULKDATA_SingleUse);
 	TEST_AND_ADD_FLAG(Sb, Flags, BULKDATA_ForceInlinePayload);
 	TEST_AND_ADD_FLAG(Sb, Flags, BULKDATA_SerializeCompressed);
-	TEST_AND_ADD_FLAG(Sb, Flags, BULKDATA_PayloadInSeperateFile);
+	TEST_AND_ADD_FLAG(Sb, Flags, BULKDATA_PayloadInSeparateFile);
 	TEST_AND_ADD_FLAG(Sb, Flags, BULKDATA_Force_NOT_InlinePayload);
 	TEST_AND_ADD_FLAG(Sb, Flags, BULKDATA_OptionalPayload);
 	TEST_AND_ADD_FLAG(Sb, Flags, BULKDATA_MemoryMappedPayload);
@@ -1356,7 +1356,7 @@ void FBulkData::Serialize(FArchive& Ar, UObject* Owner, bool bAttemptFileMapping
 				SetBulkDataFlagsOn(SerializedMeta.Flags, static_cast<EBulkDataFlags>(BULKDATA_Size64Bit));
 			}
 
-			const EBulkDataFlags FlagsToClear = static_cast<EBulkDataFlags>(BULKDATA_PayloadAtEndOfFile | BULKDATA_PayloadInSeperateFile | BULKDATA_WorkspaceDomainPayload | BULKDATA_ForceSingleElementSerialization);
+			const EBulkDataFlags FlagsToClear = static_cast<EBulkDataFlags>(BULKDATA_PayloadAtEndOfFile | BULKDATA_PayloadInSeparateFile | BULKDATA_WorkspaceDomainPayload | BULKDATA_ForceSingleElementSerialization);
 			FBulkData::ClearBulkDataFlagsOn(SerializedMeta.Flags, FlagsToClear);
 
 			const int64 MetaOffset = Ar.Tell();
@@ -1405,7 +1405,7 @@ PRAGMA_ENABLE_DEPRECATION_WARNINGS
 	InBulkDataFlags = static_cast<EBulkDataFlags>(InBulkDataFlags & ~BULKDATA_UsesIoDispatcher);
 
 	const bool bPayloadInline = !(InBulkDataFlags & BULKDATA_PayloadAtEndOfFile);
-	const bool bPayloadInSeparateFile = !bPayloadInline && (InBulkDataFlags & BULKDATA_PayloadInSeperateFile);
+	const bool bPayloadInSeparateFile = !bPayloadInline && (InBulkDataFlags & BULKDATA_PayloadInSeparateFile);
 	// fix up the file offset if the offset is relative to the file's bulkdata section
 	if (!bPayloadInline && NeedsOffsetFixup())
 	{
diff --git a/Engine/Source/Runtime/CoreUObject/Private/UObject/LinkerLoad.cpp b/Engine/Source/Runtime/CoreUObject/Private/UObject/LinkerLoad.cpp
--- a/Engine/Source/Runtime/CoreUObject/Private/UObject/LinkerLoad.cpp
+++ b/Engine/Source/Runtime/CoreUObject/Private/UObject/LinkerLoad.cpp
@@ -138,7 +138,7 @@ void TrackPackageAssetClass(UPackage* Package, FLinkerLoad& LinkerLoad, const TA
 
 EPackageSegment GetBulkDataPackageSegmentFromFlags(const EBulkDataFlags BulkDataFlags, bool bLoadingFromCookedPackage)
 {
-	if (FBulkData::HasFlags(BulkDataFlags, BULKDATA_PayloadInSeperateFile) == false)
+	if (FBulkData::HasFlags(BulkDataFlags, BULKDATA_PayloadInSeparateFile) == false)
 	{
 		return bLoadingFromCookedPackage ? EPackageSegment::Exports : EPackageSegment::Header;
 	}
@@ -7285,7 +7285,7 @@ bool FLinkerLoad::SerializeBulkData(FBulkData& BulkData, const FBulkDataSerializ
 		void* Payload = BulkData.ReallocateData(Meta.GetSize());
 		BulkData.SerializeBulkData(*this, Payload, Meta.GetSize(), Meta.GetFlags());
 	}
-	else if (Meta.HasAnyFlags(BULKDATA_PayloadInSeperateFile))
+	else if (Meta.HasAnyFlags(BULKDATA_PayloadInSeparateFile))
 	{
 		// Streaming cooked bulk data / loading from Editor Domain and referencing Workspace domain bulk data
 		if (Meta.HasAnyFlags(BULKDATA_DuplicateNonOptionalPayload))
diff --git a/Engine/Source/Runtime/CoreUObject/Private/UObject/LinkerSave.cpp b/Engine/Source/Runtime/CoreUObject/Private/UObject/LinkerSave.cpp
--- a/Engine/Source/Runtime/CoreUObject/Private/UObject/LinkerSave.cpp
+++ b/Engine/Source/Runtime/CoreUObject/Private/UObject/LinkerSave.cpp
@@ -609,7 +609,7 @@ bool FLinkerSave::SerializeBulkData(FBulkData& BulkData, const FBulkDataSerializ
 		FBulkData::SetBulkDataFlagsOn(SerializedMeta.Flags, static_cast<EBulkDataFlags>(BULKDATA_Size64Bit));
 	}
 
-	EBulkDataFlags FlagsToClear = static_cast<EBulkDataFlags>(BULKDATA_PayloadAtEndOfFile | BULKDATA_PayloadInSeperateFile | BULKDATA_WorkspaceDomainPayload | BULKDATA_ForceSingleElementSerialization | BULKDATA_NoOffsetFixUp);
+	EBulkDataFlags FlagsToClear = static_cast<EBulkDataFlags>(BULKDATA_PayloadAtEndOfFile | BULKDATA_PayloadInSeparateFile | BULKDATA_WorkspaceDomainPayload | BULKDATA_ForceSingleElementSerialization | BULKDATA_NoOffsetFixUp);
 	if (IsCooking())
 	{
 		FBulkData::SetBulkDataFlagsOn(FlagsToClear, static_cast<EBulkDataFlags>(BULKDATA_SerializeCompressed));
@@ -657,14 +657,14 @@ bool FLinkerSave::SerializeBulkData(FBulkData& BulkData, const FBulkDataSerializ
 		if (bSaveBulkDataToSeparateFiles)
 		{
 			check(bSaveBulkDataByReference == false);
-			FBulkData::SetBulkDataFlagsOn(SerializedMeta.Flags, static_cast<EBulkDataFlags>(BULKDATA_PayloadInSeperateFile | BULKDATA_NoOffsetFixUp));
+			FBulkData::SetBulkDataFlagsOn(SerializedMeta.Flags, static_cast<EBulkDataFlags>(BULKDATA_PayloadInSeparateFile | BULKDATA_NoOffsetFixUp));
 		}
 		
 		const bool bSaveByReference = bSaveBulkDataByReference && CanSaveBulkDataByReference(BulkData);
 		if (bSaveByReference)
 		{
 			check(IsCooking() == false);
-			FBulkData::SetBulkDataFlagsOn(SerializedMeta.Flags, static_cast<EBulkDataFlags>(BULKDATA_NoOffsetFixUp | BULKDATA_WorkspaceDomainPayload | BULKDATA_PayloadInSeperateFile));
+			FBulkData::SetBulkDataFlagsOn(SerializedMeta.Flags, static_cast<EBulkDataFlags>(BULKDATA_NoOffsetFixUp | BULKDATA_WorkspaceDomainPayload | BULKDATA_PayloadInSeparateFile));
 		}
 
 		if (bSaveBulkDataToSeparateFiles && FBulkData::HasFlags(SerializedMeta.Flags, BULKDATA_OptionalPayload))
diff --git a/Engine/Source/Runtime/CoreUObject/Public/Serialization/BulkData.h b/Engine/Source/Runtime/CoreUObject/Public/Serialization/BulkData.h
--- a/Engine/Source/Runtime/CoreUObject/Public/Serialization/BulkData.h
+++ b/Engine/Source/Runtime/CoreUObject/Public/Serialization/BulkData.h
@@ -118,7 +118,8 @@ enum EBulkDataFlags : uint32
 	/** Forces the payload to be always streamed, regardless of its size. */
 	BULKDATA_ForceStreamPayload UE_DEPRECATED(5.3, "This flag has had no purpose for sometime") = 1 << 7,
 	/** Set when the bulkdata saved/cooked as part of a package indicates that the payload is stored in its own file (.ubulk, .uptnl or .m.ubulk depending on other flags) */
-	BULKDATA_PayloadInSeperateFile = 1 << 8,
+	BULKDATA_PayloadInSeparateFile = 1 << 8,
+	BULKDATA_PayloadInSeperateFile UE_DEPRECATED(5.7, "This flag has been renamed to BULKDATA_PayloadInSeparateFile") = BULKDATA_PayloadInSeparateFile,
 	/** DEPRECATED */
 	BULKDATA_SerializeCompressedBitWindow UE_DEPRECATED(5.3, "This flag has had no purpose for sometime") = 1 << 9,
 	/** Should be set before being cooked as part of a package to force the cooked payload to before stored at the end of the package file either than being inline*/
@@ -643,11 +644,11 @@ public:
 
 	/**
 	* Returns whether this bulk data is currently stored in it's own file or not
-	* @return true if BULKDATA_PayloadInSeperateFile is not set
+	* @return true if BULKDATA_PayloadInSeparateFile is not set
 	*/
 	bool IsInSeparateFile() const
 	{
-		return	(GetBulkDataFlags() & BULKDATA_PayloadInSeperateFile) != 0;
+		return	(GetBulkDataFlags() & BULKDATA_PayloadInSeparateFile) != 0;
 	}
 
 	/**
diff --git a/Engine/Source/Runtime/CoreUObject/Tests/Serialization/BulkDataTests.cpp b/Engine/Source/Runtime/CoreUObject/Tests/Serialization/BulkDataTests.cpp
--- a/Engine/Source/Runtime/CoreUObject/Tests/Serialization/BulkDataTests.cpp
+++ b/Engine/Source/Runtime/CoreUObject/Tests/Serialization/BulkDataTests.cpp
@@ -234,7 +234,7 @@ TEST_CASE("CoreUObject::Serialization::FBulkData::Basic", "[CoreUObject][Seriali
 		FBulkData BulkData;
 		
 		CHECK(TestBulkDataFlags(BulkData, BULKDATA_PayloadAtEndOfFile, [&BulkData]() { return !BulkData.IsInlined(); }));
-		CHECK(TestBulkDataFlags(BulkData, BULKDATA_PayloadInSeperateFile, [&BulkData]() { return BulkData.IsInSeparateFile(); }));
+		CHECK(TestBulkDataFlags(BulkData, BULKDATA_PayloadInSeparateFile, [&BulkData]() { return BulkData.IsInSeparateFile(); }));
 		CHECK(TestBulkDataFlags(BulkData, BULKDATA_OptionalPayload, [&BulkData]() { return BulkData.IsOptional(); }));
 		CHECK(TestBulkDataFlags(BulkData, BULKDATA_DuplicateNonOptionalPayload, [&BulkData]() { return BulkData.IsDuplicateNonOptional(); }));
 		CHECK(TestBulkDataFlags(BulkData, BULKDATA_DataIsMemoryMapped, [&BulkData]() { return BulkData.IsDataMemoryMapped(); }));
@@ -285,7 +285,7 @@ TEST_CASE("CoreUObject::Serialization::FBulkData::Serialize", "[CoreUObject][Ser
 	SECTION("Serialize to memory archive does not serialize invalid flags")
 	{
 		const int64 ExpectedPayloadSize = 1024;
-		const uint32 InvalidFlags = uint32(BULKDATA_PayloadAtEndOfFile | BULKDATA_PayloadInSeperateFile | BULKDATA_WorkspaceDomainPayload);
+		const uint32 InvalidFlags = uint32(BULKDATA_PayloadAtEndOfFile | BULKDATA_PayloadInSeparateFile | BULKDATA_WorkspaceDomainPayload);
 
 		FUniqueBuffer Payload = CreatePayload(ExpectedPayloadSize);
 		FLargeMemoryWriter Ar;
