Fixed doc-comment in accessor.

diff --git a/Engine/Source/Runtime/CoreUObject/Public/Serialization/BulkData.h b/Engine/Source/Runtime/CoreUObject/Public/Serialization/BulkData.h
--- a/Engine/Source/Runtime/CoreUObject/Public/Serialization/BulkData.h
+++ b/Engine/Source/Runtime/CoreUObject/Public/Serialization/BulkData.h
@@ -644,7 +644,7 @@ public:
 
 	/**
 	* Returns whether this bulk data is currently stored in it's own file or not
-	* @return true if BULKDATA_PayloadInSeparateFile is not set
+	* @return true if BULKDATA_PayloadInSeparateFile is set
 	*/
 	bool IsInSeparateFile() const
 	{
