Updated GetDefault to support const UClass*

diff --git a/Engine/Source/Runtime/CoreUObject/Public/UObject/Class.h b/Engine/Source/Runtime/CoreUObject/Public/UObject/Class.h
--- a/Engine/Source/Runtime/CoreUObject/Public/UObject/Class.h
+++ b/Engine/Source/Runtime/CoreUObject/Public/UObject/Class.h
@@ -4639,7 +4639,7 @@ inline bool UObject::Implements() const
  * @see GetMutableDefault
  */
 template< class T > 
-inline const T* GetDefault(UClass *Class)
+inline const T* GetDefault(const UClass* Class)
 {
 	check(Class->GetDefaultObject()->IsA(T::StaticClass()));
 	return (const T*)Class->GetDefaultObject();
diff --git a/Engine/Source/Runtime/CoreUObject/Public/UObject/UObjectGlobals.h b/Engine/Source/Runtime/CoreUObject/Public/UObject/UObjectGlobals.h
--- a/Engine/Source/Runtime/CoreUObject/Public/UObject/UObjectGlobals.h
+++ b/Engine/Source/Runtime/CoreUObject/Public/UObject/UObjectGlobals.h
@@ -2157,7 +2157,7 @@ inline const T* GetDefault()
  * @see Class.h
  */
 template< class T > 
-inline const T* GetDefault(UClass *Class);
+inline const T* GetDefault(const UClass* Class);
 
 /** Version of GetDefault() that allows modification */
 template< class T >
