Making sure that an object reconstructed on top of an existing object from the disregard for GC pool will NOT have the Reachable flag set upon reconstruction.

diff --git a/Engine/Source/Runtime/CoreUObject/Private/UObject/UObjectArray.cpp b/Engine/Source/Runtime/CoreUObject/Private/UObject/UObjectArray.cpp
--- a/Engine/Source/Runtime/CoreUObject/Private/UObject/UObjectArray.cpp
+++ b/Engine/Source/Runtime/CoreUObject/Private/UObject/UObjectArray.cpp
@@ -254,14 +254,12 @@ void FUObjectArray::AllocateUObjectIndex(UObjectBase* Object, EInternalObjectFla
 
 	LockInternalArray();
 
-	const bool bDisregardForGCOpenAndEnabled = (IsOpenForDisregardForGC() & GUObjectArray.DisregardForGCEnabled()); //-V792 
-
 	if (AlreadyAllocatedIndex >= 0)
 	{
 		Index = AlreadyAllocatedIndex;
 	}
 	// Special non- garbage collectable range.
-	else if (bDisregardForGCOpenAndEnabled)
+	else if (IsOpenForDisregardForGC() & GUObjectArray.DisregardForGCEnabled()) //-V792 
 	{
 		Index = ++ObjLastNonGCIndex;
 		// Check if we're not out of bounds, unless there hasn't been any gc objects yet
@@ -297,7 +295,8 @@ void FUObjectArray::AllocateUObjectIndex(UObjectBase* Object, EInternalObjectFla
 	UE_CLOG(ObjectItem->GetObject() != nullptr, LogUObjectArray, Fatal, TEXT("Attempting to add %s at index %d but another object (0x%016llx) exists at that index!"), *Object->GetFName().ToString(), Index, (int64)(PTRINT)ObjectItem->GetObject());
 	// At this point all not-compiled-in objects are not fully constructed yet and this is the earliest we can mark them as such
 	ObjectItem->FlagsAndRefCount = (int64)((uint64)EInternalObjectFlags::PendingConstruction << 32);
-	if (!bDisregardForGCOpenAndEnabled)
+	// Objects in the disregad for GC pool don't need the reachable flag set because GC will never process them
+	if (!IsIndexDisregardForGC(Index))
 	{
 		// It's safe to access FGCFlags::GetReachableFlagValue_ForGC() here because creating new objects is being performed 
 		// under the same UObjectArray lock as swapping reachability flags inside of GC, see FGCFlags::SwapReachableAndMaybeUnreachable()
