SOL-8388 - "Make creative device map run in the new VM"

Makes FOptionalProperty pass in its DefaultData's 'ValuePointer', or nullptr if unset, to the call to instance its ValueProperty's subobjects

#rnx

diff --git a/Engine/Source/Runtime/CoreUObject/Private/UObject/PropertyOptional.cpp b/Engine/Source/Runtime/CoreUObject/Private/UObject/PropertyOptional.cpp
--- a/Engine/Source/Runtime/CoreUObject/Private/UObject/PropertyOptional.cpp
+++ b/Engine/Source/Runtime/CoreUObject/Private/UObject/PropertyOptional.cpp
@@ -475,7 +475,7 @@ void FOptionalProperty::InstanceSubobjects(void* Data, void const* DefaultData,
 		&& IsSet(Data)
 		&& (ValueProperty->ContainsInstancedObjectProperty() || InOwner->GetClass()->ShouldUseDynamicSubobjectInstancing()))
 	{
-		ValueProperty->InstanceSubobjects(Data, DefaultData, InOwner, InstanceGraph);
+		ValueProperty->InstanceSubobjects(Data, DefaultData ? GetValuePointerForReadIfSet(DefaultData) : nullptr, InOwner, InstanceGraph);
 	}
 }
 
