Move FCookDependencyContext to its own h/cpp.

#rnx

diff --git a/Engine/Plugins/FX/Niagara/Source/Niagara/Private/NiagaraSystem.cpp b/Engine/Plugins/FX/Niagara/Source/Niagara/Private/NiagaraSystem.cpp
--- a/Engine/Plugins/FX/Niagara/Source/Niagara/Private/NiagaraSystem.cpp
+++ b/Engine/Plugins/FX/Niagara/Source/Niagara/Private/NiagaraSystem.cpp
@@ -39,6 +39,7 @@
 #include "Algo/RemoveIf.h"
 #include "Algo/StableSort.h"
 #include "Cooker/CookDependency.h"
+#include "Cooker/CookDependencyContext.h"
 #include "Cooker/CookEvents.h"
 #include "HAL/LowLevelMemTracker.h"
 #include "Interfaces/ITargetPlatform.h"
diff --git a/Engine/Plugins/Runtime/ComputeFramework/Source/ComputeFramework/Private/ComputeGraph.cpp b/Engine/Plugins/Runtime/ComputeFramework/Source/ComputeFramework/Private/ComputeGraph.cpp
--- a/Engine/Plugins/Runtime/ComputeFramework/Source/ComputeFramework/Private/ComputeGraph.cpp
+++ b/Engine/Plugins/Runtime/ComputeFramework/Source/ComputeFramework/Private/ComputeGraph.cpp
@@ -25,6 +25,7 @@
 #include "UObject/Package.h"
 
 #if WITH_EDITOR
+#include "Cooker/CookDependencyContext.h"
 #include "Cooker/CookEvents.h"
 #endif
 
diff --git a/Engine/Source/Editor/Kismet/Private/Blueprints/BlueprintDependencies.cpp b/Engine/Source/Editor/Kismet/Private/Blueprints/BlueprintDependencies.cpp
--- a/Engine/Source/Editor/Kismet/Private/Blueprints/BlueprintDependencies.cpp
+++ b/Engine/Source/Editor/Kismet/Private/Blueprints/BlueprintDependencies.cpp
@@ -4,6 +4,7 @@
 
 #include "AssetRegistry/AssetRegistryModule.h"
 #include "AssetRegistry/IAssetRegistry.h"
+#include "Cooker/CookDependencyContext.h"
 #include "Engine/Blueprint.h"
 #include "HAL/FileManager.h"
 #include "IO/IoHash.h"
diff --git a/Engine/Source/Editor/UnrealEd/Private/Cooker/CookPackageArtifacts.cpp b/Engine/Source/Editor/UnrealEd/Private/Cooker/CookPackageArtifacts.cpp
--- a/Engine/Source/Editor/UnrealEd/Private/Cooker/CookPackageArtifacts.cpp
+++ b/Engine/Source/Editor/UnrealEd/Private/Cooker/CookPackageArtifacts.cpp
@@ -8,6 +8,7 @@
 #include "Cooker/CookAssetRegistryAccessTracker.h"
 #include "Cooker/CookConfigAccessTracker.h"
 #include "Cooker/CookDependency.h"
+#include "Cooker/CookDependencyContext.h"
 #include "Cooker/CookEvents.h"
 #include "Cooker/CookGenerationHelper.h"
 #include "Cooker/CookGlobalDependencies.h"
diff --git a/Engine/Source/Editor/UnrealEd/Private/EditorDomain/EditorDomainUtils.cpp b/Engine/Source/Editor/UnrealEd/Private/EditorDomain/EditorDomainUtils.cpp
--- a/Engine/Source/Editor/UnrealEd/Private/EditorDomain/EditorDomainUtils.cpp
+++ b/Engine/Source/Editor/UnrealEd/Private/EditorDomain/EditorDomainUtils.cpp
@@ -10,6 +10,7 @@
 #include "AssetRegistry/IAssetRegistry.h"
 #include "AssetRegistry/PackageReader.h"
 #include "Containers/DirectoryTree.h"
+#include "Cooker/CookDependencyContext.h"
 #include "Cooker/CookPackageArtifacts.h"
 #include "DerivedDataBuildDefinition.h"
 #include "DerivedDataCache.h"
diff --git a/Engine/Source/Runtime/CoreUObject/Private/Cooker/CookDependency.cpp b/Engine/Source/Runtime/CoreUObject/Private/Cooker/CookDependency.cpp
--- a/Engine/Source/Runtime/CoreUObject/Private/Cooker/CookDependency.cpp
+++ b/Engine/Source/Runtime/CoreUObject/Private/Cooker/CookDependency.cpp
@@ -9,6 +9,7 @@
 #include "Containers/Array.h"
 #include "Containers/Map.h"
 #include "Cooker/ConsoleVariableData.h"
+#include "Cooker/CookDependencyContext.h"
 #include "CoreGlobals.h"
 #include "HAL/FileManager.h"
 #include "HAL/IConsoleManager.h"
@@ -1051,57 +1052,6 @@ bool FCookDependency::ConsoleVariableDataEqual(const FConsoleVariableData& A, co
 	return A == B;
 }
 
-void FCookDependencyContext::Update(const void* Data, uint64 Size)
-{
-	FBlake3& Blake3Hasher = *(reinterpret_cast<FBlake3*>(Hasher));
-	Blake3Hasher.Update(Data, Size);
-}
-
-void FCookDependencyContext::ReportInvalidated()
-{
-	OnInvalidated(ELogVerbosity::Display);
-}
-
-void FCookDependencyContext::ReportError()
-{
-	OnInvalidated(ELogVerbosity::Error);
-}
-
-void FCookDependencyContext::Log(ELogVerbosity::Type Verbosity, FString Message)
-{
-	for (TUniqueFunction<FString(FString&&)>& Handler : ErrorHandlers)
-	{
-		Message = Handler(MoveTemp(Message));
-	}
-	OnLog(Verbosity, MoveTemp(Message));
-}
-
-FCookDependencyContext::FErrorHandlerScope::FErrorHandlerScope(FCookDependencyContext& InContext)
-	: Context(InContext)
-{
-}
-
-FCookDependencyContext::FErrorHandlerScope::~FErrorHandlerScope()
-{
-	check(!Context.ErrorHandlers.IsEmpty());
-	Context.ErrorHandlers.Pop(EAllowShrinking::No);
-}
-
-FCookDependencyContext::FErrorHandlerScope FCookDependencyContext::ErrorHandlerScope(
-	TUniqueFunction<FString(FString&&)>&& ErrorHandler)
-{
-	ErrorHandlers.Add(MoveTemp(ErrorHandler));
-	// We rely on Copy elision to not move-construct this scope, so that it only calls the destructor once.
-	return FErrorHandlerScope(*this);
-}
-
-void* FCookDependencyContext::SetHasher(void* NewHasher)
-{
-	void* Temp = Hasher;
-	Hasher = NewHasher;
-	return Temp;
-}
-
 } // namespace UE::Cook
 
 namespace UE::Cook::Dependency::Private
diff --git a/Engine/Source/Runtime/CoreUObject/Private/Cooker/CookDependencyContext.cpp b/Engine/Source/Runtime/CoreUObject/Private/Cooker/CookDependencyContext.cpp
new file mode 100644
--- /dev/null
+++ b/Engine/Source/Runtime/CoreUObject/Private/Cooker/CookDependencyContext.cpp
@@ -0,0 +1,84 @@
+// Copyright Epic Games, Inc. All Rights Reserved.
+
+#include "Cooker/CookDependencyContext.h"
+
+#include "Hash/Blake3.h"
+
+#if WITH_EDITOR
+
+namespace UE::Cook
+{
+FCookDependencyContext::FCookDependencyContext(void* InHasher,
+	TUniqueFunction<void(ELogVerbosity::Type, FString&&)>&& InOnLog,
+	TUniqueFunction<void(ELogVerbosity::Type)>&& InOnInvalidated,
+	FName InPackageName)
+	: OnLog(MoveTemp(InOnLog))
+	, OnInvalidated(MoveTemp(InOnInvalidated))
+	, PackageName(InPackageName)
+	, Hasher(InHasher)
+{ }
+
+void FCookDependencyContext::LogError(FString Message)
+{
+	Log(ELogVerbosity::Error, MoveTemp(Message));
+	ReportError();
+}
+
+void FCookDependencyContext::LogInvalidated(FString Message)
+{
+	Log(ELogVerbosity::Display, MoveTemp(Message));
+	ReportInvalidated();
+}
+
+void FCookDependencyContext::Update(const void* Data, uint64 Size)
+{
+	FBlake3& Blake3Hasher = *(reinterpret_cast<FBlake3*>(Hasher));
+	Blake3Hasher.Update(Data, Size);
+}
+
+void FCookDependencyContext::ReportInvalidated()
+{
+	OnInvalidated(ELogVerbosity::Display);
+}
+
+void FCookDependencyContext::ReportError()
+{
+	OnInvalidated(ELogVerbosity::Error);
+}
+
+void FCookDependencyContext::Log(ELogVerbosity::Type Verbosity, FString Message)
+{
+	for (TUniqueFunction<FString(FString&&)>& Handler : ErrorHandlers)
+	{
+		Message = Handler(MoveTemp(Message));
+	}
+	OnLog(Verbosity, MoveTemp(Message));
+}
+
+void* FCookDependencyContext::SetHasher(void* NewHasher)
+{
+	void* Temp = Hasher;
+	Hasher = NewHasher;
+	return Temp;
+}
+
+FCookDependencyContext::FErrorHandlerScope::FErrorHandlerScope(FCookDependencyContext& InContext)
+	: Context(InContext)
+{ }
+
+FCookDependencyContext::FErrorHandlerScope::~FErrorHandlerScope()
+{
+	check(!Context.ErrorHandlers.IsEmpty());
+	Context.ErrorHandlers.Pop(EAllowShrinking::No);
+}
+
+FCookDependencyContext::FErrorHandlerScope FCookDependencyContext::ErrorHandlerScope(
+	TUniqueFunction<FString(FString&&)>&& ErrorHandler)
+{
+	ErrorHandlers.Add(MoveTemp(ErrorHandler));
+	// We rely on Copy elision to not move-construct this scope, so that it only calls the destructor once.
+	return FErrorHandlerScope(*this);
+}
+}
+
+#endif //#if WITH_EDITOR
diff --git a/Engine/Source/Runtime/CoreUObject/Public/Cooker/CookDependency.h b/Engine/Source/Runtime/CoreUObject/Public/Cooker/CookDependency.h
--- a/Engine/Source/Runtime/CoreUObject/Public/Cooker/CookDependency.h
+++ b/Engine/Source/Runtime/CoreUObject/Public/Cooker/CookDependency.h
@@ -19,6 +19,7 @@ class UObject;
 namespace UE::ConfigAccessTracking { enum class ELoadType : uint8; }
 namespace UE::ConfigAccessTracking { struct FConfigAccessData; }
 namespace UE::Cook { class FConsoleVariableData; }
+namespace UE::Cook { struct FCookDependencyContext; }
 struct FARFilter;
 #endif
 
@@ -26,89 +27,6 @@ struct FARFilter;
 
 namespace UE::Cook
 {
-
-/** Context passed into UE_COOK_DEPENDENCY_FUNCTION functions to provide calling flags and receive hash output. */
-struct FCookDependencyContext
-{
-public:
-	/** InHasher is void* to mask the implementation details of the hashbuilder. See Update function. */
-	explicit FCookDependencyContext(void* InHasher,
-		TUniqueFunction<void(ELogVerbosity::Type, FString&&)>&& InOnLog,
-		TUniqueFunction<void(ELogVerbosity::Type)>&& InOnInvalidated,
-		FName InPackageName);
-
-	/**
-	 * Update the hashbuilder for the key being constructed (e.g. TargetDomainKey for cooked packages)
-	 * with the given Data of Size bytes.
-	 */
-	COREUOBJECT_API void Update(const void* Data, uint64 Size);
-
-	/**
-	 * Reports that current evaluation of the function is different from all previous evaluations for a reason that
-	 * cannot be reported as data passed into Update. When called while calculating the initial hash this call is
-	 * ignored. When called while testing incrementally skippable the package is marked modified and recooked.
-	 */
-	COREUOBJECT_API void ReportInvalidated();
-
-	/**
-	 * Reports failure to compute the hash (e.g. because a file cannot be read).
-	 * When called while calculating the initial hash the storage of the key fails and the package will be recooked on
-	 * the next cook. When called while testing incrementally skippable the package is marked modified and recooked.
-	 */
-	COREUOBJECT_API void ReportError();
-
-	/**
-	 * Send a message to the cook dependency context with the given severity. The message may be suppressed or
-	 * reduced in verbosity based on the calling context.
-	 */
-	COREUOBJECT_API void Log(ELogVerbosity::Type, FString Message);
-
-	/** Calls Log(Error, Message) and ReportError. */
-	void LogError(FString Message);
-
-	/** Calls Log(Display, Message) and ReportInvalidated. */
-	void LogInvalidated(FString Message);
-
-	/**
-	 * Private implementation struct used for AddErrorHandlerScope. Should only be used via
-	 * FCookDependencyContext::FErrorHandlerScope& Scope = Context.ErrorHandlerScope(<Function>);
-	 */
-	struct FErrorHandlerScope
-	{
-	public:
-		COREUOBJECT_API ~FErrorHandlerScope();
-	private:
-		COREUOBJECT_API FErrorHandlerScope(FCookDependencyContext& InContext);
-		friend FCookDependencyContext;
-		FCookDependencyContext& Context;
-	};
-
-	/**
-	 * Add a function that will be removed when the return value goes out of scope, to modify error strings reported
-	 * inside the scope before passing them on to higher scopes or the error consumer.
-	 * e.g. 
-	 * FCookDependencyContext::FErrorHandlerScope Scope = Context.ErrorHandlerScope([](FString&& Inner)
-	 * { return FString::Printf(TEXT("OuterClass for %s: %s"), *Name, *Inner);});
-	 */
-	[[nodiscard]] COREUOBJECT_API FErrorHandlerScope ErrorHandlerScope(
-		TUniqueFunction<FString(FString&&)>&& ErrorHandler);
-
-	/**
-	 * Get the name of the package being considered
-	 */
-	FName GetPackageName() const { return PackageName; }
-
-	/** Set a new hasher and return the old one. */
-	void* SetHasher(void* NewHasher);
-
-private:
-	TUniqueFunction<void(ELogVerbosity::Type, FString&&)> OnLog;
-	TUniqueFunction<void(ELogVerbosity::Type)> OnInvalidated;
-	TArray<TUniqueFunction<FString(FString&&)>, TInlineAllocator<1>> ErrorHandlers;
-	FName PackageName;
-	void* Hasher; // Type is void* to mask the implementation detail
-};
-
 /**
  * TypeSelector enum for the FCookDependency variable type. Values are serialized into the oplog as integers,
  * so do not change them without changing oplog version.
@@ -475,19 +393,6 @@ struct FCookDependencyFunctionRegistration
 #if WITH_EDITOR
 namespace UE::Cook
 {
-
-inline void FCookDependencyContext::LogError(FString Message)
-{
-	Log(ELogVerbosity::Error, MoveTemp(Message));
-	ReportError();
-}
-
-inline void FCookDependencyContext::LogInvalidated(FString Message)
-{
-	Log(ELogVerbosity::Display, MoveTemp(Message));
-	ReportInvalidated();
-}
-
 inline ECookDependency FCookDependency::GetType() const
 {
 	return Type;
@@ -756,17 +661,6 @@ inline void FCookDependency::SetValue(const FUtf8String& String)
 	SetValue(ConvertToHash(String));
 }
 
-inline FCookDependencyContext::FCookDependencyContext(void* InHasher,
-	TUniqueFunction<void(ELogVerbosity::Type, FString&&)>&& InOnLog,
-	TUniqueFunction<void(ELogVerbosity::Type)>&& InOnInvalidated,
-	FName InPackageName)
-	: OnLog(MoveTemp(InOnLog))
-	, OnInvalidated(MoveTemp(InOnInvalidated))
-	, PackageName(InPackageName)
-	, Hasher(InHasher)
-{
-}
-
 } // namespace UE::Cook
 
 namespace UE::Cook::Dependency::Private
diff --git a/Engine/Source/Runtime/CoreUObject/Public/Cooker/CookDependencyContext.h b/Engine/Source/Runtime/CoreUObject/Public/Cooker/CookDependencyContext.h
new file mode 100644
--- /dev/null
+++ b/Engine/Source/Runtime/CoreUObject/Public/Cooker/CookDependencyContext.h
@@ -0,0 +1,97 @@
+// Copyright Epic Games, Inc. All Rights Reserved.
+
+#pragma once
+
+#include "Containers/UnrealString.h"
+#include "HAL/Platform.h"
+#include "Templates/Function.h"
+#include "UObject/NameTypes.h"
+
+#if WITH_EDITOR
+
+namespace UE::Cook
+{
+/** Context passed into UE_COOK_DEPENDENCY_FUNCTION functions to provide calling flags and receive hash output. */
+struct FCookDependencyContext
+{
+public:
+	/** InHasher is void* to mask the implementation details of the hashbuilder. See Update function. */
+	COREUOBJECT_API explicit FCookDependencyContext(void* InHasher,
+		TUniqueFunction<void(ELogVerbosity::Type, FString&&)>&& InOnLog,
+		TUniqueFunction<void(ELogVerbosity::Type)>&& InOnInvalidated,
+		FName InPackageName);
+
+	/**
+	 * Update the hashbuilder for the key being constructed (e.g. TargetDomainKey for cooked packages)
+	 * with the given Data of Size bytes.
+	 */
+	COREUOBJECT_API void Update(const void* Data, uint64 Size);
+
+	/**
+	 * Reports that current evaluation of the function is different from all previous evaluations for a reason that
+	 * cannot be reported as data passed into Update. When called while calculating the initial hash this call is
+	 * ignored. When called while testing incrementally skippable the package is marked modified and recooked.
+	 */
+	COREUOBJECT_API void ReportInvalidated();
+
+	/**
+	 * Reports failure to compute the hash (e.g. because a file cannot be read).
+	 * When called while calculating the initial hash the storage of the key fails and the package will be recooked on
+	 * the next cook. When called while testing incrementally skippable the package is marked modified and recooked.
+	 */
+	COREUOBJECT_API void ReportError();
+
+	/**
+	 * Send a message to the cook dependency context with the given severity. The message may be suppressed or
+	 * reduced in verbosity based on the calling context.
+	 */
+	COREUOBJECT_API void Log(ELogVerbosity::Type, FString Message);
+
+	/** Calls Log(Error, Message) and ReportError. */
+	COREUOBJECT_API void LogError(FString Message);
+
+	/** Calls Log(Display, Message) and ReportInvalidated. */
+	COREUOBJECT_API void LogInvalidated(FString Message);
+
+	/**
+	 * Private implementation struct used for AddErrorHandlerScope. Should only be used via
+	 * FCookDependencyContext::FErrorHandlerScope& Scope = Context.ErrorHandlerScope(<Function>);
+	 */
+	struct FErrorHandlerScope
+	{
+	public:
+		COREUOBJECT_API ~FErrorHandlerScope();
+	private:
+		COREUOBJECT_API FErrorHandlerScope(FCookDependencyContext& InContext);
+		friend FCookDependencyContext;
+		FCookDependencyContext& Context;
+	};
+
+	/**
+	 * Add a function that will be removed when the return value goes out of scope, to modify error strings reported
+	 * inside the scope before passing them on to higher scopes or the error consumer.
+	 * e.g.
+	 * FCookDependencyContext::FErrorHandlerScope Scope = Context.ErrorHandlerScope([](FString&& Inner)
+	 * { return FString::Printf(TEXT("OuterClass for %s: %s"), *Name, *Inner);});
+	 */
+	[[nodiscard]] COREUOBJECT_API FErrorHandlerScope ErrorHandlerScope(
+		TUniqueFunction<FString(FString&&)>&& ErrorHandler);
+
+	/**
+	 * Get the name of the package being considered
+	 */
+	FName GetPackageName() const { return PackageName; }
+
+	/** Set a new hasher and return the old one. */
+	void* SetHasher(void* NewHasher);
+
+private:
+	TUniqueFunction<void(ELogVerbosity::Type, FString&&)> OnLog;
+	TUniqueFunction<void(ELogVerbosity::Type)> OnInvalidated;
+	TArray<TUniqueFunction<FString(FString&&)>, TInlineAllocator<1>> ErrorHandlers;
+	FName PackageName;
+	void* Hasher; // Type is void* to mask the implementation detail
+};
+}
+
+#endif //#if WITH_EDITOR
diff --git a/Engine/Source/Runtime/Engine/Classes/Sound/SoundWaveLoadingBehavior.cpp b/Engine/Source/Runtime/Engine/Classes/Sound/SoundWaveLoadingBehavior.cpp
--- a/Engine/Source/Runtime/Engine/Classes/Sound/SoundWaveLoadingBehavior.cpp
+++ b/Engine/Source/Runtime/Engine/Classes/Sound/SoundWaveLoadingBehavior.cpp
@@ -16,6 +16,7 @@
 #if WITH_EDITOR
 #include "Algo/Sort.h"
 #include "Cooker/CookDependency.h"
+#include "Cooker/CookDependencyContext.h"
 #include "Cooker/CookEvents.h"
 #include "Serialization/CompactBinarySerialization.h"
 #include "Serialization/CompactBinaryWriter.h"
diff --git a/Engine/Source/Runtime/Engine/Private/Materials/MaterialShader.cpp b/Engine/Source/Runtime/Engine/Private/Materials/MaterialShader.cpp
--- a/Engine/Source/Runtime/Engine/Private/Materials/MaterialShader.cpp
+++ b/Engine/Source/Runtime/Engine/Private/Materials/MaterialShader.cpp
@@ -45,6 +45,7 @@
 #include "Algo/Sort.h"
 #include "Containers/AnsiString.h"
 #include "Cooker/CookDependency.h"
+#include "Cooker/CookDependencyContext.h"
 #include "Cooker/CookEvents.h"
 #include "Interfaces/ITargetPlatform.h"
 #include "MaterialCachedData.h"
diff --git a/Engine/Source/Runtime/GameplayTags/Private/GameplayTagsManager.cpp b/Engine/Source/Runtime/GameplayTags/Private/GameplayTagsManager.cpp
--- a/Engine/Source/Runtime/GameplayTags/Private/GameplayTagsManager.cpp
+++ b/Engine/Source/Runtime/GameplayTags/Private/GameplayTagsManager.cpp
@@ -31,6 +31,7 @@
 #include "ISourceControlModule.h"
 #include "Algo/Sort.h"
 #include "Cooker/CookDependency.h"
+#include "Cooker/CookDependencyContext.h"
 #include "UObject/ICookInfo.h"
 #include "Editor.h"
 #include "PropertyHandle.h"
