Put `OnGCObjectAdded.Broadcast` back under the lock where it belonged.

diff --git a/Engine/Source/Runtime/CoreUObject/Private/Misc/GCObjectReferencer.cpp b/Engine/Source/Runtime/CoreUObject/Private/Misc/GCObjectReferencer.cpp
--- a/Engine/Source/Runtime/CoreUObject/Private/Misc/GCObjectReferencer.cpp
+++ b/Engine/Source/Runtime/CoreUObject/Private/Misc/GCObjectReferencer.cpp
@@ -8,6 +8,7 @@
 #include "CoreMinimal.h"
 #include "HAL/IConsoleManager.h"
 #include "Misc/ScopeLock.h"
+#include "Misc/TransactionallySafeCriticalSection.h"
 #include "Templates/Casts.h"
 #include "UObject/GarbageCollection.h"
 #include "UObject/GarbageCollectionVerification.h"
@@ -30,7 +31,7 @@ namespace UE::GC
 struct UGCObjectReferencer::FImpl
 {
 	/** Critical section used when adding and removing objects */
-	FCriticalSection ReferencedObjectsCritical;
+	FTransactionallySafeCriticalSection ReferencedObjectsCritical;
 	/** Objects without EFlags::AddStableNativeReferencesOnly */
 	TArray<FGCObject*> RemainingReferencedObjects;
 	/** Objects with EFlags::AddStableNativeReferencesOnly */
@@ -103,23 +104,13 @@ void UGCObjectReferencer::AddObject(FGCObject* Object)
 	check(Object);
 	check(!IsUnrealGCMarkingPassInProgress() || IsFrankenGCDestroyingPassInProgress());
 
+	UE::TScopeLock ReferencedObjectsLock(Impl->ReferencedObjectsCritical);
+
 	TArray<FGCObject*>* ReferencedObjects = Object->bCanMakeInitialReferences ? &Impl->InitialReferencedObjects : &Impl->RemainingReferencedObjects;
 
-	UE_AUTORTFM_OPEN
-	{
-		FScopeLock ReferencedObjectsLock(&Impl->ReferencedObjectsCritical);
-		// Make sure there are no duplicates. Should be impossible...
-		checkSlow(!ReferencedObjects->Contains(Object));
-		ReferencedObjects->Add(Object);
-	};
-
-	// Remove the object on abort.
-	UE_AUTORTFM_ONABORT(this, Object, ReferencedObjects)
-	{
-		FScopeLock ReferencedObjectsLock(&Impl->ReferencedObjectsCritical);
-		int32 NumRemoved = ReferencedObjects->RemoveSingleSwap(Object);
-		check(NumRemoved == 1);
-	};
+	// Make sure there are no duplicates. Should be impossible...
+	checkSlow(!ReferencedObjects->Contains(Object));
+	ReferencedObjects->Add(Object);
 
 #if WITH_EDITORONLY_DATA
 	OnGCObjectAdded.Broadcast(Object);
@@ -130,22 +121,12 @@ void UGCObjectReferencer::RemoveObject(FGCObject* Object)
 {
 	check(Object);
 	check(!IsUnrealGCMarkingPassInProgress() || IsFrankenGCDestroyingPassInProgress());
+	
+	UE::TScopeLock ReferencedObjectsLock(Impl->ReferencedObjectsCritical);
 
 	TArray<FGCObject*>* ReferencedObjects = Object->bCanMakeInitialReferences ? &Impl->InitialReferencedObjects : &Impl->RemainingReferencedObjects;
-
-	UE_AUTORTFM_OPEN
-	{
-		FScopeLock ReferencedObjectsLock(&Impl->ReferencedObjectsCritical);
-		int32 NumRemoved = ReferencedObjects->RemoveSingleSwap(Object);
-		check(NumRemoved == 1);
-	};
-
-	// Re-add the object on abort.
-	UE_AUTORTFM_ONABORT(this, Object, ReferencedObjects)
-	{
-		FScopeLock ReferencedObjectsLock(&Impl->ReferencedObjectsCritical);
-		ReferencedObjects->Add(Object);
-	};
+	int32 NumRemoved = ReferencedObjects->RemoveSingleSwap(Object);
+	check(NumRemoved == 1);
 }
 
 bool UGCObjectReferencer::GetReferencerName(UObject* Object, FString& OutName, bool bOnlyIfAddingReferenced) const
