Make zenloader use `TAsyncAtomic` for its atomics, and move all atomic uses over to this type (backed by std::atomic).

This is required for future changes we need to make in zenloader to support all the intracacies of package loading from within transactions.

diff --git a/Engine/Source/Runtime/CoreUObject/Private/Serialization/AsyncLoading2.cpp b/Engine/Source/Runtime/CoreUObject/Private/Serialization/AsyncLoading2.cpp
--- a/Engine/Source/Runtime/CoreUObject/Private/Serialization/AsyncLoading2.cpp
+++ b/Engine/Source/Runtime/CoreUObject/Private/Serialization/AsyncLoading2.cpp
@@ -102,6 +102,8 @@
 
 #include <atomic>
 
+template<typename PayloadType> using TAsyncAtomic = std::atomic<PayloadType>;
+
 // For now, the partial request behavior is reserved for the editor only.
 #define WITH_PARTIAL_REQUEST_DURING_RECURSION WITH_EDITOR
 
@@ -1528,7 +1530,7 @@ class FLoadedPackageCellsRef
 
 private:
 	UE::FMutex Mutex;
-	std::atomic<bool> bPinned = true;
+	TAsyncAtomic<bool> bPinned = true;
 	TMap<uint64, Verse::VCell*> PublicExportMap;
 
 public:
@@ -3046,7 +3048,7 @@ private:
 	uint32 DependenciesCapacity = 0;
 	int32 Priority = 0;
 	int32 ImportOrExportIndex = -1;
-	std::atomic<int32> BarrierCount { 0 };
+	TAsyncAtomic<int32> BarrierCount { 0 };
 	enum EQueueStatus
 	{
 		QueueStatus_None = 0,
@@ -3054,8 +3056,8 @@ private:
 		QueueStatus_External = 2
 	};
 	uint8 QueueStatus = QueueStatus_None;
-	std::atomic<bool> bIsUpdatingDependencies { false };
-	std::atomic<bool> bIsDone { false };
+	TAsyncAtomic<bool> bIsUpdatingDependencies { false };
+	TAsyncAtomic<bool> bIsDone { false };
 };
 
 class FAsyncLoadEventGraphAllocator
@@ -3079,8 +3081,8 @@ public:
 		TotalArcCount -= Count;
 	}
 
-	TAtomic<int64> TotalArcCount { 0 };
-	TAtomic<int64> TotalAllocated { 0 };
+	TAsyncAtomic<int64> TotalArcCount { 0 };
+	TAsyncAtomic<int64> TotalAllocated { 0 };
 };
 
 /** Queue based on DepletableMpscQueue.h 
@@ -3151,7 +3153,7 @@ private:
 	TIoPriorityQueue<FEventLoadNode2> LocalQueue;
 	FCriticalSection ExternalCritical;
 	TIoPriorityQueue<FEventLoadNode2> ExternalQueue;
-	std::atomic<int64> ExternalQueueState = MIN_int64;
+	TAsyncAtomic<int64> ExternalQueueState = MIN_int64;
 	FEventLoadNode2* TimedOutEventNode = nullptr;
 	int32 ExecuteSyncLoadEventsCallCounter = 0;
 };
@@ -3289,7 +3291,7 @@ struct FAsyncLoadingThreadState2
 	bool bCanAccessAsyncLoadingThreadData = false;
 	// used to probe activity for stall detection by the game thread
 	// we use relaxed memory ordering with a simple store to avoid any form of costly interlocked operations.
-	std::atomic<bool> bIsActive = false;
+	TAsyncAtomic<bool> bIsActive = false;
 	bool bShouldFireNodes = true;
 	bool bUseTimeLimit = false;
 	double TimeLimit = 0.0;
@@ -3332,7 +3334,7 @@ struct FAsyncPackageData
 	TArrayView<FEventLoadNode2> ExportBundleNodes;
 	TArrayView<const FSHAHash> ShaderMapHashes;
 	int32 TotalExportBundleCount = 0;
-	std::atomic<bool> bCellExportsInitialized = false;
+	TAsyncAtomic<bool> bCellExportsInitialized = false;
 };
 
 struct FAsyncPackageSerializationState
@@ -3559,10 +3561,10 @@ private:
 	IAsyncReadRequest* UexpReadRequest = nullptr;
 	int64 UassetSize = -1;
 	int64 UexpSize = -1;
-	std::atomic<int8> PendingSizeRequests = 0;
-	std::atomic<int8> PendingReadRequests = 0;
-	std::atomic<bool> bDone = false;
-	std::atomic<bool> bFailed = false;
+	TAsyncAtomic<int8> PendingSizeRequests = 0;
+	TAsyncAtomic<int8> PendingReadRequests = 0;
+	TAsyncAtomic<bool> bDone = false;
+	TAsyncAtomic<bool> bFailed = false;
 	bool bNeedsEngineVersionChecks = false;
 };
 #endif // ALT2_ENABLE_NEW_ARCHIVE_FOR_LINKERLOAD
@@ -3746,12 +3748,12 @@ private:
 	/** Package which is going to have its exports and imports loaded */
 	UPackage* LinkerRoot = nullptr;
 	// The sync load context associated with this package
-	std::atomic<uint64>			SyncLoadContextId = 0;
+	TAsyncAtomic<uint64>			SyncLoadContextId = 0;
 	FAsyncLoadingPostLoadGroup* PostLoadGroup = nullptr;
 	FAsyncLoadingPostLoadGroup* DeferredPostLoadGroup = nullptr;
 	/** Time load begun. This is NOT the time the load was requested in the case of pending requests.	*/
 	double						LoadStartTime = 0.0;
-	std::atomic<int32>			RefCount = 0;
+	TAsyncAtomic<int32>			RefCount = 0;
 	bool						bHasStartedImportingPackages = false;
 	int32						ProcessedExportBundlesCount = 0;
 	/** Current bundle entry index in the current export bundle */
@@ -3763,9 +3765,9 @@ private:
 	/** Current index into Export objects array used to spread routing PostLoadInstances over several frames			*/
 	int32						PostLoadInstanceIndex = 0;
 	/** Current loading state of a package. */
-	std::atomic<EAsyncPackageLoadingState2> AsyncPackageLoadingState { EAsyncPackageLoadingState2::NewPackage };
+	TAsyncAtomic<EAsyncPackageLoadingState2> AsyncPackageLoadingState { EAsyncPackageLoadingState2::NewPackage };
 	/** Whether the constructed objects have been moved to the second loading phase. */
-	std::atomic<bool> ObjectsNowInPhase2 = false;
+	TAsyncAtomic<bool> ObjectsNowInPhase2 = false;
 
 	struct FAllDependenciesState
 	{
@@ -4116,17 +4118,17 @@ public:
 
 	uint64 ContextId;
 	TArray<int32, TInlineAllocator<4>> RequestIDs;
-	std::atomic<bool> bHasFoundRequestedPackages{ false };
+	TAsyncAtomic<bool> bHasFoundRequestedPackages{ false };
 	UE::FManualResetEvent DoneEvent;
 private:
 	TArray<TRefCountPtr<FAsyncPackage2>, TInlineAllocator<4>> RequestedPackages;
 	TRefCountPtr<FAsyncPackage2> RequestingPackage = nullptr;
-	std::atomic<int32> RefCount = 1;
+	TAsyncAtomic<int32> RefCount = 1;
 
-	static std::atomic<uint64> NextContextId;
+	static TAsyncAtomic<uint64> NextContextId;
 };
 
-std::atomic<uint64> FAsyncLoadingSyncLoadContext::NextContextId{ 1 };
+TAsyncAtomic<uint64> FAsyncLoadingSyncLoadContext::NextContextId{ 1 };
 
 struct FAsyncLoadingPostLoadGroup
 {
@@ -4155,11 +4157,11 @@ private:
 	/** Thread to run the worker FRunnable on */
 	static constexpr int32 DefaultAsyncPackagesReserveCount = 512;
 	FRunnableThread* Thread;
-	std::atomic<bool> bStopRequested = false;
-	std::atomic<int32> SuspendRequestedCount = 0;
+	TAsyncAtomic<bool> bStopRequested = false;
+	TAsyncAtomic<int32> SuspendRequestedCount = 0;
 	bool bHasRegisteredAllScriptObjects = false;
 	/** [ASYNC/GAME THREAD] true if the async thread is actually started. We don't start it until after we boot because the boot process on the game thread can create objects that are also being created by the loader */
-	std::atomic<bool> bThreadStarted = false;
+	TAsyncAtomic<bool> bThreadStarted = false;
 
 #if !UE_BUILD_SHIPPING
 	FPlatformFileOpenLog* FileOpenLogWrapper = nullptr;
@@ -4274,7 +4276,7 @@ private:
 	TStripedMap<32, FPackageId, FAsyncPackage2*> AsyncPackageLookup;
 
 	TMpscQueue<FAsyncPackage2*> ExternalReadQueue;
-	TAtomic<int32> PendingIoRequestsCounter{ 0 };
+	TAsyncAtomic<int32> PendingIoRequestsCounter{ 0 };
 	
 	/** List of all pending package requests */
 	TSet<int32> PendingRequests;
@@ -4283,14 +4285,14 @@ private:
 	TMap<int32, FAsyncPackage2*> RequestIdToPackageMap; // Only accessed from the async loading thread
 
 	/** [ASYNC/GAME THREAD] Number of package load requests in the async loading queue */
-	TAtomic<int32> QueuedPackagesCounter { 0 };
+	TAsyncAtomic<int32> QueuedPackagesCounter { 0 };
 	/** [ASYNC/GAME THREAD] Number of packages being loaded on the async thread and post loaded on the game thread */
-	TAtomic<int32> LoadingPackagesCounter { 0 };
+	TAsyncAtomic<int32> LoadingPackagesCounter { 0 };
 	/** Encapsulate our counter to make sure all decrements are monitored. */
 	class FPackagesWithRemainingWorkCounter
 	{
 		UE::FManualResetEvent* WakeEvent = nullptr;
-		TAtomic<int32> PackagesWithRemainingWorkCounter {0};
+		TAsyncAtomic<int32> PackagesWithRemainingWorkCounter {0};
 	public:
 		void SetWakeEvent(UE::FManualResetEvent* InWakeEvent) {	WakeEvent = InWakeEvent; }
 		int32 operator++() { return ++PackagesWithRemainingWorkCounter; }
@@ -4314,7 +4316,7 @@ private:
 	/** [ASYNC/GAME THREAD] While this is non-zero there's work left to do */
 	FPackagesWithRemainingWorkCounter PackagesWithRemainingWorkCounter;
 
-	FThreadSafeCounter AsyncThreadReady;
+	TAsyncAtomic<int32> AsyncThreadReady;
 
 	/** When cancelling async loading: list of package requests to cancel */
 	TArray<FAsyncPackageDesc2*> QueuedPackagesToCancel;
@@ -4322,7 +4324,7 @@ private:
 	TSet<FAsyncPackage2*> PackagesToCancel;
 
 	/** Async loading thread ID */
-	std::atomic<uint32> AsyncLoadingThreadID;
+	TAsyncAtomic<uint32> AsyncLoadingThreadID;
 
 	/** I/O Dispatcher */
 	FIoDispatcher& IoDispatcher;
@@ -5218,7 +5220,7 @@ void FAsyncLoadingThread2::InitializeLoading()
 		AltZenaphore.NotifyOne();
 	});
 	
-	AsyncThreadReady.Increment();
+	AsyncThreadReady.fetch_add(1);
 
 	UE_LOG(LogStreaming, Log, TEXT("AsyncLoading2 - Initialized"));
 }
@@ -6931,7 +6933,7 @@ void FAsyncPackage2::StartLoading(FAsyncLoadingThreadState2& ThreadState, FIoBat
 #if WITH_EDITOR
 	if (OptionalSegmentHeaderData.IsSet())
 	{
-		int32 LocalPendingIoRequestsCounter = AsyncLoadingThread.PendingIoRequestsCounter.IncrementExchange() + 1;
+		int32 LocalPendingIoRequestsCounter = AsyncLoadingThread.PendingIoRequestsCounter.fetch_add(1) + 1;
 		TRACE_COUNTER_SET(AsyncLoadingPendingIoRequests, LocalPendingIoRequestsCounter);
 
 		GetPackageNode(EEventLoadNode2::Package_ProcessSummary).AddBarrier();
@@ -6946,7 +6948,7 @@ void FAsyncPackage2::StartLoading(FAsyncLoadingThreadState2& ThreadState, FIoBat
 						TEXT("Failed reading optional chunk for package: %s"), *Result.Status().ToString());
 					LoadStatus = UE::AsyncPackage2::AsyncLoadingResultFromIoError(Result.Status().GetErrorCode());
 				}
-				int32 LocalPendingIoRequestsCounter = AsyncLoadingThread.PendingIoRequestsCounter.DecrementExchange() - 1;
+				int32 LocalPendingIoRequestsCounter = AsyncLoadingThread.PendingIoRequestsCounter.fetch_sub(1) - 1;
 				TRACE_COUNTER_SET(AsyncLoadingPendingIoRequests, LocalPendingIoRequestsCounter);
 				FAsyncLoadingThread2& LocalAsyncLoadingThread = AsyncLoadingThread;
 				GetPackageNode(EEventLoadNode2::Package_ProcessSummary).ReleaseBarrier(nullptr);
@@ -6958,7 +6960,7 @@ void FAsyncPackage2::StartLoading(FAsyncLoadingThreadState2& ThreadState, FIoBat
 	}
 #endif
 
-	int32 LocalPendingIoRequestsCounter = AsyncLoadingThread.PendingIoRequestsCounter.IncrementExchange() + 1;
+	int32 LocalPendingIoRequestsCounter = AsyncLoadingThread.PendingIoRequestsCounter.fetch_add(1) + 1;
 	TRACE_COUNTER_SET(AsyncLoadingPendingIoRequests, LocalPendingIoRequestsCounter);
 
 
@@ -6984,7 +6986,7 @@ void FAsyncPackage2::StartLoading(FAsyncLoadingThreadState2& ThreadState, FIoBat
 					TEXT("Failed reading chunk for package: %s"), *Result.Status().ToString());
 				LoadStatus = UE::AsyncPackage2::AsyncLoadingResultFromIoError(Result.Status().GetErrorCode());
 			}
-			int32 LocalPendingIoRequestsCounter = AsyncLoadingThread.PendingIoRequestsCounter.DecrementExchange() - 1;
+			int32 LocalPendingIoRequestsCounter = AsyncLoadingThread.PendingIoRequestsCounter.fetch_sub(1) - 1;
 			TRACE_COUNTER_SET(AsyncLoadingPendingIoRequests, LocalPendingIoRequestsCounter);
 			FAsyncLoadingThread2& LocalAsyncLoadingThread = AsyncLoadingThread;
 			GetPackageNode(EEventLoadNode2::Package_ProcessSummary).ReleaseBarrier(nullptr);
@@ -7000,13 +7002,13 @@ void FAsyncPackage2::StartLoading(FAsyncLoadingThreadState2& ThreadState, FIoBat
 		auto ReadShaderMapFunc = [this, &IoBatch](const FIoChunkId& ChunkId, FGraphEventRef GraphEvent)
 		{
 			GetPackageNode(Package_ExportsSerialized).AddBarrier();
-			int32 LocalPendingIoRequestsCounter = AsyncLoadingThread.PendingIoRequestsCounter.IncrementExchange() + 1;
+			int32 LocalPendingIoRequestsCounter = AsyncLoadingThread.PendingIoRequestsCounter.fetch_add(1) + 1;
 			TRACE_COUNTER_SET(AsyncLoadingPendingIoRequests, LocalPendingIoRequestsCounter);
 			return IoBatch.ReadWithCallback(ChunkId, FIoReadOptions(), Desc.Priority,
 				[this, GraphEvent](TIoStatusOr<FIoBuffer> Result)
 				{
 					GraphEvent->DispatchSubsequents();
-					int32 LocalPendingIoRequestsCounter = AsyncLoadingThread.PendingIoRequestsCounter.DecrementExchange() - 1;
+					int32 LocalPendingIoRequestsCounter = AsyncLoadingThread.PendingIoRequestsCounter.fetch_sub(1) - 1;
 					TRACE_COUNTER_SET(AsyncLoadingPendingIoRequests, LocalPendingIoRequestsCounter);
 					FAsyncLoadingThread2& LocalAsyncLoadingThread = AsyncLoadingThread;
 					GetPackageNode(Package_ExportsSerialized).ReleaseBarrier(nullptr);
@@ -10518,7 +10520,7 @@ uint32 FAsyncLoadingThread2::Run()
 					}
 				} // release FGCScopeGuard
 
-				if (PendingIoRequestsCounter.Load() > 0)
+				if (PendingIoRequestsCounter.load() > 0)
 				{
 					TRACE_CPUPROFILER_EVENT_SCOPE(WaitingForIo);
 					Waiter.Wait();
@@ -10550,7 +10552,7 @@ EAsyncPackageState::Type FAsyncLoadingThread2::TickAsyncThreadFromGameThread(FAs
 	check(IsInGameThread());
 	EAsyncPackageState::Type Result = EAsyncPackageState::Complete;
 	
-	if (AsyncThreadReady.GetValue())
+	if (AsyncThreadReady)
 	{
 		if (ThreadState.IsTimeLimitExceeded(TEXT("TickAsyncThreadFromGameThread")))
 		{
@@ -11905,7 +11907,7 @@ void FAsyncLoadingThread2::FlushLoading(TConstArrayView<int32> RequestIDs)
 							AsyncLoadingThreadState->ResetActivity();
 						}
 					}
-					else if (GStallDetectorTimeout != 0.0f && PendingIoRequestsCounter.Load() == 0 && FPlatformTime::Seconds() - LastActivity > GStallDetectorTimeout && ++IdleLoopCount > GStallDetectorIdleLoops)
+					else if (GStallDetectorTimeout != 0.0f && PendingIoRequestsCounter.load() == 0 && FPlatformTime::Seconds() - LastActivity > GStallDetectorTimeout && ++IdleLoopCount > GStallDetectorIdleLoops)
 					{
 #if ALT2_DUMP_STATE_ON_HANG
 						{
