FRedirectCollector::GetAssetPathRedirection optimization: compute the string hash before acquiring the lock.

diff --git a/Engine/Source/Runtime/CoreUObject/Private/Misc/RedirectCollector.cpp b/Engine/Source/Runtime/CoreUObject/Private/Misc/RedirectCollector.cpp
--- a/Engine/Source/Runtime/CoreUObject/Private/Misc/RedirectCollector.cpp
+++ b/Engine/Source/Runtime/CoreUObject/Private/Misc/RedirectCollector.cpp
@@ -470,10 +470,12 @@ void FRedirectCollector::RemoveAssetPathRedirection(const FSoftObjectPath& Origi
 
 FSoftObjectPath FRedirectCollector::GetAssetPathRedirection(const FSoftObjectPath& OriginalPath) const
 {
+	const uint32 OriginalPathHash = GetTypeHash(OriginalPath);
+
 	FScopeLock ScopeLock(&CriticalSection);
 
 	const FSimpleOrChainedRedirect* Redirection;
-	if (Redirection = ObjectPathRedirectionMap.Find(OriginalPath); Redirection)
+	if (Redirection = ObjectPathRedirectionMap.FindByHash(OriginalPathHash, OriginalPath); Redirection)
 	{
 		return Redirection->GetFinalTarget();
 	}
