Make the new 'obj remote' command return true instead of falling back to false
Fixed deprecation message.

diff --git a/Engine/Source/Runtime/CoreUObject/Public/UObject/RemoteObjectTypes.h b/Engine/Source/Runtime/CoreUObject/Public/UObject/RemoteObjectTypes.h
--- a/Engine/Source/Runtime/CoreUObject/Public/UObject/RemoteObjectTypes.h
+++ b/Engine/Source/Runtime/CoreUObject/Public/UObject/RemoteObjectTypes.h
@@ -69,7 +69,7 @@ struct FRemoteServerId
 		*this = FromIdNumber(InId);
 	}
 
-	UE_DEPRECATED(5.6, "Use FRemoteServerId::FromString(const FString&) instead.")
+	UE_DEPRECATED(5.6, "Use FRemoteServerId::FromString(const FStringView&) instead.")
 	explicit COREUOBJECT_API FRemoteServerId(const FString& InText);
 
 	static COREUOBJECT_API FRemoteServerId FromString(const FStringView& InText);
diff --git a/Engine/Source/Runtime/Engine/Private/UnrealEngine.cpp b/Engine/Source/Runtime/Engine/Private/UnrealEngine.cpp
--- a/Engine/Source/Runtime/Engine/Private/UnrealEngine.cpp
+++ b/Engine/Source/Runtime/Engine/Private/UnrealEngine.cpp
@@ -10210,7 +10210,7 @@ bool UEngine::HandleObjCommand( const TCHAR* Cmd, FOutputDevice& Ar )
 		{
 			Ar.Logf(TEXT("Unable to parse remote object id '%s'"), *IdText);
 		}
-
+		return true;
 	}
 #endif // UE_WITH_REMOTE_OBJECT_HANDLE
 #endif // !UE_BUILD_SHIPPING
