[Backout] - CL45088104
#rnx
Original CL Desc
-----------------------------------------------------------------
[AsyncLoading2]

- Increase the refcount for AsyncPackages stored in a SyncLoadContext and release them when the context is destroyed.

This fixes an issue where the game thread may be flushing a set of request ids, the ALT processes some but not all and allows those completed packages to be cleaned up. If a package in the set of packages being flushed causes a loader stall (e.g. circular dependency) then the loader's dumpstate call will read from now deleted FAsyncPackages in the SyncLoadContext for the packages being flushed.

#rnx

diff --git a/Engine/Source/Runtime/CoreUObject/Private/Serialization/AsyncLoading2.cpp b/Engine/Source/Runtime/CoreUObject/Private/Serialization/AsyncLoading2.cpp
--- a/Engine/Source/Runtime/CoreUObject/Private/Serialization/AsyncLoading2.cpp
+++ b/Engine/Source/Runtime/CoreUObject/Private/Serialization/AsyncLoading2.cpp
@@ -2994,7 +2994,6 @@ class FEventLoadGraphAllocator;
 struct FAsyncLoadEventSpec;
 struct FAsyncLoadingThreadState2;
 class FAsyncLoadEventQueue2;
-class FAsyncLoadingSyncLoadContext;
 
 enum class EEventLoadNodeExecutionResult : uint8
 {
@@ -3174,6 +3173,53 @@ struct FAsyncLoadEventSpec
 	const TCHAR* Name = nullptr;
 };
 
+class FAsyncLoadingSyncLoadContext
+{
+public:
+	FAsyncLoadingSyncLoadContext(TConstArrayView<int32> InRequestIDs)
+		: RequestIDs(InRequestIDs)
+	{
+		RequestedPackages.AddZeroed(RequestIDs.Num());
+		ContextId = NextContextId++;
+		if (NextContextId == 0)
+		{
+			NextContextId = 1;
+		}
+	}
+
+	~FAsyncLoadingSyncLoadContext()
+	{
+	}
+
+	void AddRef()
+	{
+		++RefCount;
+	}
+
+	void ReleaseRef()
+	{
+		int32 NewRefCount = --RefCount;
+		check(NewRefCount >= 0);
+		if (NewRefCount == 0)
+		{
+			delete this;
+		}
+	}
+
+	uint64 ContextId;
+	TArray<int32, TInlineAllocator<4>> RequestIDs;
+	TArray<FAsyncPackage2*, TInlineAllocator<4>> RequestedPackages;
+	FAsyncPackage2* RequestingPackage = nullptr;
+	TAsyncAtomic<bool> bHasFoundRequestedPackages { false };
+	UE::FManualResetEvent DoneEvent;
+private:
+	TAsyncAtomic<int32> RefCount = 1;
+
+	static TAsyncAtomic<uint64> NextContextId;
+};
+
+TAsyncAtomic<uint64> FAsyncLoadingSyncLoadContext::NextContextId { 1 };
+
 struct FAsyncLoadingThreadState2
 {
 	struct FTimeLimitScope
@@ -3575,7 +3621,7 @@ struct FAsyncLoadingPostLoadGroup;
 * Structure containing intermediate data required for async loading of all exports of a package.
 */
 
-struct FAsyncPackage2 : public FRefCountBase
+struct FAsyncPackage2
 {
 	friend struct FAsyncPackageScope2;
 	friend class FAsyncLoadingThread2;
@@ -3612,7 +3658,7 @@ struct FAsyncPackage2 : public FRefCountBase
 		}
 	}
 
-	void Release();
+	void ReleaseRef();
 
 	void ClearImportedPackages();
 
@@ -3863,7 +3909,7 @@ private:
 			PackageToRemoveState.NextLink = nullptr;
 			FAsyncPackage2* WaitingForPackage = PackageToRemoveState.WaitingForPackage;
 			PackageToRemoveState.WaitingForPackage = nullptr;
-			WaitingForPackage->Release();
+			WaitingForPackage->ReleaseRef();
 		}
 	};
 	FAllDependenciesState		AllDependenciesSetupState;
@@ -4054,82 +4100,6 @@ public:
 	FUObjectSerializeContext* GetSerializeContext();
 };
 
-class FAsyncLoadingSyncLoadContext
-{
-public:
-	FAsyncLoadingSyncLoadContext(TConstArrayView<int32> InRequestIDs, FAsyncPackage2* CurrentlyExecutingPackage)
-		: RequestIDs(InRequestIDs)
-		, RequestingPackage(CurrentlyExecutingPackage)
-	{
-		RequestedPackages.AddZeroed(RequestIDs.Num());
-		ContextId = NextContextId++;
-		if (NextContextId == 0)
-		{
-			NextContextId = 1;
-		}
-	}
-
-	~FAsyncLoadingSyncLoadContext() = default;
-
-	void AddRef()
-	{
-		++RefCount;
-	}
-
-	void ReleaseRef()
-	{
-		int32 NewRefCount = --RefCount;
-		check(NewRefCount >= 0);
-		if (NewRefCount == 0)
-		{
-			delete this;
-		}
-	}
-
-	TRefCountPtr<FAsyncPackage2>& GetRequestingPackage()
-	{
-		return RequestingPackage;
-	}
-
-	TArray<TRefCountPtr<FAsyncPackage2>, TInlineAllocator<4>>& GetRequestedPackages()
-	{
-		return RequestedPackages;
-	}
-
-	TRefCountPtr<FAsyncPackage2>& GetRequestedPackage(int32 Index)
-	{
-		return RequestedPackages[Index];
-	}
-
-	void SetRequestedPackage(int32 Index, FAsyncPackage2* RequestedPackage)
-	{
-		RequestedPackages[Index] = RequestedPackage;
-	}
-
-	bool IsRequestedPackageValid(int32 Index) const
-	{
-		return RequestedPackages[Index].IsValid();
-	}
-
-	int32 FindRequestedPackage(FAsyncPackage2* Package)
-	{
-		return RequestedPackages.Find(Package);
-	}
-
-	uint64 ContextId;
-	TArray<int32, TInlineAllocator<4>> RequestIDs;
-	TAsyncAtomic<bool> bHasFoundRequestedPackages{ false };
-	UE::FManualResetEvent DoneEvent;
-private:
-	TArray<TRefCountPtr<FAsyncPackage2>, TInlineAllocator<4>> RequestedPackages;
-	TRefCountPtr<FAsyncPackage2> RequestingPackage = nullptr;
-	TAsyncAtomic<int32> RefCount = 1;
-
-	static TAsyncAtomic<uint64> NextContextId;
-};
-
-TAsyncAtomic<uint64> FAsyncLoadingSyncLoadContext::NextContextId{ 1 };
-
 struct FAsyncLoadingPostLoadGroup
 {
 	uint64 SyncLoadContextId = 0;
@@ -8976,13 +8946,13 @@ void FAsyncPackage2::ConditionalReleasePartialRequests(FAsyncLoadingThreadState2
 	// This means the request we're trying to satisfy is not necessarily on top of the sync context stack anymore.
 	for (FAsyncLoadingSyncLoadContext* SyncLoadContext : ThreadState.SyncLoadContextStack)
 	{
-		if (TRefCountPtr<FAsyncPackage2>& RequestingPackage = SyncLoadContext->GetRequestingPackage())
+		if (SyncLoadContext->RequestingPackage)
 		{
-			EAsyncPackageLoadingState2 RequesterState = RequestingPackage->AsyncPackageLoadingState;
+			EAsyncPackageLoadingState2 RequesterState = SyncLoadContext->RequestingPackage->AsyncPackageLoadingState;
 
 			if (AsyncPackageLoadingState > RequesterState)
 			{
-				int32 Index = SyncLoadContext->FindRequestedPackage(this);
+				int32 Index = SyncLoadContext->RequestedPackages.Find(this);
 				if (Index != INDEX_NONE)
 				{
 					int32 RequestId = SyncLoadContext->RequestIDs[Index];
@@ -9004,13 +8974,13 @@ void FAsyncPackage2::ConditionalReleasePartialRequests(FAsyncLoadingThreadState2
 						// sync load was released to avoid a deadlock. This will make sure that by the time the requesting package
 						// finishes loading, the requested package will be done too.
 						bool bIsAlreadyInSet = false;
-						RequestingPackage->AdditionalImportedAsyncPackages.FindOrAdd(this, &bIsAlreadyInSet);
+						SyncLoadContext->RequestingPackage->AdditionalImportedAsyncPackages.FindOrAdd(this, &bIsAlreadyInSet);
 						if (!bIsAlreadyInSet)
 						{
 							this->AddRef();
 
 							UE_LOG(LogStreaming, Display, TEXT("Package %s is adding a dynamic import to package %s because of a recursive sync load"),
-								*RequestingPackage->Desc.UPackageName.ToString(),
+								*SyncLoadContext->RequestingPackage->Desc.UPackageName.ToString(),
 								*Desc.UPackageName.ToString()
 							);
 						}
@@ -9712,7 +9682,7 @@ void FAsyncLoadingThread2::UpdateSyncLoadContext(FAsyncLoadingThreadState2& Thre
 		for (int32 i=0; i < SyncLoadContext->RequestIDs.Num(); ++i)
 		{
 			int32 RequestID = SyncLoadContext->RequestIDs[i];
-			if (SyncLoadContext->IsRequestedPackageValid(i))
+			if (SyncLoadContext->RequestedPackages[i] != nullptr)
 			{
 				++FoundPackages;
 			}
@@ -9720,10 +9690,11 @@ void FAsyncLoadingThread2::UpdateSyncLoadContext(FAsyncLoadingThreadState2& Thre
 			{
 				// Set RequestedPackage before setting bHasFoundRequestedPackage so that another thread looking at RequestedPackage
 				// after validating that bHasFoundRequestedPackage is true would see the proper value.
-				SyncLoadContext->SetRequestedPackage(i, RequestedPackage);
+				SyncLoadContext->RequestedPackages[i] = RequestedPackage;
 
 #if WITH_PARTIAL_REQUEST_DURING_RECURSION
-				if (TRefCountPtr<FAsyncPackage2>& RequestingPackage = SyncLoadContext->GetRequestingPackage())
+				FAsyncPackage2* RequestingPackage = SyncLoadContext->RequestingPackage;
+				if (RequestingPackage != nullptr)
 				{
 					// If the flush is coming from a step before the requesting package is back on GT, there is no way to fully flush
 					// the requested package unless its already done. We have no choice but to trigger partial loading in that case.
@@ -9764,7 +9735,7 @@ void FAsyncLoadingThread2::UpdateSyncLoadContext(FAsyncLoadingThreadState2& Thre
 		for (int32 i = 0; i < SyncLoadContext->RequestIDs.Num(); ++i)
 		{
 			int32 RequestID = SyncLoadContext->RequestIDs[i];
-			TRefCountPtr<FAsyncPackage2>& RequestedPackage = SyncLoadContext->GetRequestedPackage(i);
+			FAsyncPackage2* RequestedPackage = SyncLoadContext->RequestedPackages[i];
 			if (RequestedPackage && ThreadState.PackagesOnStack.Contains(RequestedPackage))
 			{
 				// Flushing a package while it's already being processed on the stack, if we're done preloading we let it pass and remove the request id
@@ -9903,7 +9874,7 @@ EAsyncPackageState::Type FAsyncLoadingThread2::ProcessLoadedPackagesFromGameThre
 		while (ThreadState.PackagesToReprioritize.Dequeue(PackageToRepriortize))
 		{
 			MainThreadEventQueue.UpdatePackagePriority(PackageToRepriortize);
-			PackageToRepriortize->Release();
+			PackageToRepriortize->ReleaseRef();
 		}
 		uint64 SyncLoadContextId = !ThreadState.SyncLoadContextStack.IsEmpty() ? ThreadState.SyncLoadContextStack.Top()->ContextId : 0;
 		if (SyncLoadContextId)
@@ -10065,7 +10036,7 @@ EAsyncPackageState::Type FAsyncLoadingThread2::ProcessLoadedPackagesFromGameThre
 			if (Package->CompletionCallbacks.IsEmpty() && Package->ProgressCallbacks.IsEmpty())
 			{
 				RemovePendingRequests(ThreadState, Package->RequestIDs);
-				Package->Release();
+				Package->ReleaseRef();
 			}
 			else
 			{
@@ -10100,7 +10071,7 @@ EAsyncPackageState::Type FAsyncLoadingThread2::ProcessLoadedPackagesFromGameThre
 			CompletedPackageRequest.CallCompletionCallbacks();
 			if (CompletedPackageRequest.AsyncPackage)
 			{
-				CompletedPackageRequest.AsyncPackage->Release();
+				CompletedPackageRequest.AsyncPackage->ReleaseRef();
 			}
 			else
 			{
@@ -11094,7 +11065,7 @@ FAsyncPackage2::~FAsyncPackage2()
 	check(DeferredPostLoadGroup == nullptr);
 }
 
-void FAsyncPackage2::Release()
+void FAsyncPackage2::ReleaseRef()
 {
 	int32 OldRefCount = RefCount.fetch_sub(1);
 	check(OldRefCount > 0);
@@ -11110,7 +11081,7 @@ void FAsyncPackage2::ClearImportedPackages()
 {
 	TRACE_CPUPROFILER_EVENT_SCOPE(ClearImportedPackages);
 	// Reset the ImportedAsyncPackages member array before releasing the async package references.
-	// Release may queue up the ImportedAsyncPackage for deletion on the ALT.
+	// ReleaseRef may queue up the ImportedAsyncPackage for deletion on the ALT.
 	// ImportedAsyncPackages of this package can still be accessed by both the GT and the ALT.
 	TArrayView<FAsyncPackage2*> LocalImportedAsyncPackages = Data.ImportedAsyncPackages;
 	Data.ImportedAsyncPackages = MakeArrayView(Data.ImportedAsyncPackages.GetData(), 0);
@@ -11118,7 +11089,7 @@ void FAsyncPackage2::ClearImportedPackages()
 	{
 		if (ImportedAsyncPackage)
 		{
-			ImportedAsyncPackage->Release();
+			ImportedAsyncPackage->ReleaseRef();
 		}
 	}
 
@@ -11126,7 +11097,7 @@ void FAsyncPackage2::ClearImportedPackages()
 	TSet<FAsyncPackage2*> LocalAdditionalImportedAsyncPackages = MoveTemp(AdditionalImportedAsyncPackages);
 	for (FAsyncPackage2* ImportedAsyncPackage : LocalAdditionalImportedAsyncPackages)
 	{
-		ImportedAsyncPackage->Release();
+		ImportedAsyncPackage->ReleaseRef();
 	}
 #endif
 }
@@ -11662,11 +11633,11 @@ void FAsyncLoadingThread2::WarnAboutPotentialSyncLoadStall(FAsyncLoadingSyncLoad
 		int32 RequestId = SyncLoadContext->RequestIDs[Index];
 		if (ContainsRequestID(RequestId))
 		{
-			TRefCountPtr<FAsyncPackage2>& Package = SyncLoadContext->GetRequestedPackage(Index);
+			FAsyncPackage2* Package = SyncLoadContext->RequestedPackages[Index];
 			UE_LOG(LogStreaming, Warning, TEXT("A flush request appear to be stuck waiting on package %s at state %s to reach state > %s"),
 				*Package->Desc.UPackageName.ToString(),
 				LexToString(Package->AsyncPackageLoadingState),
-				LexToString(SyncLoadContext->GetRequestingPackage()->AsyncPackageLoadingState)
+				LexToString(SyncLoadContext->RequestingPackage->AsyncPackageLoadingState)
 			);
 		}
 	}
@@ -11681,7 +11652,10 @@ void FAsyncLoadingThread2::FlushLoadingFromParallelLoadingThread(FAsyncLoadingTh
 		return;
 	}
 
-	FAsyncLoadingSyncLoadContext* SyncLoadContext = new FAsyncLoadingSyncLoadContext(RequestIDs, FAsyncPackage2::GetCurrentlyExecutingPackage(ThreadState));
+	FAsyncLoadingSyncLoadContext* SyncLoadContext = new FAsyncLoadingSyncLoadContext(RequestIDs);
+
+	FAsyncPackage2* GetCurrentlyExecutingPackage = FAsyncPackage2::GetCurrentlyExecutingPackage(ThreadState);
+	SyncLoadContext->RequestingPackage = GetCurrentlyExecutingPackage;
 
 	UE_LOG(LogStreaming, VeryVerbose, TEXT("Pushing ALT Worker SyncLoadContext %" UINT64_FMT), SyncLoadContext->ContextId);
 	ThreadState.SyncLoadContextStack.Push(SyncLoadContext);
@@ -11710,7 +11684,10 @@ void FAsyncLoadingThread2::FlushLoadingFromLoadingThread(FAsyncLoadingThreadStat
 		return;
 	}
 
-	FAsyncLoadingSyncLoadContext* SyncLoadContext = new FAsyncLoadingSyncLoadContext(RequestIDs, FAsyncPackage2::GetCurrentlyExecutingPackage(ThreadState));
+	FAsyncLoadingSyncLoadContext* SyncLoadContext = new FAsyncLoadingSyncLoadContext(RequestIDs);
+
+	FAsyncPackage2* GetCurrentlyExecutingPackage = FAsyncPackage2::GetCurrentlyExecutingPackage(ThreadState);
+	SyncLoadContext->RequestingPackage = GetCurrentlyExecutingPackage;
 
 	UE_LOG(LogStreaming, VeryVerbose, TEXT("Pushing ALT SyncLoadContext %" UINT64_FMT), SyncLoadContext->ContextId);
 	AsyncLoadingThreadState->SyncLoadContextStack.Push(SyncLoadContext);
@@ -11859,7 +11836,8 @@ void FAsyncLoadingThread2::FlushLoading(TConstArrayView<int32> RequestIDs)
 		FAsyncLoadingSyncLoadContext* SyncLoadContext = nullptr;
 		if (RequestIDs.Num() != 0 && GOnlyProcessRequiredPackagesWhenSyncLoading)
 		{
-			SyncLoadContext = new FAsyncLoadingSyncLoadContext(RequestIDs, CurrentlyExecutingPackage);
+			SyncLoadContext = new FAsyncLoadingSyncLoadContext(RequestIDs);
+			SyncLoadContext->RequestingPackage = CurrentlyExecutingPackage;
 
 			UE_LOG(LogStreaming, VeryVerbose, TEXT("Pushing GT SyncLoadContext %" UINT64_FMT), SyncLoadContext->ContextId);
 
@@ -11926,7 +11904,7 @@ void FAsyncLoadingThread2::FlushLoading(TConstArrayView<int32> RequestIDs)
 
 							if (SyncLoadContext)
 							{
-								for (TRefCountPtr<FAsyncPackage2>& Package : SyncLoadContext->GetRequestedPackages())
+								for (FAsyncPackage2* Package : SyncLoadContext->RequestedPackages)
 								{
 									if (Package)
 									{
