Quick logging change

#rnx

diff --git a/Engine/Source/Runtime/CoreUObject/Private/UObject/RemoteExecutor.cpp b/Engine/Source/Runtime/CoreUObject/Private/UObject/RemoteExecutor.cpp
--- a/Engine/Source/Runtime/CoreUObject/Private/UObject/RemoteExecutor.cpp
+++ b/Engine/Source/Runtime/CoreUObject/Private/UObject/RemoteExecutor.cpp
@@ -947,8 +947,8 @@ FRemoteWorkPriority CreateRootWorkPriority()
 
 void ExecutePendingWork()
 {
-	static const FName Name = "ExecuteTransactional";
-	UE::RemoteExecutor::OnRegionBeginDelegate.ExecuteIfBound(Name);
+	static const FName FuncName{ __func__ };
+	UE::RemoteExecutor::OnRegionBeginDelegate.ExecuteIfBound(FuncName);
 	GRemoteExecutor.ExecutePendingWork();
 	UE::RemoteExecutor::OnRegionEndDelegate.ExecuteIfBound(FString());
 }
