Fix possible regression from transformation of the property key

#rnx

diff --git a/Engine/Source/Runtime/CoreUObject/Private/UObject/OverriddenPropertySet.cpp b/Engine/Source/Runtime/CoreUObject/Private/UObject/OverriddenPropertySet.cpp
--- a/Engine/Source/Runtime/CoreUObject/Private/UObject/OverriddenPropertySet.cpp
+++ b/Engine/Source/Runtime/CoreUObject/Private/UObject/OverriddenPropertySet.cpp
@@ -305,9 +305,10 @@ void FOverriddenPropertyPath::operator+=(const FName Name)
 	CachedHash = HashCombine(CachedHash, GetTypeHash(Name));
 }
 
-void FOverriddenPropertyPath::operator+=(const FOverriddenPropertyPath& Other)
+void FOverriddenPropertyPath::operator+=(const TArray<FName>& Names)
 {
-	for(FName Name : Other.Path)
+	Path.Reserve(Path.Num() + Names.Num());
+	for (const FName& Name : Names)
 	{
 		this->operator+=(Name);
 	}
@@ -333,9 +334,8 @@ FOverriddenPropertyNodeID::FOverriddenPropertyNodeID(TNotNull<const FProperty*>
         Property->SaveTypeName(TypeNameBuilder);
 	}
 #endif
-	UE::FPropertyTypeName TypeName = TypeNameBuilder.Build();
 	Path = Property->GetFName();
-	Path += TypeName.GetName();
+	Path += TypeNameBuilder.ToArray();
 }
 
 FOverriddenPropertyNodeID::FOverriddenPropertyNodeID(const FOverriddenPropertyNodeID& ParentNodeID, const FOverriddenPropertyNodeID& NodeID)
diff --git a/Engine/Source/Runtime/CoreUObject/Private/UObject/PropertyTypeName.cpp b/Engine/Source/Runtime/CoreUObject/Private/UObject/PropertyTypeName.cpp
--- a/Engine/Source/Runtime/CoreUObject/Private/UObject/PropertyTypeName.cpp
+++ b/Engine/Source/Runtime/CoreUObject/Private/UObject/PropertyTypeName.cpp
@@ -619,4 +619,14 @@ void FPropertyTypeNameBuilder::Reset()
 	ActiveIndex = INDEX_NONE;
 }
 
+TArray<FName> FPropertyTypeNameBuilder::ToArray() const
+{
+	TArray<FName> Names;
+	Names.Reserve(Nodes.Num());
+	for( const FPropertyTypeNameNode& Node : Nodes)
+	{
+		Names.Add(Node.Name);
+	}
+	return MoveTemp(Names);
+}
 } // UE
diff --git a/Engine/Source/Runtime/CoreUObject/Public/UObject/OverriddenPropertySet.h b/Engine/Source/Runtime/CoreUObject/Public/UObject/OverriddenPropertySet.h
--- a/Engine/Source/Runtime/CoreUObject/Public/UObject/OverriddenPropertySet.h
+++ b/Engine/Source/Runtime/CoreUObject/Public/UObject/OverriddenPropertySet.h
@@ -223,7 +223,11 @@ struct FOverriddenPropertyPath
 	FOverriddenPropertyPath operator+(const FName Name) const;
 	FOverriddenPropertyPath operator+(const FOverriddenPropertyPath& Other) const;
 	void operator+=(const FName Name);
-	void operator+=(const FOverriddenPropertyPath& Other);
+	void operator+=(const TArray<FName>& Names);
+	void operator+=(const FOverriddenPropertyPath& Other)
+	{
+		this->operator+=(Other.Path);
+	}
 
 	bool operator==(const FOverriddenPropertyPath& Other) const
 	{
diff --git a/Engine/Source/Runtime/CoreUObject/Public/UObject/PropertyTypeName.h b/Engine/Source/Runtime/CoreUObject/Public/UObject/PropertyTypeName.h
--- a/Engine/Source/Runtime/CoreUObject/Public/UObject/PropertyTypeName.h
+++ b/Engine/Source/Runtime/CoreUObject/Public/UObject/PropertyTypeName.h
@@ -174,6 +174,9 @@ public:
 	/** Resets the builder to allow it to be reused. */
 	UE_API void Reset();
 
+	/** Return a TArray<FName> of the type */
+	UE_API TArray<FName> ToArray() const;
+
 	/**
 	 * Try to parse and add a complete type name with its type parameters.
 	 *
