[UObject]

Hoists an esoteric error in the Linker to a more contextually meaningful one in Rename for when a user attempts to rename an object that is still being loaded.

#rnx

diff --git a/Engine/Source/Runtime/CoreUObject/Private/UObject/Obj.cpp b/Engine/Source/Runtime/CoreUObject/Private/UObject/Obj.cpp
--- a/Engine/Source/Runtime/CoreUObject/Private/UObject/Obj.cpp
+++ b/Engine/Source/Runtime/CoreUObject/Private/UObject/Obj.cpp
@@ -283,7 +283,30 @@ bool UObject::Rename(const TCHAR* InName, UObject* NewOuter, ERenameFlags Flags)
 			// have a linker (in which case we don't want to clear the existing linker)
 			if (NewOuterLinker || GetPackage() != NewOuter->GetPackage())
 			{
-				SetLinker(nullptr, INDEX_NONE);
+				if (Linker && HasAnyFlags(RF_NeedLoad | RF_NeedPostLoad))
+				{
+					const bool bIsPreloading = HasAnyFlags(RF_NeedLoad);
+					const TCHAR* PreloadingGuidance =
+						TEXT("Ensure the target package has completed loading before renaming to correct this error. ") 
+						TEXT("Renaming before an object has been serialized will leave the object uninitialized.");
+					const TCHAR* PostloadingGuidance =
+						TEXT("If this object is no longer valid / is being discarded, call FLinkerLoad::InvalidateExport(ObjectToRename) ")
+						TEXT("before calling Rename, otherwise ensure this object has been postloaded by waiting for the target package's load to complete, or explicitly calling ConditionalPostload(ObjectToRename).");
+
+					UE_LOG(LogObj, Error,
+						TEXT("Rename is moving object '%s' to a new package (%s) or a new linker while the object is still being loaded (it hasn't been %s yet). ")
+						TEXT("%s ")
+						TEXT("Renaming a partially loaded object without first calling FLinkerLoad::InvalidateExport(ObjectToRename) ")
+						TEXT("will result in a partially loaded package likely leading to crashes and/or corrupt data."),
+						*GetFullName(),
+						*NewOuter->GetPackage()->GetFullName(),
+						(HasAnyFlags(RF_NeedLoad) ? TEXT("serialized") : TEXT("postloaded")),
+						(bIsPreloading ? PreloadingGuidance : PostloadingGuidance)
+					);
+					FDebug::DumpStackTraceToLog(ELogVerbosity::Display);
+				}
+
+				SetLinker(nullptr, INDEX_NONE, true /*bShouldDetachExisting*/, true /*bSilenceLogs*/);
 			}
 		}
 	}
diff --git a/Engine/Source/Runtime/CoreUObject/Private/UObject/UObjectLinker.cpp b/Engine/Source/Runtime/CoreUObject/Private/UObject/UObjectLinker.cpp
--- a/Engine/Source/Runtime/CoreUObject/Private/UObject/UObjectLinker.cpp
+++ b/Engine/Source/Runtime/CoreUObject/Private/UObject/UObjectLinker.cpp
@@ -101,7 +101,7 @@ void CleanupLinkerAnnotations()
 }
 
 
-void UObject::SetLinker( FLinkerLoad* LinkerLoad, int32 LinkerIndex, bool bShouldDetachExisting )
+void UObject::SetLinker( FLinkerLoad* LinkerLoad, int32 LinkerIndex, bool bShouldDetachExisting, bool bSilenceLogs)
 {
 	FLinkerIndexPair Existing = LinkerAnnotation.GetAnnotation(this);
 	Existing.CheckInvariants();
@@ -117,7 +117,7 @@ void UObject::SetLinker( FLinkerLoad* LinkerLoad, int32 LinkerIndex, bool bShoul
 	{
 		if (HasAnyFlags(RF_NeedLoad | RF_NeedPostLoad))
 		{
-			UE_LOG(LogUObjectLinker, Error,
+			UE_CLOG(!bSilenceLogs, LogUObjectLinker, Error,
 				TEXT("Detaching from existing linker %s while object %s needs loading (%s). Setting linker to %s. See log for more information."),
 				*Existing.Linker->GetArchiveName(),
 				*GetFullName(),
@@ -156,7 +156,7 @@ void UObject::SetLinker( FLinkerLoad* LinkerLoad, int32 LinkerIndex, bool bShoul
 #if WITH_EDITOR
 		PostLinkerChange();
 #else
-		UE_CLOG(Existing.Linker && LinkerLoad, LogUObjectLinker, Fatal,
+		UE_CLOG(!bSilenceLogs && Existing.Linker && LinkerLoad, LogUObjectLinker, Fatal,
 			TEXT("It is only legal to change linkers in the editor. Trying to change linker on %s from %s (Existing->LinkerRoot=%s) to %s (LinkerLoad->LinkerRoot=%s)"),
 			*GetFullName(),
 			*Existing.Linker->GetDebugName(),
diff --git a/Engine/Source/Runtime/CoreUObject/Public/UObject/Object.h b/Engine/Source/Runtime/CoreUObject/Public/UObject/Object.h
--- a/Engine/Source/Runtime/CoreUObject/Public/UObject/Object.h
+++ b/Engine/Source/Runtime/CoreUObject/Public/UObject/Object.h
@@ -1439,8 +1439,9 @@ public:
 	 * @param LinkerLoad				New LinkerLoad object to set
 	 * @param LinkerIndex				New LinkerIndex to set
 	 * @param bShouldDetachExisting		If true, detach existing linker and call PostLinkerChange
+	 * @param bSilenceLogs				If true, this function will not output log messages
 	 */
-	COREUOBJECT_API void SetLinker( FLinkerLoad* LinkerLoad, int32 LinkerIndex, bool bShouldDetachExisting=true );
+	COREUOBJECT_API void SetLinker( FLinkerLoad* LinkerLoad, int32 LinkerIndex, bool bShouldDetachExisting = true, bool bSilenceLogs = false);
 
 	/**
 	 * Return the template that an object with this class, outer and name would be
