[ObjectLinker]

Always print fatal logs, even if bSilenceLogs is passed in

#rnx

diff --git a/Engine/Source/Runtime/CoreUObject/Private/UObject/UObjectLinker.cpp b/Engine/Source/Runtime/CoreUObject/Private/UObject/UObjectLinker.cpp
--- a/Engine/Source/Runtime/CoreUObject/Private/UObject/UObjectLinker.cpp
+++ b/Engine/Source/Runtime/CoreUObject/Private/UObject/UObjectLinker.cpp
@@ -156,7 +156,7 @@ void UObject::SetLinker( FLinkerLoad* LinkerLoad, int32 LinkerIndex, bool bShoul
 #if WITH_EDITOR
 		PostLinkerChange();
 #else
-		UE_CLOG(!bSilenceLogs && Existing.Linker && LinkerLoad, LogUObjectLinker, Fatal,
+		UE_CLOG(Existing.Linker && LinkerLoad, LogUObjectLinker, Fatal,
 			TEXT("It is only legal to change linkers in the editor. Trying to change linker on %s from %s (Existing->LinkerRoot=%s) to %s (LinkerLoad->LinkerRoot=%s)"),
 			*GetFullName(),
 			*Existing.Linker->GetDebugName(),
diff --git a/Engine/Source/Runtime/CoreUObject/Public/UObject/Object.h b/Engine/Source/Runtime/CoreUObject/Public/UObject/Object.h
--- a/Engine/Source/Runtime/CoreUObject/Public/UObject/Object.h
+++ b/Engine/Source/Runtime/CoreUObject/Public/UObject/Object.h
@@ -1439,7 +1439,7 @@ public:
 	 * @param LinkerLoad				New LinkerLoad object to set
 	 * @param LinkerIndex				New LinkerIndex to set
 	 * @param bShouldDetachExisting		If true, detach existing linker and call PostLinkerChange
-	 * @param bSilenceLogs				If true, this function will not output log messages
+	 * @param bSilenceLogs				If true, this function will not output non-fatal log messages
 	 */
 	COREUOBJECT_API void SetLinker( FLinkerLoad* LinkerLoad, int32 LinkerIndex, bool bShouldDetachExisting = true, bool bSilenceLogs = false);
 
