Optimize live UObject checks on merged modules unload

- Don't rely on explicit synchronous garbage collection as it's going to be done in the deferred pass after GFPs unmount, saving a huge amount of time
- Account for Verse VNI packages in a more systematic way
- Fix recursion in the module manager when unloading shared merged libraries

#rnx

diff --git a/Engine/Source/Runtime/Core/Private/Modules/ModuleManager.cpp b/Engine/Source/Runtime/Core/Private/Modules/ModuleManager.cpp
--- a/Engine/Source/Runtime/Core/Private/Modules/ModuleManager.cpp
+++ b/Engine/Source/Runtime/Core/Private/Modules/ModuleManager.cpp
@@ -645,18 +645,22 @@ void FModuleManager::InternalFreeLibrary(FName ModuleName, void* Handle)
 		OnModulesUnloadCallback.Broadcast(UnloadedModules);
 	}
 
-	// Identify if any library is now useless
-	for (const TPair<FString, void*>& LibraryNameAndHandle : LibraryHandles)
+	// Some merged libraries are shared, meaning they are potentially depended on by multiple merged libraries; unload any that's unused.
+	// This is NOT done when ModuleName isn't set to avoid recursion; only non-shared, merged, dynamic libraries can have dependencies.
+	if (ModuleName != NAME_None)
 	{
-		const FString LibraryName = FPaths::ConvertRelativePathToFull(LibraryNameAndHandle.Key);
-		const TArray<FName> LoadedModules = GetLoadedModulesForLibrary(LibraryName);
-		const TArray<FName> LoadedLibraries = GetLoadedLibrariesForSharedLibrary(LibraryName);
-
-		if (LoadedModules.IsEmpty() && LoadedLibraries.IsEmpty())
+		for (const TPair<FString, void*>& LibraryNameAndHandle : LibraryHandles)
 		{
-			UE_LOG(LogModuleManager, Verbose, TEXT("InternalFreeLibrary: '%s' is now unused"), *LibraryName);
+			const FString LibraryName = FPaths::ConvertRelativePathToFull(LibraryNameAndHandle.Key);
+			const TArray<FName> LoadedModules = GetLoadedModulesForLibrary(LibraryName);
+			const TArray<FName> LoadedLibraries = GetLoadedLibrariesForSharedLibrary(LibraryName);
 
-			InternalFreeLibrary(NAME_None, LibraryNameAndHandle.Value);
+			if (LoadedModules.IsEmpty() && LoadedLibraries.IsEmpty())
+			{
+				UE_LOG(LogModuleManager, Verbose, TEXT("InternalFreeLibrary: '%s' is now unused"), *LibraryName);
+
+				InternalFreeLibrary(NAME_None, LibraryNameAndHandle.Value);
+			}
 		}
 	}
 
diff --git a/Engine/Source/Runtime/CoreUObject/Private/UObject/Class.cpp b/Engine/Source/Runtime/CoreUObject/Private/UObject/Class.cpp
--- a/Engine/Source/Runtime/CoreUObject/Private/UObject/Class.cpp
+++ b/Engine/Source/Runtime/CoreUObject/Private/UObject/Class.cpp
@@ -3107,7 +3107,7 @@ void UScriptStruct::RemoveDeferredCppStructOps(FTopLevelAssetPath Target)
 
 /**
  * On module unload, remove all entries in a package from the stash
- * The package name should be the name component itself, e.g. "Engine" rather than "/Script/Engine"
+ * The package name should be the full package name, such as "/Script/Engine"
  * @param PackageName Name of the package
  */
 void UScriptStruct::RemoveAllDeferredCppStructOps(FName PackageName)
@@ -3116,12 +3116,9 @@ void UScriptStruct::RemoveAllDeferredCppStructOps(FName PackageName)
 
 	TMap<FTopLevelAssetPath, UScriptStruct::ICppStructOps*>& DeferredStructOps = GetDeferredCppStructOps();
 
-	// We're looking for paths that match "/PackageName" - possible full names include "/Script/PackageName" and "/Plugin/_Verse/VNI/PackageName"
-	const FString PackageNamePattern = TEXT("/") + PackageName.ToString();
-
 	for (auto DeferredStructOpsIt = DeferredStructOps.CreateIterator(); DeferredStructOpsIt; ++DeferredStructOpsIt)
 	{
-		if (WriteToString<FName::StringBufferSize>(DeferredStructOpsIt->Key.GetPackageName()).ToView().EndsWith(PackageNamePattern))
+		if (DeferredStructOpsIt->Key.GetPackageName() == PackageName)
 		{
 			UScriptStruct::ICppStructOps* ExistingOps = DeferredStructOpsIt->Value;
 			check(ExistingOps);
diff --git a/Engine/Source/Runtime/CoreUObject/Private/UObject/UObjectBase.cpp b/Engine/Source/Runtime/CoreUObject/Private/UObject/UObjectBase.cpp
--- a/Engine/Source/Runtime/CoreUObject/Private/UObject/UObjectBase.cpp
+++ b/Engine/Source/Runtime/CoreUObject/Private/UObject/UObjectBase.cpp
@@ -1104,13 +1104,24 @@ void ProcessNewlyLoadedUObjects(FName InModuleName, bool bCanProcessNewlyLoadedO
 	}
 }
 
-static UPackage* GetScriptPackageFromModuleName (FName InModuleName )
+static TArray<UPackage*> GetScriptPackagesFromModuleNames(TConstArrayView<FName> InModuleNames)
 {
-	const FName PackagePath = *(TEXT("/Script/") + InModuleName .ToString());
+	TArray<UPackage*> Packages;
 
-	UPackage* Package = FindObjectFast<UPackage>(/* Outer */ nullptr, PackagePath);
+	for (const FName ModuleName : InModuleNames)
+	{
+		if (UPackage* Package = FindObjectFast<UPackage>(/* Outer */ nullptr, WriteToString<256>(TEXT("/Script/"), ModuleName).ToString()))
+		{
+			Packages.Add(Package);
+		}
 
-	return Package;
+		if (UPackage* Package = FindObjectFast<UPackage>(/* Outer */ nullptr, WriteToString<256>(TEXT("/_Verse/VNI/"), ModuleName).ToString()))
+		{
+			Packages.Add(Package);
+		}
+	}
+
+	return Packages;
 }
 
 static void RemoveObjectsInModules(TConstArrayView<FName> InModuleNames)
@@ -1121,28 +1132,36 @@ static void RemoveObjectsInModules(TConstArrayView<FName> InModuleNames)
 
 	FPermanentObjectPoolExtents PermanentPool;
 
-	for (const FName PackageName : InModuleNames)
+	const TArray<UPackage*> Packages = GetScriptPackagesFromModuleNames(InModuleNames);
+	for (UPackage* Package : Packages)
 	{
-		UPackage* Package = GetScriptPackageFromModuleName(PackageName);
-		if (Package)
-		{
-			UE_LOG(LogUObjectBase, Log, TEXT("RemoveObjectsInModules: removing package %s"), *Package->GetName());
+		UE_LOG(LogUObjectBase, Log, TEXT("RemoveObjectsInModules: removing package %s"), *Package->GetName());
 
-			ForEachObjectWithPackage(Package, [&PermanentPool](UObject* InObject) mutable -> bool
+		ForEachObjectWithPackage(Package, [&PermanentPool](UObject* InObject) mutable -> bool
+		{
+			if (!PermanentPool.Contains(InObject) && !GUObjectArray.IsDisregardForGC(InObject))
 			{
-				check(!PermanentPool.Contains(InObject));
+				if (InObject->IsA<UClass>())
+				{
+					ForEachObjectOfClass(Cast<UClass>(InObject), [InObject](UObject* InClassObject) -> bool
+					{
+						InClassObject->ClearInternalFlags(EInternalObjectFlags::RootSet | EInternalObjectFlags::Native);
+						InClassObject->ClearFlags(RF_Standalone);
+						InClassObject->MarkAsGarbage();
 
-				UE_LOG(LogUObjectBase, Verbose, TEXT("RemoveObjectsInModules: marking %s as garbage"), *InObject->GetName());
+						return true;
+					}, true);
+				}
 
-				InObject->ClearInternalFlags(EInternalObjectFlags_RootFlags);
+				InObject->ClearInternalFlags(EInternalObjectFlags::RootSet | EInternalObjectFlags::Native);
 				InObject->ClearFlags(RF_Standalone);
 				InObject->MarkAsGarbage();
+			}
 
-				return true;
-			}, true);
+			return true;
+		}, true);
 
-			UScriptStruct::RemoveAllDeferredCppStructOps(PackageName);
-		}
+		UScriptStruct::RemoveAllDeferredCppStructOps(*Package->GetName());
 	}
 }
 
@@ -1150,41 +1169,45 @@ static bool CheckLiveObjectsInModules(TConstArrayView<FName> InModuleNames)
 {
 	DECLARE_SCOPE_CYCLE_COUNTER(TEXT("CheckLiveObjectsInModules"), STAT_CheckLiveObjectsInModules, STATGROUP_ObjectVerbose);
 
-	bool bObjectsToRemoveStillLive = false;
+	TArray<UObject*> ObjectsToRemoveStillLive;
 
-	for (const FName PackageName : InModuleNames)
+	const TArray<UPackage*> Packages = GetScriptPackagesFromModuleNames(InModuleNames);
+	for (UPackage* Package : Packages)
 	{
-		UPackage* Package = GetScriptPackageFromModuleName (PackageName);
-		if (Package)
+		ForEachObjectWithPackage(Package, [&Package, &ObjectsToRemoveStillLive](UObject* InObject) -> bool
 		{
-			ForEachObjectWithPackage(Package, [&bObjectsToRemoveStillLive](UObject* InObject) -> bool
-			{
-				UE_LOG(LogUObjectBase, Warning, TEXT("CheckLiveObjectsInModules: object %s should have been removed (0x%0lX 0x%0lX)"),
-					*InObject->GetName(), InObject->GetFlags(), InObject->GetInternalFlags());
+			UE_LOG(LogUObjectBase, Warning, TEXT("CheckLiveObjectsInModules: object %s in %s should have been removed (0x%0lX 0x%0lX)"),
+				*InObject->GetName(), *Package->GetName(), InObject->GetFlags(), InObject->GetInternalFlags());
 
-				if (InObject->IsA<UClass>())
+			if (InObject->IsA<UClass>())
+			{
+				ForEachObjectOfClass(Cast<UClass>(InObject), [InObject, &Package, &ObjectsToRemoveStillLive](UObject* InClassObject) -> bool
 				{
-					ForEachObjectOfClass(Cast<UClass>(InObject), [InObject](UObject* InClassObject) -> bool
-					{
-						UE_LOG(LogUObjectBase, Warning, TEXT("CheckLiveObjectsInModules: %s has live instance %s"),
-							*InObject->GetName(), *InClassObject->GetName());
+					UE_LOG(LogUObjectBase, Warning, TEXT("CheckLiveObjectsInModules: instance %s in %s should have been removed (0x%0lX 0x%0lX)"),
+						*InClassObject->GetName(), *Package->GetName(), InClassObject->GetFlags(), InClassObject->GetInternalFlags());
+					ObjectsToRemoveStillLive.Add(InClassObject);
 
-						FReferenceChainSearch::FindAndPrintStaleReferencesToObject(InClassObject, EPrintStaleReferencesOptions::Log);
-						return true;
-					});
-				}
-				else
-				{
-					FReferenceChainSearch::FindAndPrintStaleReferencesToObject(InObject, EPrintStaleReferencesOptions::Log);
-				}
+					return true;
+				}, true);
+			}
 
-				bObjectsToRemoveStillLive = true;
-				return true;
-			}, true);
-		}
+			ObjectsToRemoveStillLive.Add(InObject);
+			return true;
+		}, true);
 	}
 
-	return bObjectsToRemoveStillLive;
+	if (ObjectsToRemoveStillLive.Num() > 0)
+	{
+#if !UE_BUILD_SHIPPING
+		FReferenceChainSearch::FindAndPrintStaleReferencesToObjects(ObjectsToRemoveStillLive, EPrintStaleReferencesOptions::Log);
+#endif // !UE_BUILD_SHIPPING
+		return true;
+	}
+	else
+	{
+		UE_LOG(LogUObjectBase, Log, TEXT("CheckLiveObjectsInModules: all objects were correctly removed"));
+		return false;
+	}
 }
 
 static void RemoveLoadedUObjects(TConstArrayView<FName> InModuleNames)
@@ -1195,7 +1218,6 @@ static void RemoveLoadedUObjects(TConstArrayView<FName> InModuleNames)
 
 	FlushAsyncLoading();
 	RemoveObjectsInModules(InModuleNames);
-	CollectGarbage(GARBAGE_COLLECTION_KEEPFLAGS, true);
 }
 
 static int32 GVarMaxObjectsNotConsideredByGC;
diff --git a/Engine/Source/Runtime/CoreUObject/Public/UObject/Class.h b/Engine/Source/Runtime/CoreUObject/Public/UObject/Class.h
--- a/Engine/Source/Runtime/CoreUObject/Public/UObject/Class.h
+++ b/Engine/Source/Runtime/CoreUObject/Public/UObject/Class.h
@@ -2160,7 +2160,7 @@ public:
 
 	/**
 	 * On module unload, remove all entries in a package from the stash
-	 * The package name should be the name component itself, e.g. "Engine" rather than "/Script/Engine"
+	 * The package name should be the full package name, such as "/Script/Engine"
 	 * @param PackageName Name of the package
 	 */
 	static COREUOBJECT_API void RemoveAllDeferredCppStructOps(FName PackageName);
