PropertyBag: When a property is overridden, update all its members, not just the type.

diff --git a/Engine/Source/Runtime/CoreUObject/Private/StructUtils/PropertyBag.cpp b/Engine/Source/Runtime/CoreUObject/Private/StructUtils/PropertyBag.cpp
--- a/Engine/Source/Runtime/CoreUObject/Private/StructUtils/PropertyBag.cpp
+++ b/Engine/Source/Runtime/CoreUObject/Private/StructUtils/PropertyBag.cpp
@@ -1975,8 +1975,9 @@ EPropertyBagAlterationResult FInstancedPropertyBag::AddProperties(const TConstAr
 		{
 			if (bOverwrite)
 			{
-				ExistingProperty->ValueType = NewDesc.ValueType;
-				ExistingProperty->ValueTypeObject = NewDesc.ValueTypeObject;
+				const FGuid ExistingID = ExistingProperty->ID;
+				*ExistingProperty = NewDesc;
+				ExistingProperty->ID = ExistingID;
 			}
 			else
 			{
