Disable OS when cooking so we are closer to what UPS is doing.

#rnx

diff --git a/Engine/Source/Runtime/CoreUObject/Private/UObject/Class.cpp b/Engine/Source/Runtime/CoreUObject/Private/UObject/Class.cpp
--- a/Engine/Source/Runtime/CoreUObject/Private/UObject/Class.cpp
+++ b/Engine/Source/Runtime/CoreUObject/Private/UObject/Class.cpp
@@ -1555,8 +1555,15 @@ enum class EClassSerializationControlExtension : uint8
 };
 ENUM_CLASS_FLAGS(EClassSerializationControlExtension);
 
+namespace UE::OS
+{
+	static bool GUseForCooking = false;
+	static FAutoConsoleVariableRef CVarOSUseForCooking(TEXT("OS.UseForCooking"), GUseForCooking, TEXT("Use overridable serialization for cooking else fall back on delta serialization"));
+} // UE::OS
+
 struct FSerializationControlExtensionContext
 {
+	FArchive& UnderlyingArchive;
 	uint8* Data = nullptr;
 	bool bEnableOverridableSerialization = false;
 	FOverriddenPropertySet* OverriddenProperties = nullptr;
@@ -1564,13 +1571,20 @@ struct FSerializationControlExtensionContext
 	EClassSerializationControlExtension InitializeSerializationControlExtensions()
 	{
 		EClassSerializationControlExtension SerializationExtension = EClassSerializationControlExtension::NoExtension;
-	
-		// Overridable serialization information initialization
-		if (FOverriddenPropertySet* ObjectOverriddenProperties = FOverridableManager::Get().GetOverriddenProperties((UObject*)Data))
+
+		// Let's not enable OS(Overridable Serialization) when cooking and fall back on DS(Delta Serialization).
+		// That way the result will be similar if the server is using TPS(Tag Property Serialization) and the client is using UPS(Unversion Property Serialization)
+		// This fixes the problem CDO divergences in network layer with the server having -0.0 with OS and the client having 0.0 using delta serialization
+		// as mathematically (-0.0 == 0.0) but is memcmp(&-0.0, &0.0) is false
+		if (UE::OS::GUseForCooking || !UnderlyingArchive.IsCooking())
 		{
-			SerializationExtension |= EClassSerializationControlExtension::OverridableSerializationInformation;
-			bEnableOverridableSerialization = true;
-			OverriddenProperties = ObjectOverriddenProperties;
+			// Overridable serialization information initialization
+			if (FOverriddenPropertySet* ObjectOverriddenProperties = FOverridableManager::Get().GetOverriddenProperties((UObject*)Data))
+			{
+				SerializationExtension |= EClassSerializationControlExtension::OverridableSerializationInformation;
+				bEnableOverridableSerialization = true;
+				OverriddenProperties = ObjectOverriddenProperties;
+			}
 		}
 	
 		return SerializationExtension;
@@ -1593,7 +1607,7 @@ void UStruct::SerializeVersionedTaggedProperties(FStructuredArchive::FSlot Slot,
 
 	// Setup serialization control data extensions, this is serialized only on root i.e. UObject and not structs!
 	const bool bIsUClass = IsA<UClass>();
-	FSerializationControlExtensionContext ControlContext{ Data };
+	FSerializationControlExtensionContext ControlContext{ UnderlyingArchive, Data };
 	if (bIsUClass && UnderlyingArchive.UEVer() >= EUnrealEngineObjectUE5Version::PROPERTY_TAG_EXTENSION_AND_OVERRIDABLE_SERIALIZATION)
 	{
 		EClassSerializationControlExtension SerializationControl = EClassSerializationControlExtension::NoExtension;
