Optimization for FSoftObjectPath::ResolveObjectInternal to avoid the round trip to/from string-- changed to just do directly what we end up doing with the string internally-- split into package + sub-object path.

diff --git a/Engine/Source/Runtime/CoreUObject/Private/UObject/SoftObjectPath.cpp b/Engine/Source/Runtime/CoreUObject/Private/UObject/SoftObjectPath.cpp
--- a/Engine/Source/Runtime/CoreUObject/Private/UObject/SoftObjectPath.cpp
+++ b/Engine/Source/Runtime/CoreUObject/Private/UObject/SoftObjectPath.cpp
@@ -897,10 +897,13 @@ UObject* FSoftObjectPath::ResolveObject() const
 
 UObject* FSoftObjectPath::ResolveObjectInternal() const
 {
-	TStringBuilder<FName::StringBufferSize> Builder;
-	Builder << *this;
-
-	UObject* FoundObject = FindObject<UObject>(nullptr, *Builder);
+	UObject* FoundObject = FindObject<UObject>(AssetPath);
+	if (FoundObject && !SubPathString.IsEmpty())
+	{
+		TStringBuilder<FName::StringBufferSize> Builder;
+		Builder << SubPathString;
+		FoundObject = FindObject<UObject>(FoundObject, *Builder);
+	}
 
 	if (!FoundObject && IsSubobject())
 	{
