Add `!IsGarbageCollectionLocked()` check before we call shutdown delegates to prevent hard-locking

diff --git a/Engine/Source/Runtime/ApplicationCore/Private/IOS/IOSAppDelegate.cpp b/Engine/Source/Runtime/ApplicationCore/Private/IOS/IOSAppDelegate.cpp
--- a/Engine/Source/Runtime/ApplicationCore/Private/IOS/IOSAppDelegate.cpp
+++ b/Engine/Source/Runtime/ApplicationCore/Private/IOS/IOSAppDelegate.cpp
@@ -1832,7 +1832,10 @@ extern double GCStartTime;
 	 Save data if appropriate.
 	 See also applicationDidEnterBackground:.
 	 */
-	FCoreDelegates::GetApplicationWillTerminateDelegate().Broadcast();
+	if (!UE::IsGCLocked())
+	{
+		FCoreDelegates::GetApplicationWillTerminateDelegate().Broadcast();
+	}
     
     // note that we are shutting down
     // TODO: fix the reason why we are hanging when asked to shutdown
diff --git a/Engine/Source/Runtime/Core/Private/Mac/MacPlatformMisc.cpp b/Engine/Source/Runtime/Core/Private/Mac/MacPlatformMisc.cpp
--- a/Engine/Source/Runtime/Core/Private/Mac/MacPlatformMisc.cpp
+++ b/Engine/Source/Runtime/Core/Private/Mac/MacPlatformMisc.cpp
@@ -720,8 +720,11 @@ void FMacPlatformMisc::RequestExit( bool Force, const TCHAR* CallSite)
 	UE_LOG(LogMac, Log,  TEXT("FPlatformMisc::RequestExit(%i, %s)"), Force,
 		CallSite ? CallSite : TEXT("<NoCallSiteInfo>"));
 
-	FCoreDelegates::GetApplicationWillTerminateDelegate().Broadcast();
-
+	if (!UE::IsGCLocked())
+	{
+		FCoreDelegates::GetApplicationWillTerminateDelegate().Broadcast();
+	}
+	
 	notify_cancel(GMacAppInfo.PowerSourceNotification);
 	GMacAppInfo.PowerSourceNotification = 0;
 	
diff --git a/Engine/Source/Runtime/Core/Private/Misc/CoreGlobals.cpp b/Engine/Source/Runtime/Core/Private/Misc/CoreGlobals.cpp
--- a/Engine/Source/Runtime/Core/Private/Misc/CoreGlobals.cpp
+++ b/Engine/Source/Runtime/Core/Private/Misc/CoreGlobals.cpp
@@ -1011,3 +1011,23 @@ void SetMultiprocessId(int32 MultiprocessId)
 }
 
 }
+
+bool(*IsGCLockedFuncPtr)() = nullptr;
+
+namespace UE
+{
+bool IsGCLocked()
+{
+	return IsGCLockedFuncPtr && IsGCLockedFuncPtr();
+}
+
+}
+namespace UE::Private
+{
+void SetIsGCLockedFunction(bool(*InIsGCLockedFuncPtr)())
+{
+	IsGCLockedFuncPtr = InIsGCLockedFuncPtr;
+}
+
+}
+
diff --git a/Engine/Source/Runtime/Core/Private/Windows/WindowsPlatformMisc.cpp b/Engine/Source/Runtime/Core/Private/Windows/WindowsPlatformMisc.cpp
--- a/Engine/Source/Runtime/Core/Private/Windows/WindowsPlatformMisc.cpp
+++ b/Engine/Source/Runtime/Core/Private/Windows/WindowsPlatformMisc.cpp
@@ -1377,7 +1377,10 @@ void FWindowsPlatformMisc::RequestExitWithStatus(bool Force, uint8 ReturnCode, c
 
 	RequestEngineExit(TEXT("Win RequestExit"));
 
-	FCoreDelegates::GetApplicationWillTerminateDelegate().Broadcast();
+	if (!UE::IsGCLocked())
+	{
+		FCoreDelegates::GetApplicationWillTerminateDelegate().Broadcast();
+	}
 
 	if (Force)
 	{
diff --git a/Engine/Source/Runtime/Core/Public/CoreGlobals.h b/Engine/Source/Runtime/Core/Public/CoreGlobals.h
--- a/Engine/Source/Runtime/Core/Public/CoreGlobals.h
+++ b/Engine/Source/Runtime/Core/Public/CoreGlobals.h
@@ -846,3 +846,12 @@ namespace UE {
  */
 UE_DEPRECATED(5.6, "IsThisNotNull has been deprecated.")
 CORE_API FORCENOINLINE bool IsThisNotNull(const void* This, const ANSICHAR* Function);
+
+namespace UE
+{
+	CORE_API bool IsGCLocked();
+}
+namespace UE::Private
+{
+	CORE_API void SetIsGCLockedFunction(bool(*InIsGCLockedFuncPtr)());
+}
diff --git a/Engine/Source/Runtime/CoreUObject/Private/UObject/GarbageCollection.cpp b/Engine/Source/Runtime/CoreUObject/Private/UObject/GarbageCollection.cpp
--- a/Engine/Source/Runtime/CoreUObject/Private/UObject/GarbageCollection.cpp
+++ b/Engine/Source/Runtime/CoreUObject/Private/UObject/GarbageCollection.cpp
@@ -151,14 +151,21 @@ public:
 	}
 };
 
+bool IsGarbageCollectionLocked()
+{
+	return FGCCSyncObject::Get().IsAsyncLocked();
+}
+
 FGCCSyncObject::FGCCSyncObject()
 {
 	GCUnlockedEvent = FPlatformProcess::GetSynchEventFromPool(true);
+	UE::Private::SetIsGCLockedFunction(IsGarbageCollectionLocked);
 }
 FGCCSyncObject::~FGCCSyncObject()
 {
 	FPlatformProcess::ReturnSynchEventToPool(GCUnlockedEvent);
 	GCUnlockedEvent = nullptr;
+	UE::Private::SetIsGCLockedFunction(nullptr);
 }
 
 FGCCSyncObject* GGCSingleton;
@@ -218,11 +225,6 @@ FGCScopeTryGuard::~FGCScopeTryGuard()
 	}
 }
 
-bool IsGarbageCollectionLocked()
-{
-	return FGCCSyncObject::Get().IsAsyncLocked();
-}
-
 static int32 GIncrementalBeginDestroyEnabled = 1;
 static FAutoConsoleVariableRef CIncrementalBeginDestroyEnabled(
 	TEXT("gc.IncrementalBeginDestroyEnabled"),
diff --git a/Engine/Source/Runtime/Launch/Private/Android/LaunchAndroid.cpp b/Engine/Source/Runtime/Launch/Private/Android/LaunchAndroid.cpp
--- a/Engine/Source/Runtime/Launch/Private/Android/LaunchAndroid.cpp
+++ b/Engine/Source/Runtime/Launch/Private/Android/LaunchAndroid.cpp
@@ -1928,7 +1928,10 @@ static void OnAppCommandCB(struct android_app* app, int32_t cmd)
 		{
 			FGraphEventRef WillTerminateTask = FFunctionGraphTask::CreateAndDispatchWhenReady([]()
 			{
-				FCoreDelegates::GetApplicationWillTerminateDelegate().Broadcast();
+				if (!UE::IsGCLocked())
+				{
+					FCoreDelegates::GetApplicationWillTerminateDelegate().Broadcast();
+				}
 			}, TStatId(), NULL, ENamedThreads::GameThread);
 			FTaskGraphInterface::Get().WaitUntilTaskCompletes(WillTerminateTask);
 			FAndroidMisc::NonReentrantRequestExit();
