MPCook: Add function comments to IMPCollector.
#rnx

diff --git a/Engine/Source/Runtime/CoreUObject/Public/Cooker/MPCollector.h b/Engine/Source/Runtime/CoreUObject/Public/Cooker/MPCollector.h
--- a/Engine/Source/Runtime/CoreUObject/Public/Cooker/MPCollector.h
+++ b/Engine/Source/Runtime/CoreUObject/Public/Cooker/MPCollector.h
@@ -220,14 +220,41 @@ class IMPCollector : public FRefCountBase
 public:
 	virtual ~IMPCollector() {}
 
+	/**
+	 * Return a Guid that identifies the IMPCollector uniquely from all other IMPCollectors. Implement this by
+	 * return FGuid(0xA,0xB,0xC,0xD), where the values for A,B,C,D are copied from a guid generator tool.
+	 */
 	virtual FGuid GetMessageType() const = 0;
+	/** Return a debug name for the Collector, used in warnings or errors from the cooker about it. */
 	virtual const TCHAR* GetDebugName() const = 0;
 
+	/**
+	 * Called during events on the CookDirector (see Context.GetEventType()). Not called periodically, if you
+	 * need a periodic call, implement FTickableCookObject. Call AddMessage which will be passed to
+	 * ClientReceiveMessage on each CookWorker.
+	 **/
 	virtual void ServerTick(FMPCollectorServerTickContext& Context) {}
+	/**
+	 * Called periodically on CookWorkers (period is determined by cooker), and called once with IsFlush()=true at end
+	 * of cook. Call AddMessage which will be passed to ServerReceiveMessage on the CookDirector.
+	 */
 	virtual void ClientTick(FMPCollectorClientTickContext& Context) {}
+	/**
+	 * Called on Director each time a package is assigned to a CookWorker (note this may occur several times for a
+	 * single package in cases of retraction). Call AddMessage which will be passed to ClientReceiveMessage on the
+	 * assigned CookWorker.
+	 */
 	virtual void ServerTickPackage(FMPCollectorServerTickPackageContext& Context) {}
+	/**
+	 * Called on CookWorker when a package is saved, or found to be unsaveable, on a CookWoorker.
+	 * In MultiprocessCooks, it will usually be called once with a set of the platforms saved, but it is possible it
+	 * will be called once per platform if the package is found at different times to be needed on different platforms.
+	 * Call AddMessage which will be passed to ServerReceiveMessage on the CookDirector.
+	 */
 	virtual void ClientTickPackage(FMPCollectorClientTickPackageContext& Context) {}
+	/** Called on CookWorker to receive a message from ServerTick or ServerTickPackage. */
 	virtual void ClientReceiveMessage(FMPCollectorClientMessageContext& Context, FCbObjectView Message) {}
+	/** Called on CookDirector to receive a message from ClientTick or ClientTickPackage. */
 	virtual void ServerReceiveMessage(FMPCollectorServerMessageContext& Context, FCbObjectView Message) {}
 };
 
