MPCook: IMPCollector: fix typo in comment.
#rnx

diff --git a/Engine/Source/Runtime/CoreUObject/Public/Cooker/MPCollector.h b/Engine/Source/Runtime/CoreUObject/Public/Cooker/MPCollector.h
--- a/Engine/Source/Runtime/CoreUObject/Public/Cooker/MPCollector.h
+++ b/Engine/Source/Runtime/CoreUObject/Public/Cooker/MPCollector.h
@@ -246,7 +246,7 @@ public:
 	 */
 	virtual void ServerTickPackage(FMPCollectorServerTickPackageContext& Context) {}
 	/**
-	 * Called on CookWorker when a package is saved, or found to be unsaveable, on a CookWoorker.
+	 * Called on CookWorker when a package is saved, or found to be unsaveable.
 	 * In MultiprocessCooks, it will usually be called once with a set of the platforms saved, but it is possible it
 	 * will be called once per platform if the package is found at different times to be needed on different platforms.
 	 * Call AddMessage which will be passed to ServerReceiveMessage on the CookDirector.
