[AsyncLoading]

Allow runtime builds using loose cooked data to properly load redirected localized content, previously only editor builds / -game would load loose localized content.

#rnx

diff --git a/Engine/Source/Runtime/CoreUObject/Private/Serialization/AsyncLoading2.cpp b/Engine/Source/Runtime/CoreUObject/Private/Serialization/AsyncLoading2.cpp
--- a/Engine/Source/Runtime/CoreUObject/Private/Serialization/AsyncLoading2.cpp
+++ b/Engine/Source/Runtime/CoreUObject/Private/Serialization/AsyncLoading2.cpp
@@ -623,25 +623,21 @@ FName ApplyPackageNameRedirections(FName PackageName)
 
 	return RedirectedName;
 }
+#endif
 
 /** Returns localized package name to load when running with -game. Otherwise returns the name unchanged. */
 FName ApplyLocalizationPackageNameRedirects(FName PackageName)
 {
-	if (!GIsEditor)
+	UE::Core::Private::FPlayInEditorLoadingScope PlayInEditorIDScope(INDEX_NONE);
+	FName RedirectedName = FPackageLocalizationManager::Get().FindLocalizedPackageName(PackageName);
+	if (!RedirectedName.IsNone())
 	{
-		UE::Core::Private::FPlayInEditorLoadingScope PlayInEditorIDScope(INDEX_NONE);
-		FName RedirectedName = FPackageLocalizationManager::Get().FindLocalizedPackageName(PackageName);
-		if (!RedirectedName.IsNone())
-		{
-			UE_LOG(LogStreaming, Verbose, TEXT("Applying localization redirection for package %s to package %s"), *PackageName.ToString(), *RedirectedName.ToString());
-			PackageName = RedirectedName;
-		}
+		UE_LOG(LogStreaming, Verbose, TEXT("Applying localization redirection for package %s to package %s"), *PackageName.ToString(), *RedirectedName.ToString());
+		PackageName = RedirectedName;
 	}
 	return PackageName;
 }
 
-#endif
-
 struct FAsyncPackage2;
 class FAsyncLoadingThread2;
 
@@ -5683,10 +5679,12 @@ bool FAsyncLoadingThread2::CreateAsyncPackagesFromQueue(FAsyncLoadingThreadState
 			// Redirects should redirect both the UPackage name as well as the package name to load
 			PackageNameToLoad = ApplyPackageNameRedirections(PackageNameToLoad);
 			UPackageName = PackageNameToLoad;
-
-			// Localization should only update the package name to load
-			PackageNameToLoad = ApplyLocalizationPackageNameRedirects(PackageNameToLoad);
 #endif
+			if (FAsyncLoadingThreadSettings::Get().bLooseFileLoadingEnabled)
+			{
+				// Localization should only update the package name to load
+				PackageNameToLoad = ApplyLocalizationPackageNameRedirects(PackageNameToLoad);
+			}
 
 			FPackageId PackageIdToLoad = FPackageId::FromName(PackageNameToLoad);
 			{
@@ -6589,8 +6587,8 @@ void FAsyncPackage2::ImportPackagesRecursiveInner(FAsyncLoadingThreadState2& Thr
 		}
 #endif
 
-#if WITH_EDITOR
 		// Note, localization should apply after remapping from instancing has occurred
+		if (FAsyncLoadingThreadSettings::Get().bLooseFileLoadingEnabled)
 		{
 			FName LocalizedPackageNameToLoad = ApplyLocalizationPackageNameRedirects(ImportedPackageNameToLoad);
 			if (LocalizedPackageNameToLoad != ImportedPackageNameToLoad)
@@ -6600,7 +6598,6 @@ void FAsyncPackage2::ImportPackagesRecursiveInner(FAsyncLoadingThreadState2& Thr
 				ImportedPackageIdToLoad = FPackageId::FromName(ImportedPackageNameToLoad);
 			}
 		}
-#endif
 
 		{
 			FName SourcePackageName;
