Adding the ability to add the remote object name/pathname/fullname when converting a remote object id to a string

diff --git a/Engine/Source/Runtime/CoreUObject/Private/UObject/RemoteObject.cpp b/Engine/Source/Runtime/CoreUObject/Private/UObject/RemoteObject.cpp
--- a/Engine/Source/Runtime/CoreUObject/Private/UObject/RemoteObject.cpp
+++ b/Engine/Source/Runtime/CoreUObject/Private/UObject/RemoteObject.cpp
@@ -34,6 +34,19 @@ static FAutoConsoleVariableRef CVarInitRemoteServerIdBeforeUObjectInit(
 	GInitRemoteServerIdBeforeUObjectInit,
 	TEXT("Initializes remote server id before UObject system is initialized"));
 
+int32 GRemoteIdToStringVerbosity = (int32)ERemoteIdToStringVerbosity::Id;
+static FAutoConsoleVariableRef CVarRemoteIdToStringVerbosity(
+	TEXT("ro.IdToStringVerbosity"),
+	GRemoteIdToStringVerbosity,
+	TEXT("Sets the verbosity FRemoteObjectId::ToString() prints the id to log with (1: id only, 2: id + pathname, 4: id + fullname, 5: id + fullname + attributes"),
+	FConsoleVariableDelegate::CreateLambda([](IConsoleVariable*)
+		{
+			if (GRemoteIdToStringVerbosity <= (int32)ERemoteIdToStringVerbosity::Default || GRemoteIdToStringVerbosity > (int32)ERemoteIdToStringVerbosity::Max)
+			{
+				GRemoteIdToStringVerbosity = (int32)ERemoteIdToStringVerbosity::Id;
+			}
+		}));
+
 FRemoteServerId FRemoteServerId::GlobalServerId(ERemoteServerIdConstants::Local); // Here Local means uninitialized
 
 void AssignGlobalServerIdDebuggingState()
@@ -946,9 +959,85 @@ FRemoteObjectId FRemoteObjectId::Generate(UObjectBase* InObject, const TCHAR* In
 	return Result;
 }
 
-FString FRemoteObjectId::ToString() const
+FString FRemoteObjectId::ToString(ERemoteIdToStringVerbosity InVerbosityOverride /*= ERemoteIdToStringVerbosity::Default*/) const
 {
-	return FString::Printf(TEXT("%s-%llu"), *GetServerId().ToString(), SerialNumber);
+	using namespace UE::RemoteObject::Handle;
+
+	const int32 Verbosity = (InVerbosityOverride == ERemoteIdToStringVerbosity::Default) ? 
+		FMath::Clamp(GRemoteIdToStringVerbosity, (int32)ERemoteIdToStringVerbosity::Id, (int32)ERemoteIdToStringVerbosity::Max) : (int32)InVerbosityOverride;
+
+	if (Verbosity <= (int32)ERemoteIdToStringVerbosity::Id)
+	{
+		return FString::Printf(TEXT("%s-%llu"), *GetServerId().ToString(), SerialNumber);
+	}
+	else
+	{
+		FString AdditionalInfo;
+		bool bStubOnly = false;
+
+		if (UObject* ExistingObject = StaticFindObjectFastInternal(*this))
+		{
+			switch ((ERemoteIdToStringVerbosity)Verbosity)
+			{
+			case ERemoteIdToStringVerbosity::Name:
+				AdditionalInfo = ExistingObject->GetName();
+				break;
+			case ERemoteIdToStringVerbosity::PathName:
+				AdditionalInfo = ExistingObject->GetPathName();
+				break;
+			case ERemoteIdToStringVerbosity::FullName:
+			case ERemoteIdToStringVerbosity::FullNameAttributes:
+				AdditionalInfo = ExistingObject->GetFullName();
+				break;
+			}
+		}
+		else
+		{
+			FRemoteObjectPathName RemotePathName(*this);
+			if (RemotePathName.Num())
+			{
+				bStubOnly = true;
+				switch ((ERemoteIdToStringVerbosity)Verbosity)
+				{
+				case ERemoteIdToStringVerbosity::Name:
+					AdditionalInfo = RemotePathName.GetObjectName().ToString();
+					break;
+				case ERemoteIdToStringVerbosity::PathName:
+					AdditionalInfo = RemotePathName.ToString();
+					break;
+				case ERemoteIdToStringVerbosity::FullName:
+				case ERemoteIdToStringVerbosity::FullNameAttributes:
+				{
+					UClass* Class = GetClass(*this, ERemoteObjectGetClassBehavior::ReturnNullIfNeverLocal);
+					AdditionalInfo = FString::Printf(TEXT("%s %s"), *GetNameSafe(Class), *RemotePathName.ToString());
+				}
+				break;
+				}
+			}
+			else
+			{
+				AdditionalInfo = TEXT("Unknown object");
+			}
+		}
+
+		if (Verbosity >= (int32)ERemoteIdToStringVerbosity::FullNameAttributes)
+		{
+			if (IsRemote(*this))
+			{
+				AdditionalInfo += TEXT(" (remote)");
+			}
+			if (IsOwned(*this))
+			{
+				AdditionalInfo += TEXT(" (owned)");
+			}
+			if (bStubOnly)
+			{
+				AdditionalInfo += TEXT(" (stub)");
+			}
+		}
+
+		return FString::Printf(TEXT("%s-%llu %s"), *GetServerId().ToString(), SerialNumber, *AdditionalInfo);
+	}
 }
 
 FRemoteObjectId FRemoteObjectId::FromString(const FStringView& InText)
diff --git a/Engine/Source/Runtime/CoreUObject/Private/UObject/RemoteObjectSerialization.cpp b/Engine/Source/Runtime/CoreUObject/Private/UObject/RemoteObjectSerialization.cpp
--- a/Engine/Source/Runtime/CoreUObject/Private/UObject/RemoteObjectSerialization.cpp
+++ b/Engine/Source/Runtime/CoreUObject/Private/UObject/RemoteObjectSerialization.cpp
@@ -83,7 +83,7 @@ void DeserializeNameTables(FArchive& Ar, FRemoteObjectData& OutObjectData)
 
 FString GenerateRemoteObjectFilename(FRemoteObjectId ObjectId, FRemoteServerId OwnerServerId)
 {
-	return FPaths::Combine(*FPaths::ProjectSavedDir(), *(FString::Printf(TEXT("%s-%s_%s.remote"), *FRemoteServerId::GetLocalServerId().ToString(), *ObjectId.ToString(), *OwnerServerId.ToString())));
+	return FPaths::Combine(*FPaths::ProjectSavedDir(), *(FString::Printf(TEXT("%s-%s_%s.remote"), *FRemoteServerId::GetLocalServerId().ToString(), *ObjectId.ToString(ERemoteIdToStringVerbosity::Id), *OwnerServerId.ToString())));
 }
 
 void LoadObjectFromDisk(const FUObjectMigrationContext& MigrationContext)
diff --git a/Engine/Source/Runtime/CoreUObject/Public/UObject/RemoteObjectPathName.h b/Engine/Source/Runtime/CoreUObject/Public/UObject/RemoteObjectPathName.h
--- a/Engine/Source/Runtime/CoreUObject/Public/UObject/RemoteObjectPathName.h
+++ b/Engine/Source/Runtime/CoreUObject/Public/UObject/RemoteObjectPathName.h
@@ -90,6 +90,15 @@ struct FRemoteObjectPathName : public FRemoteObjectTables
 		return RemoteIds[InSegmentIndex];
 	}
 
+	UE_FORCEINLINE_HINT FName GetObjectName() const
+	{
+		if (Names.Num())
+		{
+			return Names[0];
+		}
+		return FName();
+	}
+
 	COREUOBJECT_API UObject* Resolve() const;
 	COREUOBJECT_API FString ToString(int32 InMinPathSegmentIndex = 0) const;
 };
\ No newline at end of file
diff --git a/Engine/Source/Runtime/CoreUObject/Public/UObject/RemoteObjectTypes.h b/Engine/Source/Runtime/CoreUObject/Public/UObject/RemoteObjectTypes.h
--- a/Engine/Source/Runtime/CoreUObject/Public/UObject/RemoteObjectTypes.h
+++ b/Engine/Source/Runtime/CoreUObject/Public/UObject/RemoteObjectTypes.h
@@ -195,6 +195,19 @@ private:
 	COREUOBJECT_API static FRemoteServerId GlobalServerId;
 };
 
+enum class ERemoteIdToStringVerbosity
+{
+	Default = 0,
+	Id = 1,
+	Name = 2,
+	PathName = 3,
+	FullName = 4,
+	FullNameAttributes = 5,
+
+	LastPlusOne,
+	Max = LastPlusOne - 1
+};
+
 struct FRemoteObjectId
 {
 private:
@@ -314,7 +327,7 @@ public:
 		return GetServerId().IsLocal();
 	}
 
-	COREUOBJECT_API FString ToString() const;
+	COREUOBJECT_API FString ToString(ERemoteIdToStringVerbosity InVerbosityOverride = ERemoteIdToStringVerbosity::Default) const;
 	static COREUOBJECT_API FRemoteObjectId FromString(const FStringView& InText);
 
 	COREUOBJECT_API static FRemoteObjectId Generate(UObjectBase* InObject, const TCHAR* InName, EInternalObjectFlags InInitialFlags = EInternalObjectFlags::None);
diff --git a/Engine/Source/Runtime/Engine/Private/UnrealEngine.cpp b/Engine/Source/Runtime/Engine/Private/UnrealEngine.cpp
--- a/Engine/Source/Runtime/Engine/Private/UnrealEngine.cpp
+++ b/Engine/Source/Runtime/Engine/Private/UnrealEngine.cpp
@@ -10167,44 +10167,7 @@ bool UEngine::HandleObjCommand( const TCHAR* Cmd, FOutputDevice& Ar )
 
 		if (RemoteId.IsValid())
 		{
-			FString PathName;
-			FString ClassName;
-			bool bStubOnly = false;
-
-			if (UObject* ExistingObject = StaticFindObjectFastInternal(RemoteId))
-			{
-				PathName = ExistingObject->GetPathName();
-				ClassName = ExistingObject->GetClass()->GetName();
-			}
-			else if (UClass* Class = UE::RemoteObject::Handle::GetClass(RemoteId, ERemoteObjectGetClassBehavior::ReturnNullIfNeverLocal))
-			{
-				FRemoteObjectPathName RemotePathName(RemoteId);
-				PathName = RemotePathName.ToString();
-				ClassName = Class->GetName();
-				bStubOnly = true;
-			}
-
-			if (PathName.Len())
-			{
-				FString Attributes;
-				if (IsRemote(RemoteId))
-				{
-					Attributes += TEXT(" (remote)");
-				}
-				if (IsOwned(RemoteId))
-				{
-					Attributes += TEXT(" (owned)");
-				}
-				if (bStubOnly)
-				{
-					Attributes += TEXT(" (stub)");
-				}
-				Ar.Logf(TEXT("%s%s %s %s"), *RemoteId.ToString(), *Attributes, *ClassName, *PathName);
-			}
-			else
-			{
-				Ar.Logf(TEXT("No object found matching remote id '%s'"), *RemoteId.ToString());
-			}
+			Ar.Log(*RemoteId.ToString(ERemoteIdToStringVerbosity::Max));
 		}
 		else
 		{
