Fixing console variable help text

diff --git a/Engine/Source/Runtime/CoreUObject/Private/UObject/RemoteObject.cpp b/Engine/Source/Runtime/CoreUObject/Private/UObject/RemoteObject.cpp
--- a/Engine/Source/Runtime/CoreUObject/Private/UObject/RemoteObject.cpp
+++ b/Engine/Source/Runtime/CoreUObject/Private/UObject/RemoteObject.cpp
@@ -38,7 +38,7 @@ int32 GRemoteIdToStringVerbosity = (int32)ERemoteIdToStringVerbosity::Id;
 static FAutoConsoleVariableRef CVarRemoteIdToStringVerbosity(
 	TEXT("ro.IdToStringVerbosity"),
 	GRemoteIdToStringVerbosity,
-	TEXT("Sets the verbosity FRemoteObjectId::ToString() prints the id to log with (1: id only, 2: id + pathname, 4: id + fullname, 5: id + fullname + attributes"),
+	TEXT("Sets the verbosity FRemoteObjectId::ToString() prints the id to log with (1: id only, 2: id + name, 3: id + pathname, 4: id + fullname, 5: id + fullname + attributes"),
 	FConsoleVariableDelegate::CreateLambda([](IConsoleVariable*)
 		{
 			if (GRemoteIdToStringVerbosity <= (int32)ERemoteIdToStringVerbosity::Default || GRemoteIdToStringVerbosity > (int32)ERemoteIdToStringVerbosity::Max)
