PackagePath: Add callstack to warning that calls out a code bug.
#rnx

diff --git a/Engine/Source/Runtime/CoreUObject/Private/Misc/PackagePath.cpp b/Engine/Source/Runtime/CoreUObject/Private/Misc/PackagePath.cpp
--- a/Engine/Source/Runtime/CoreUObject/Private/Misc/PackagePath.cpp
+++ b/Engine/Source/Runtime/CoreUObject/Private/Misc/PackagePath.cpp
@@ -4,6 +4,7 @@
 
 #include "Math/NumericLimits.h"
 #include "Misc/AssertionMacros.h"
+#include "Misc/CoreMiscDefines.h"
 #include "Misc/PackageName.h"
 #include "Misc/PackageSegment.h"
 #include "Misc/Paths.h"
@@ -811,6 +812,12 @@ void FPackagePath::AppendLocalFullPath(FStringBuilderBase& Builder, EPackageSegm
 					TEXT("GetLocalFullPath called on FPackagePath %s which has an unspecified header extension,")
 					TEXT(" and the path does not exist on disk. Assuming EPackageExtension::Asset."),
 					Builder.ToString() + BeforeLen); // Can't use GetDebugName since that can call this function recursively
+				// This warning is a bug in the calling code so we need a callstack to diagnose. But only write it on
+				// the first occurrence to prevent spam.
+				UE_CALL_ONCE([]()
+					{
+						FDebug::DumpStackTraceToLog(ELogVerbosity::Display);
+					});
 				Extension = EPackageExtension::Asset;
 			}
 		}
