UE::SavePackageUtilities::SaveThumbnails: don't go through all exports to gather objects with thumbnails if the PackageThumbnailMap map is empty.
Saves a lot of time when saving external entities containing lots of owned entities.

diff --git a/Engine/Source/Runtime/CoreUObject/Private/UObject/SavePackage/SavePackageUtilities.cpp b/Engine/Source/Runtime/CoreUObject/Private/UObject/SavePackage/SavePackageUtilities.cpp
--- a/Engine/Source/Runtime/CoreUObject/Private/UObject/SavePackage/SavePackageUtilities.cpp
+++ b/Engine/Source/Runtime/CoreUObject/Private/UObject/SavePackage/SavePackageUtilities.cpp
@@ -1319,36 +1319,40 @@ void SaveThumbnails(UPackage* InOuter, FLinkerSave* Linker, FStructuredArchive::
 		// for objects that are actually in the export map.  This is so that we avoid saving out
 		// thumbnails that were cached for deleted objects and such.
 		TArray< FObjectFullNameAndThumbnail > ObjectsWithThumbnails;
-		for( int32 i=0; i<Linker->ExportMap.Num(); i++ )
+
+		if (PackageThumbnailMap.Num())
 		{
-			FObjectExport& Export = Linker->ExportMap[i];
-			if( Export.Object )
+			for( int32 i=0; i<Linker->ExportMap.Num(); i++ )
 			{
-				const FName ObjectFullName( *Export.Object->GetFullName(), FNAME_Find );
-				const FObjectThumbnail* ObjectThumbnail = nullptr;
-				// If the FName does not exist, then we know it is not in the map and do not need to search
-				if (!ObjectFullName.IsNone())
+				FObjectExport& Export = Linker->ExportMap[i];
+				if( Export.Object )
 				{
-					ObjectThumbnail = PackageThumbnailMap.Find(ObjectFullName);
-				}
+					const FName ObjectFullName( *Export.Object->GetFullName(), FNAME_Find );
+					const FObjectThumbnail* ObjectThumbnail = nullptr;
+					// If the FName does not exist, then we know it is not in the map and do not need to search
+					if (!ObjectFullName.IsNone())
+					{
+						ObjectThumbnail = PackageThumbnailMap.Find(ObjectFullName);
+					}
 		
-				// if we didn't find the object via full name, try again with ??? as the class name, to support having
-				// loaded old packages without going through the editor (ie cooking old packages)
-				if (ObjectThumbnail == nullptr)
-				{
-					// can't overwrite ObjectFullName, so that we add it properly to the map
-					FName OldPackageStyleObjectFullName = FName(*FString::Printf(TEXT("??? %s"), *Export.Object->GetPathName()), FNAME_Find);
-					if (!OldPackageStyleObjectFullName.IsNone())
+					// if we didn't find the object via full name, try again with ??? as the class name, to support having
+					// loaded old packages without going through the editor (ie cooking old packages)
+					if (ObjectThumbnail == nullptr)
 					{
-						ObjectThumbnail = PackageThumbnailMap.Find(OldPackageStyleObjectFullName);
+						// can't overwrite ObjectFullName, so that we add it properly to the map
+						FName OldPackageStyleObjectFullName = FName(*FString::Printf(TEXT("??? %s"), *Export.Object->GetPathName()), FNAME_Find);
+						if (!OldPackageStyleObjectFullName.IsNone())
+						{
+							ObjectThumbnail = PackageThumbnailMap.Find(OldPackageStyleObjectFullName);
+						}
+					}
+					if( ObjectThumbnail != nullptr )
+					{
+						// IMPORTANT: We save all thumbnails here, even if they are a shared (empty) thumbnail!
+						// Empty thumbnails let us know that an asset is in a package without having to
+						// make a linker for it.
+						ObjectsWithThumbnails.Add( FObjectFullNameAndThumbnail( ObjectFullName, ObjectThumbnail ) );
 					}
-				}
-				if( ObjectThumbnail != nullptr )
-				{
-					// IMPORTANT: We save all thumbnails here, even if they are a shared (empty) thumbnail!
-					// Empty thumbnails let us know that an asset is in a package without having to
-					// make a linker for it.
-					ObjectsWithThumbnails.Add( FObjectFullNameAndThumbnail( ObjectFullName, ObjectThumbnail ) );
 				}
 			}
 		}
