[UObject] Add the asset's archive name to the error printout when a PropertyTag fails to load.

#rnx

diff --git a/Engine/Source/Runtime/CoreUObject/Private/UObject/PropertyTag.cpp b/Engine/Source/Runtime/CoreUObject/Private/UObject/PropertyTag.cpp
--- a/Engine/Source/Runtime/CoreUObject/Private/UObject/PropertyTag.cpp
+++ b/Engine/Source/Runtime/CoreUObject/Private/UObject/PropertyTag.cpp
@@ -586,7 +586,7 @@ void FPropertyTag::SerializeTaggedProperty(FStructuredArchive::FSlot Slot, FProp
 	const int64 EndOfProperty = UnderlyingArchive.Tell();
 	if (Size && (EndOfProperty - StartOfProperty != Size))
 	{
-		UE_LOG(LogClass, Error, TEXT("Failed loading tagged %s. Read %" INT64_FMT "B, expected %dB."), *GetFullNameSafe(Property), EndOfProperty - StartOfProperty, Size);
+		UE_LOG(LogClass, Error, TEXT("Failed loading tagged %s. Read %" INT64_FMT "B, expected %dB. Package: %s"), *GetFullNameSafe(Property), EndOfProperty - StartOfProperty, Size, *UnderlyingArchive.GetArchiveName());
 		UnderlyingArchive.Seek(StartOfProperty + Size);
 		Property->ClearValue(Value);
 	}
