[AsyncLoading2]

Only apply loose file localization when loading using loose cooked content in a non-editor process (-game or runtime)

#rnx

diff --git a/Engine/Source/Runtime/CoreUObject/Private/Serialization/AsyncLoading2.cpp b/Engine/Source/Runtime/CoreUObject/Private/Serialization/AsyncLoading2.cpp
--- a/Engine/Source/Runtime/CoreUObject/Private/Serialization/AsyncLoading2.cpp
+++ b/Engine/Source/Runtime/CoreUObject/Private/Serialization/AsyncLoading2.cpp
@@ -625,15 +625,18 @@ FName ApplyPackageNameRedirections(FName PackageName)
 }
 #endif
 
-/** Returns localized package name to load when running with -game. Otherwise returns the name unchanged. */
-FName ApplyLocalizationPackageNameRedirects(FName PackageName)
+/** Returns localized package name to load when running as -game or using loose cooked files in a runtime process. Otherwise returns the name unchanged. */
+FName ApplyLooseFileLocalizationPackageNameRedirects(FName PackageName)
 {
-	UE::Core::Private::FPlayInEditorLoadingScope PlayInEditorIDScope(INDEX_NONE);
-	FName RedirectedName = FPackageLocalizationManager::Get().FindLocalizedPackageName(PackageName);
-	if (!RedirectedName.IsNone())
+	if (!GIsEditor && FAsyncLoadingThreadSettings::Get().bLooseFileLoadingEnabled)
 	{
-		UE_LOG(LogStreaming, Verbose, TEXT("Applying localization redirection for package %s to package %s"), *PackageName.ToString(), *RedirectedName.ToString());
-		PackageName = RedirectedName;
+		UE::Core::Private::FPlayInEditorLoadingScope PlayInEditorIDScope(INDEX_NONE);
+		FName RedirectedName = FPackageLocalizationManager::Get().FindLocalizedPackageName(PackageName);
+		if (!RedirectedName.IsNone())
+		{
+			UE_LOG(LogStreaming, Verbose, TEXT("Applying localization redirection for package %s to package %s"), *PackageName.ToString(), *RedirectedName.ToString());
+			PackageName = RedirectedName;
+		}
 	}
 	return PackageName;
 }
@@ -5680,11 +5683,8 @@ bool FAsyncLoadingThread2::CreateAsyncPackagesFromQueue(FAsyncLoadingThreadState
 			PackageNameToLoad = ApplyPackageNameRedirections(PackageNameToLoad);
 			UPackageName = PackageNameToLoad;
 #endif
-			if (FAsyncLoadingThreadSettings::Get().bLooseFileLoadingEnabled)
-			{
-				// Localization should only update the package name to load
-				PackageNameToLoad = ApplyLocalizationPackageNameRedirects(PackageNameToLoad);
-			}
+			// Localization should only update the package name to load
+			PackageNameToLoad = ApplyLooseFileLocalizationPackageNameRedirects(PackageNameToLoad);
 
 			FPackageId PackageIdToLoad = FPackageId::FromName(PackageNameToLoad);
 			{
@@ -6588,9 +6588,8 @@ void FAsyncPackage2::ImportPackagesRecursiveInner(FAsyncLoadingThreadState2& Thr
 #endif
 
 		// Note, localization should apply after remapping from instancing has occurred
-		if (FAsyncLoadingThreadSettings::Get().bLooseFileLoadingEnabled)
 		{
-			FName LocalizedPackageNameToLoad = ApplyLocalizationPackageNameRedirects(ImportedPackageNameToLoad);
+			FName LocalizedPackageNameToLoad = ApplyLooseFileLocalizationPackageNameRedirects(ImportedPackageNameToLoad);
 			if (LocalizedPackageNameToLoad != ImportedPackageNameToLoad)
 			{
 				// Localization should only update the package name to load
