`MarkAsGarbage` -> `ClearGarbage` in the same transaction was broken.

diff --git a/Engine/Source/Programs/AutoRTFMTests/Private/AutoRTFMTests.cpp b/Engine/Source/Programs/AutoRTFMTests/Private/AutoRTFMTests.cpp
--- a/Engine/Source/Programs/AutoRTFMTests/Private/AutoRTFMTests.cpp
+++ b/Engine/Source/Programs/AutoRTFMTests/Private/AutoRTFMTests.cpp
@@ -151,7 +151,7 @@ int RunTests(int ArgC, const char* ArgV[])
 		
 		EPackageLoader GetSupportedLoaders() override
 		{
-			return EPackageLoader::LinkerLoad;
+			return EPackageLoader::Zen | EPackageLoader::LinkerLoad;
 		}
 	};
 
diff --git a/Engine/Source/Programs/AutoRTFMTests/Private/UObjectTests.cpp b/Engine/Source/Programs/AutoRTFMTests/Private/UObjectTests.cpp
--- a/Engine/Source/Programs/AutoRTFMTests/Private/UObjectTests.cpp
+++ b/Engine/Source/Programs/AutoRTFMTests/Private/UObjectTests.cpp
@@ -383,6 +383,98 @@ TEST_CASE("UObject.MarkAsReachable")
 	SetReachabilityAnalysisTimeLimit(Original);
 }
 
+TEST_CASE("UObject")
+{
+	SECTION("MarkAsGarbage")
+	{
+		UMyAutoRTFMTestObject* const Object = NewObject<UMyAutoRTFMTestObject>();
+
+		AutoRTFM::Testing::Commit([&]
+		{
+			Object->MarkAsGarbage();
+		});
+
+		REQUIRE(Object->HasAnyFlags(EObjectFlags::RF_MirroredGarbage));
+		REQUIRE(Object->HasAnyInternalFlags(EInternalObjectFlags::Garbage));
+	}
+	
+	SECTION("ClearGarbage")
+	{
+		UMyAutoRTFMTestObject* const Object = NewObject<UMyAutoRTFMTestObject>();
+		Object->MarkAsGarbage();
+		REQUIRE(Object->HasAnyFlags(EObjectFlags::RF_MirroredGarbage));
+		REQUIRE(Object->HasAnyInternalFlags(EInternalObjectFlags::Garbage));
+
+		AutoRTFM::Testing::Commit([&]
+		{
+			Object->ClearGarbage();
+		});
+		
+		REQUIRE(!Object->HasAnyFlags(EObjectFlags::RF_MirroredGarbage));
+		REQUIRE(!Object->HasAnyInternalFlags(EInternalObjectFlags::Garbage));
+	}
+	
+	SECTION("MarkAsGarbage + ClearGarbage")
+	{
+		UMyAutoRTFMTestObject* const Object = NewObject<UMyAutoRTFMTestObject>();
+
+		AutoRTFM::Testing::Commit([&]
+		{
+			Object->MarkAsGarbage();
+			Object->ClearGarbage();
+		});
+
+		REQUIRE(!Object->HasAnyFlags(EObjectFlags::RF_MirroredGarbage));
+		REQUIRE(!Object->HasAnyInternalFlags(EInternalObjectFlags::Garbage));
+	}
+	
+	SECTION("Originally Garbage -> MarkAsGarbage + Abort")
+	{
+		UMyAutoRTFMTestObject* const Object = NewObject<UMyAutoRTFMTestObject>();
+		Object->MarkAsGarbage();
+
+		// Check that we don't clear something as being garbage that was garbage before the transaction aborted!
+		AutoRTFM::Testing::Abort([&]
+		{
+			Object->MarkAsGarbage();
+			AutoRTFM::AbortTransaction();
+		});
+
+		REQUIRE(Object->HasAnyFlags(EObjectFlags::RF_MirroredGarbage));
+		REQUIRE(Object->HasAnyInternalFlags(EInternalObjectFlags::Garbage));
+	}
+	
+	SECTION("ClearGarbage + Abort")
+	{
+		UMyAutoRTFMTestObject* const Object = NewObject<UMyAutoRTFMTestObject>();
+
+		// Check that we don't mark something as garbage that wasn't garbage before the transaction aborted!
+		AutoRTFM::Testing::Abort([&]
+		{
+			Object->ClearGarbage();
+			AutoRTFM::AbortTransaction();
+		});
+
+		REQUIRE(!Object->HasAnyFlags(EObjectFlags::RF_MirroredGarbage));
+		REQUIRE(!Object->HasAnyInternalFlags(EInternalObjectFlags::Garbage));
+	}
+	
+	SECTION("MarkAsGarbage + Async + Abort")
+	{
+		UMyAutoRTFMTestObject* const Object = NewObject<UMyAutoRTFMTestObject>();
+		Object->SetInternalFlags(EInternalObjectFlags::Async);
+
+		// Check that the async flag is recovered if we abort!
+		AutoRTFM::Testing::Abort([&]
+		{
+			Object->MarkAsGarbage();
+			AutoRTFM::AbortTransaction();
+		});
+
+		REQUIRE(Object->HasAnyInternalFlags(EInternalObjectFlags::Async));
+	}
+}
+
 TEST_CASE("FUObjectAnnotationSparse.AddAnnotation")
 {
 	struct FTestAnnotation
diff --git a/Engine/Source/Runtime/CoreUObject/Public/UObject/UObjectArray.h b/Engine/Source/Runtime/CoreUObject/Public/UObject/UObjectArray.h
--- a/Engine/Source/Runtime/CoreUObject/Public/UObject/UObjectArray.h
+++ b/Engine/Source/Runtime/CoreUObject/Public/UObject/UObjectArray.h
@@ -242,21 +242,32 @@ public:
 	inline bool ThisThreadAtomicallyClearedFlag(EInternalObjectFlags FlagToClear)
 	{
 		checkf((int32(FlagToClear) & ~int32(EInternalObjectFlags_AllFlags)) == 0, TEXT("%d is not a valid internal flag value"), int32(FlagToClear));
-		bool Result = false;
+		bool bIChangedIt = false;
+
 		UE_AUTORTFM_OPEN
 		{
 			FlagToClear &= ~EInternalObjectFlags_ReachabilityFlags; // reachability flags can only be cleared by GC through *_ForGC functions
 			if (!!(FlagToClear & EInternalObjectFlags_RootFlags))
 			{
-				Result = ClearRootFlags(FlagToClear);
+				bIChangedIt = ClearRootFlags(FlagToClear);
 			}
 			else
 			{
-				Result = AtomicallyClearFlag_ForGC(FlagToClear);
+				bIChangedIt = AtomicallyClearFlag_ForGC(FlagToClear);
 			}
 		};
 
-		return Result;
+		if (bIChangedIt)
+		{
+			// If we abort we undo clearing the flags we just set because `bIChangedIt` being true tells us that we were the ones who cleared the flag.
+			// #jira SOL-8462: this could potentially undo flag changes made on a separate thread if an abort occurs.
+			UE_AUTORTFM_ONABORT(this, FlagToClear)
+			{
+				ThisThreadAtomicallySetFlag(FlagToClear);
+			};
+		}
+
+		return bIChangedIt;
 	}
 
 	/**
@@ -267,20 +278,32 @@ public:
 	inline bool ThisThreadAtomicallySetFlag(EInternalObjectFlags FlagToSet)
 	{
 		checkf((int32(FlagToSet) & ~int32(EInternalObjectFlags_AllFlags)) == 0, TEXT("%d is not a valid internal flag value"), int32(FlagToSet));
-		bool Result = false;
+		bool bIChangedIt = false;
+
 		UE_AUTORTFM_OPEN
 		{
 			FlagToSet &= ~EInternalObjectFlags_ReachabilityFlags; // reachability flags can only be cleared by GC through *_ForGC functions
 			if (!!(FlagToSet & EInternalObjectFlags_RootFlags))
 			{
-				Result = SetRootFlags(FlagToSet);
+				bIChangedIt = SetRootFlags(FlagToSet);
 			}
 			else
 			{
-				Result = AtomicallySetFlag_ForGC(FlagToSet);
+				bIChangedIt = AtomicallySetFlag_ForGC(FlagToSet);
 			}
 		};
-		return Result;
+
+		if (bIChangedIt)
+		{
+			// If we abort we undo setting the flags we just set because `bIChangedIt` being true tells us that we were the ones who set the flag.
+			// #jira SOL-8462: this could potentially undo flag changes made on a separate thread if an abort occurs.
+			UE_AUTORTFM_ONABORT(this, FlagToSet)
+			{
+				ThisThreadAtomicallyClearedFlag(FlagToSet);
+			};
+		}
+		
+		return bIChangedIt;
 	}
 
 	UE_FORCEINLINE_HINT bool HasAnyFlags(EInternalObjectFlags InFlags) const
@@ -300,11 +323,11 @@ public:
 
 	UE_FORCEINLINE_HINT void SetGarbage()
 	{
-		AtomicallySetFlag_ForGC(EInternalObjectFlags::Garbage);
+		ThisThreadAtomicallySetFlag(EInternalObjectFlags::Garbage);
 	}
 	UE_FORCEINLINE_HINT void ClearGarbage()
 	{
-		AtomicallyClearFlag_ForGC(EInternalObjectFlags::Garbage);
+		ThisThreadAtomicallyClearedFlag(EInternalObjectFlags::Garbage);
 	}
 	UE_FORCEINLINE_HINT bool IsGarbage() const
 	{
diff --git a/Engine/Source/Runtime/CoreUObject/Public/UObject/UObjectBaseUtility.h b/Engine/Source/Runtime/CoreUObject/Public/UObject/UObjectBaseUtility.h
--- a/Engine/Source/Runtime/CoreUObject/Public/UObject/UObjectBaseUtility.h
+++ b/Engine/Source/Runtime/CoreUObject/Public/UObject/UObjectBaseUtility.h
@@ -177,20 +177,11 @@ public:
 	{
 		check(!IsRooted());
 
-		// We need to set the RF_MirroredGarbage flag before OnCommit because something in the same transaction may be checking this flag and AtomicallySetFlags() is already AutoRTFM aware.
-		// A common example used across many of the engine systems is IsValid(Object) which checks the RF_MirroredGarbage flag.
-		// However both EInternalObjectFlags::Async and EInternalObjectFlags::Garbage flags are only checked by GC so it's safe to clear/set them OnCommit.
 		AtomicallySetFlags(RF_MirroredGarbage);
+		GUObjectArray.IndexToObject(InternalIndex)->SetGarbage();
 
-		// only set these flags after we have finished our transaction
-		// GC should not happen during a transaction so delaying this wont cause GC issues
-		UE_AUTORTFM_ONCOMMIT(this)
-		{
-			// If we explicitly marked the object as garbage, remove the async flag so it's visible to the GC
-			AtomicallyClearInternalFlags(EInternalObjectFlags::Async);
-
-			GUObjectArray.IndexToObject(InternalIndex)->SetGarbage();
-		};
+		// If we explicitly marked the object as garbage, remove the async flag so it's visible to the GC
+		AtomicallyClearInternalFlags(EInternalObjectFlags::Async);
 	}
 
 	/**
diff --git a/Engine/Source/Runtime/Engine/Private/Tests/AutoRTFM/AutoRTFMActorComponentTests.cpp b/Engine/Source/Runtime/Engine/Private/Tests/AutoRTFM/AutoRTFMActorComponentTests.cpp
--- a/Engine/Source/Runtime/Engine/Private/Tests/AutoRTFM/AutoRTFMActorComponentTests.cpp
+++ b/Engine/Source/Runtime/Engine/Private/Tests/AutoRTFM/AutoRTFMActorComponentTests.cpp
@@ -379,6 +379,8 @@ AUTORTFM_ACTOR_COMPONENT_TEST(ChildActor)
 
 	TestTrueExpr(AutoRTFM::ETransactionResult::Committed == Result);
 	TestTrueExpr(nullptr != ChildActorComponent->GetChildActor());
+	TestFalseExpr(ChildActorComponent->GetChildActor()->HasAnyFlags(EObjectFlags::RF_MirroredGarbage));
+	TestFalseExpr(ChildActorComponent->GetChildActor()->HasAnyInternalFlags(EInternalObjectFlags::Garbage));
 
 	Result = AutoRTFM::Transact([&]
 		{
