Fix that transactionally loaded async packages get the correct `CustomPackageName` if provided.

diff --git a/Engine/Source/Programs/AutoRTFMTests/Private/UPackageTests.cpp b/Engine/Source/Programs/AutoRTFMTests/Private/UPackageTests.cpp
--- a/Engine/Source/Programs/AutoRTFMTests/Private/UPackageTests.cpp
+++ b/Engine/Source/Programs/AutoRTFMTests/Private/UPackageTests.cpp
@@ -11,6 +11,14 @@
 
 #include "Catch2Includes.h"
 
+static void TrashPackage(UPackage* const Package)
+{
+	// TODO: If we could move the trashing logic into `UPackage` we could just call that here?
+	const FName NewName = MakeUniqueObjectName(nullptr, UPackage::StaticClass(), NAME_TrashedPackage);
+	Package->Rename(*NewName.ToString(), nullptr, REN_DontCreateRedirectors | REN_NonTransactional | REN_DoNotDirty);
+	Package->SetFlags(RF_Transient);
+}
+
 TEST_CASE("UPackage.SetPackageFlagsTo")
 {
 	SECTION("Commit")
@@ -407,10 +415,7 @@ TEST_CASE("UPackage.AsyncLoading")
 		UPackage* const Package = LoadPackage(nullptr, *Name, LOAD_None);
 		REQUIRE(Package);
 			
-		// TODO: If we could move the trashing logic into `UPackage` we could just call that here?
-		const FName NewName = MakeUniqueObjectName(nullptr, UPackage::StaticClass(), NAME_TrashedPackage);
-		Package->Rename(*NewName.ToString(), nullptr, REN_DontCreateRedirectors | REN_NonTransactional | REN_DoNotDirty);
-		Package->SetFlags(RF_Transient);
+		TrashPackage(Package);
 
 		AutoRTFM::Testing::Commit([&]
 		{
@@ -453,10 +458,7 @@ TEST_CASE("UPackage.AsyncLoading")
 			UPackage* const Package = LoadPackage(nullptr, *Name, LOAD_None);
 			REQUIRE(Package);
 			
-			// TODO: If we could move the trashing logic into `UPackage` we could just call that here?
-			const FName NewName = MakeUniqueObjectName(nullptr, UPackage::StaticClass(), NAME_TrashedPackage);
-			Package->Rename(*NewName.ToString(), nullptr, REN_DontCreateRedirectors | REN_NonTransactional | REN_DoNotDirty);
-			Package->SetFlags(RF_Transient);
+			TrashPackage(Package);
 
 			UPackage* const ReloadedPackage = LoadPackage(nullptr, *Name, LOAD_None);
 			REQUIRE(ReloadedPackage);
@@ -478,10 +480,7 @@ TEST_CASE("UPackage.AsyncLoading")
 			UPackage* const FoundPackage = FindObject<UPackage>(nullptr, *Name);
 			REQUIRE(Package == FoundPackage);
 			
-			// TODO: If we could move the trashing logic into `UPackage` we could just call that here?
-			const FName NewName = MakeUniqueObjectName(nullptr, UPackage::StaticClass(), NAME_TrashedPackage);
-			Package->Rename(*NewName.ToString(), nullptr, REN_DontCreateRedirectors | REN_NonTransactional | REN_DoNotDirty);
-			Package->SetFlags(RF_Transient);
+			TrashPackage(Package);
 
 			UPackage* const ReloadedPackage = LoadPackage(nullptr, *Name, LOAD_None);
 			REQUIRE(ReloadedPackage);
@@ -542,11 +541,8 @@ TEST_CASE("UPackage.AsyncLoading")
 			{
 				UPackage* const Package = LoadPackage(nullptr, *Name, LOAD_None);
 				REQUIRE(Package);
-			
-				// TODO: If we could move the trashing logic into `UPackage` we could just call that here?
-				const FName NewName = MakeUniqueObjectName(nullptr, UPackage::StaticClass(), NAME_TrashedPackage);
-				Package->Rename(*NewName.ToString(), nullptr, REN_DontCreateRedirectors | REN_NonTransactional | REN_DoNotDirty);
-				Package->SetFlags(RF_Transient);
+
+				TrashPackage(Package);
 
 				UPackage* const ReloadedPackage = LoadPackage(nullptr, *Name, LOAD_None);
 				REQUIRE(ReloadedPackage);
@@ -578,4 +574,98 @@ TEST_CASE("UPackage.AsyncLoading")
 			REQUIRE(FindObject<UPackage>(nullptr, *Name2) != nullptr);
 		});
 	}
+
+	SECTION("Load With Custom Name")
+	{
+		FString Name(FString::Printf(TEXT("/Game/%dAutoRTFMTestPackage%d"), FPlatformProcess::GetCurrentProcessId(), __LINE__));
+		FPackageScopedMaker _(Name);
+
+		AutoRTFM::Testing::Commit([&]
+		{
+			FName CustomName("Wowwee");
+
+			FLoadPackageAsyncDelegate Delegate;
+
+			Delegate.BindLambda([&](const FName& Name, UPackage* const Package, EAsyncLoadingResult::Type Result)
+			{
+				REQUIRE(EAsyncLoadingResult::Succeeded == Result);
+				REQUIRE(Package->GetFName() == Name);
+				REQUIRE(CustomName == Name);
+			});
+
+			FlushAsyncLoading(LoadPackageAsync(FPackagePath::FromPackageNameUnchecked(*Name), CustomName, Delegate));
+		});
+	}
+
+	SECTION("Trash And Load With Custom Name")
+	{
+		FString Name(FString::Printf(TEXT("/Game/%dAutoRTFMTestPackage%d"), FPlatformProcess::GetCurrentProcessId(), __LINE__));
+		FPackageScopedMaker _(Name);
+
+		bool bHit = false;
+
+		UPackage* const Package = LoadPackage(nullptr, *Name, LOAD_None);
+		REQUIRE(Package);
+
+		AutoRTFM::Testing::Commit([&]
+		{
+			TrashPackage(Package);
+
+			FName CustomName("Wowwee");
+
+			FLoadPackageAsyncDelegate Delegate;
+
+			Delegate.BindLambda([&](const FName& Name, UPackage* const Package, EAsyncLoadingResult::Type Result)
+			{
+				REQUIRE(AutoRTFM::IsClosed());
+				REQUIRE(EAsyncLoadingResult::Succeeded == Result);
+				REQUIRE(Package->GetFName() == Name);
+				REQUIRE(CustomName == Name);
+				bHit = true;
+			});
+
+			FlushAsyncLoading(LoadPackageAsync(FPackagePath::FromPackageNameUnchecked(*Name), CustomName, Delegate));
+		});
+
+		REQUIRE(bHit);
+	}
+
+	SECTION("First Time Load Package, Second Time Don't")
+	{
+		FString Name(FString::Printf(TEXT("/Game/%dAutoRTFMTestPackage%d"), FPlatformProcess::GetCurrentProcessId(), __LINE__));
+		FPackageScopedMaker _(Name);
+
+		bool bHit = false;
+
+		bool bFirst = true;
+
+		FName CustomName("Wowwee");
+
+		FLoadPackageAsyncDelegate Delegate;
+
+		Delegate.BindLambda([&](const FName& Name, UPackage* const Package, EAsyncLoadingResult::Type Result)
+		{
+			REQUIRE(!AutoRTFM::IsTransactional());
+			REQUIRE(EAsyncLoadingResult::Succeeded == Result);
+			REQUIRE(Package->GetFName() == Name);
+			REQUIRE(CustomName == Name);
+			bHit = true;
+		});
+
+		AutoRTFM::Testing::Commit([&]
+		{
+			if (!bFirst)
+			{
+				return;
+			}
+
+			AutoRTFM::OnComplete([&] { bFirst = false; });
+
+			FlushAsyncLoading(LoadPackageAsync(FPackagePath::FromPackageNameUnchecked(*Name), CustomName, Delegate));
+		});
+
+		FlushAsyncLoading(LoadPackageAsync(FPackagePath::FromPackageNameUnchecked(*Name), CustomName, Delegate));
+
+		REQUIRE(bHit);
+	}
 }
diff --git a/Engine/Source/Runtime/CoreUObject/Private/Serialization/TransactionallySafeAsyncLoading.cpp b/Engine/Source/Runtime/CoreUObject/Private/Serialization/TransactionallySafeAsyncLoading.cpp
--- a/Engine/Source/Runtime/CoreUObject/Private/Serialization/TransactionallySafeAsyncLoading.cpp
+++ b/Engine/Source/Runtime/CoreUObject/Private/Serialization/TransactionallySafeAsyncLoading.cpp
@@ -122,6 +122,11 @@ public:
 			const FLinkerInstancingContext* InInstancingContext,
 			uint32 InLoadFlags) override
 	{
+		if (CustomPackageName.IsNone())
+		{
+			CustomPackageName = PackagePath.GetPackageFName();
+		}
+
 		if (AutoRTFM::IsClosed())
 		{
 			return AutoRTFM::Open([&]
@@ -214,11 +219,11 @@ public:
 							// in `FindObject` calls.
 							if (Package)
 							{
-								Package->Rename(*(MapPayload->PackageName).ToString(), nullptr, REN_DontCreateRedirectors | REN_NonTransactional | REN_DoNotDirty);
+								Package->Rename(*CustomPackageName.ToString(), nullptr, REN_DontCreateRedirectors | REN_NonTransactional | REN_DoNotDirty);
 							}
 						}
 
-						InCompletionDelegate.ExecuteIfBound(MapPayload->PackageName, Package, MapPayload->Result);
+						InCompletionDelegate.ExecuteIfBound(CustomPackageName, Package, MapPayload->Result);
 					});
 
 					return PackageCacheHit;
@@ -227,7 +232,6 @@ public:
 		}
 		else
 		{
-			AutoRTFM::UnreachableIfTransactional(UnreachableMessage);
 			return LoadAndCachePackage(PackagePath,
 				CustomPackageName,
 				InCompletionDelegate,
@@ -303,7 +307,7 @@ public:
 							WrappedPackageLoader->FlushLoading({});
 						});
 					});
-					ensure(Status == AutoRTFM::EContextStatus::OnTrack);
+					check(Status == AutoRTFM::EContextStatus::OnTrack);
 					return;
 				}
 
@@ -355,15 +359,16 @@ public:
 
 						AutoRTFM::OnComplete([this, RequestIdsClone = ::MoveTemp(RequestIdsClone)]() mutable
 						{
+							check(AutoRTFM::IsRetrying());
 							UE_LOG(LogStreaming, Display, TEXT("A call to `FlushLoading` that is flushing non-cached packages is causing Verse to abort and retry."));
-							DrainTransactionallyDeferredLoadPackages(false, &RequestIdsClone, true);
+							DrainTransactionallyDeferredLoadPackages(&RequestIdsClone, true);
 							UpdateMaxFlushedRequestId(RequestIdsClone);
 							WrappedPackageLoader->FlushLoading(RequestIdsClone);
 							UE_LOG(LogStreaming, Display, TEXT("`FlushLoading` has completed after Verse aborted, and we are now retrying."));
 						});
 						AutoRTFM::CascadingRetryTransaction();
 					});
-					ensure(Status == AutoRTFM::EContextStatus::AbortedByCascadingRetry);
+					check(Status == AutoRTFM::EContextStatus::AbortedByCascadingRetry);
 				}
 			});
 		}
@@ -501,12 +506,11 @@ private:
 	struct FMapPayload final
 	{
 		FMapPayload() = default;
-		FMapPayload(FName PackageName, UPackage* const LoadedPackage, EAsyncLoadingResult::Type Result) :
-			PackageName(PackageName), LoadedPackages({ MakeWeakObjectPtr(LoadedPackage) }), Result(Result) {}
+		FMapPayload(UPackage* const LoadedPackage, EAsyncLoadingResult::Type Result) :
+			LoadedPackages({ MakeWeakObjectPtr(LoadedPackage) }), Result(Result) {}
 
-		FName PackageName;
 		TArray<TWeakObjectPtr<UPackage>> LoadedPackages;
-		EAsyncLoadingResult::Type Result;
+		EAsyncLoadingResult::Type Result = EAsyncLoadingResult::Failed;
 
 		// Tells us whether a given map payload entry is still valid (EG. a cached loaded package didn't get GC'ed).
 		bool IsStillValid() const
@@ -523,21 +527,6 @@ private:
 					return false;
 				}
 
-				if (UNLIKELY(CVar_TransactionallySafeAsyncLoadingLegacyTrashedPackages))
-				{
-					return true;
-				}
-
-				static FString AutoRTFMUnusedPackage = FName(NAME_AutoRTFMUnusedPackage).ToString();
-				if (Package->GetFName().ToString().StartsWith(AutoRTFMUnusedPackage))
-				{
-					// Unused packages are ones that have been created during a transaction
-					// retry, and thus are valid!
-					return true;
-				}
-
-				ensure(Package->GetFName() == PackageName);
-
 				return true;
 			}
 			else
@@ -612,58 +601,8 @@ private:
 
 		FString PackagePathStr = PackagePath.GetDebugNameWithExtension();
 
-		UPackage* OldPackage = nullptr;
-
-		if (LIKELY(!CVar_TransactionallySafeAsyncLoadingLegacyTrashedPackages))
-		{
-			if (AutoRTFM::IsRetrying())
-			{
-				UE::TScopeLock _(Mutex);
-
-				if (FMapPayload* const Payload = PackageCache.Find(PackagePathStr))
-				{
-					for (int32 Index = 0; Index < Payload->LoadedPackages.Num(); Index++)
-					{
-						UPackage* const Package = Payload->LoadedPackages[Index].Get();
-
-						if (Package == nullptr)
-						{
-							continue;
-						}
-
-						// We should only find a single old-package with the original package name.
-						ensure(OldPackage == nullptr);
-
-						static FString AutoRTFMUnusedPackage = FName(NAME_AutoRTFMUnusedPackage).ToString();
-						static FString TrashedPackage = FName(NAME_TrashedPackage).ToString();
-
-						// Unused packages (transactionally created during a retry, but not yet in use) and
-						// trashed packages are both handled elsewhere and won't affect us creating a brand
-						// new package like we want, so skip them.
-						if (Package->GetFName().ToString().StartsWith(AutoRTFMUnusedPackage) ||
-							Package->GetFName().ToString().StartsWith(TrashedPackage))
-						{
-							continue;
-						}
-
-						// We need to rename the old package because we need to make a brand new object with
-						// the same name and have that be a unique object. This case can only occur when we've
-						// trashed a package in a transaction and then issued a sync-load in the same
-						// transaction to reload the object. When we abort-and-retry (and get to the current
-						// code location) any previously trashed packages will have had that trashing undone
-						// during the transaction abort, but we *know* that when we retry the transaction we
-						// are going to immediately trash the object and thus need a new unique object.
-						const FName NewName = MakeUniqueObjectName(nullptr, UPackage::StaticClass(), NAME_AutoRTFMUnusedPackage);
-						Package->Rename(*(NewName.ToString()), nullptr, REN_DontCreateRedirectors | REN_NonTransactional | REN_DoNotDirty);
-
-						OldPackage = Package;
-					}
-				}
-			}
-		}
-
 		FLoadPackageAsyncDelegate WrapperDelegate;
-		WrapperDelegate.BindLambda([this, PackagePathStr = PackagePathStr, CompletionDelegate = ::MoveTemp(InCompletionDelegate), OldPackage = OldPackage](
+		WrapperDelegate.BindLambda([this, PackagePathStr, CustomPackageName, CompletionDelegate = ::MoveTemp(InCompletionDelegate)](
 			const FName& Name, UPackage* const Package, EAsyncLoadingResult::Type Type)
 		{
 			if (EAsyncLoadingResult::Canceled != Type)
@@ -672,34 +611,18 @@ private:
 
 				if (FMapPayload* const Payload = PackageCache.Find(PackagePathStr))
 				{
+					// If we are not retrying it means we've done a package load outside of a transaction, so we want to just forget all previous loaded packages,
+					// and start from a blank array.
+					if (!AutoRTFM::IsRetrying())
+					{
+						Payload->LoadedPackages.Empty();
+					}
+
 					Payload->LoadedPackages.Add(Package);
 				}
 				else
 				{
-					PackageCache.FindOrAdd(PackagePathStr) = FMapPayload(Name, Package, Type);
-				}
-
-				if (LIKELY(!CVar_TransactionallySafeAsyncLoadingLegacyTrashedPackages))
-				{
-					if (AutoRTFM::IsRetrying() && Package)
-					{
-						// If we are retrying and managed to load a package, we need to rename the package from the actual given
-						// name to our special unused-package name. This is to prevent it appearing in `FindObject` calls *before*
-						// the sync-load is performed on the retry within the transaction.
-						const FString CreatedPackageName = Package->GetFName().ToString();
-
-						const FName NewName = MakeUniqueObjectName(nullptr, UPackage::StaticClass(), NAME_AutoRTFMUnusedPackage);
-						Package->Rename(*(NewName.ToString()), nullptr, REN_DontCreateRedirectors | REN_NonTransactional | REN_DoNotDirty);
-
-						// If `OldPackage` is a valid object it means we renamed it for the purposes of creating a new
-						// unused-as-of-yet package above. So we need to rename it back to the actual package name here
-						// so that when we begin the retry of the transaction the object has its correct name back and
-						// can be found in `FindObject` searches.
-						if (OldPackage)
-						{
-							OldPackage->Rename(*CreatedPackageName, nullptr, REN_DontCreateRedirectors | REN_NonTransactional | REN_DoNotDirty);
-						}
-					}
+					PackageCache.FindOrAdd(PackagePathStr) = FMapPayload(Package, Type);
 				}
 			}
 
@@ -716,15 +639,17 @@ private:
 			InLoadFlags);
 	}
 
-	void DrainTransactionallyDeferredLoadPackages(const bool bCallOriginalCompletionDelegate = true, TArray<int32>* RequestIds = nullptr, const bool bLogPackagePaths = false)
+	void DrainTransactionallyDeferredLoadPackages(TArray<int32>* RequestIds = nullptr, const bool bLogPackagePaths = false)
 	{
 		AutoRTFM::UnreachableIfTransactional(UnreachableMessage);
 
+		const bool bIsRetrying = AutoRTFM::IsRetrying();
+
 		for (const FTransactionallyDeferredLoadPackagePayload& Payload : TransactionallyDeferredLoadPackages)
 		{
 			const int32 RequestId = LoadAndCachePackage(Payload.PackagePath,
-				Payload.CustomPackageName,
-				bCallOriginalCompletionDelegate ? Payload.CompletionDelegate : FLoadPackageAsyncDelegate(),
+				bIsRetrying && LIKELY(!CVar_TransactionallySafeAsyncLoadingLegacyTrashedPackages) ? MakeUniqueObjectName(nullptr, UPackage::StaticClass(), NAME_AutoRTFMUnusedPackage) : Payload.CustomPackageName,
+				bIsRetrying ? FLoadPackageAsyncDelegate() : Payload.CompletionDelegate,
 				Payload.PackageFlags,
 				Payload.PIEInstanceID,
 				Payload.PackagePriority,
