Removed the new ToArray() from FPropertyTypeNameBuilder

An array of names without child counts can be ambiguous in some cases. Expose the complete underlying array through iteration as an alternative.

#rnx

diff --git a/Engine/Source/Runtime/CoreUObject/Private/UObject/OverriddenPropertySet.cpp b/Engine/Source/Runtime/CoreUObject/Private/UObject/OverriddenPropertySet.cpp
--- a/Engine/Source/Runtime/CoreUObject/Private/UObject/OverriddenPropertySet.cpp
+++ b/Engine/Source/Runtime/CoreUObject/Private/UObject/OverriddenPropertySet.cpp
@@ -335,7 +335,10 @@ FOverriddenPropertyNodeID::FOverriddenPropertyNodeID(TNotNull<const FProperty*>
 	}
 #endif
 	Path = Property->GetFName();
-	Path += TypeNameBuilder.ToArray();
+	for (const UE::FPropertyTypeNameNode& Node : TypeNameBuilder)
+	{
+		Path += Node.Name;
+	}
 }
 
 FOverriddenPropertyNodeID::FOverriddenPropertyNodeID(const FOverriddenPropertyNodeID& ParentNodeID, const FOverriddenPropertyNodeID& NodeID)
diff --git a/Engine/Source/Runtime/CoreUObject/Private/UObject/PropertyTypeName.cpp b/Engine/Source/Runtime/CoreUObject/Private/UObject/PropertyTypeName.cpp
--- a/Engine/Source/Runtime/CoreUObject/Private/UObject/PropertyTypeName.cpp
+++ b/Engine/Source/Runtime/CoreUObject/Private/UObject/PropertyTypeName.cpp
@@ -619,14 +619,4 @@ void FPropertyTypeNameBuilder::Reset()
 	ActiveIndex = INDEX_NONE;
 }
 
-TArray<FName> FPropertyTypeNameBuilder::ToArray() const
-{
-	TArray<FName> Names;
-	Names.Reserve(Nodes.Num());
-	for( const FPropertyTypeNameNode& Node : Nodes)
-	{
-		Names.Add(Node.Name);
-	}
-	return MoveTemp(Names);
-}
 } // UE
diff --git a/Engine/Source/Runtime/CoreUObject/Public/UObject/PropertyTypeName.h b/Engine/Source/Runtime/CoreUObject/Public/UObject/PropertyTypeName.h
--- a/Engine/Source/Runtime/CoreUObject/Public/UObject/PropertyTypeName.h
+++ b/Engine/Source/Runtime/CoreUObject/Public/UObject/PropertyTypeName.h
@@ -174,9 +174,6 @@ public:
 	/** Resets the builder to allow it to be reused. */
 	UE_API void Reset();
 
-	/** Return a TArray<FName> of the type */
-	UE_API TArray<FName> ToArray() const;
-
 	/**
 	 * Try to parse and add a complete type name with its type parameters.
 	 *
@@ -188,6 +185,9 @@ public:
 	[[nodiscard]] UE_API bool TryParse(FStringView Name);
 
 private:
+	[[nodiscard]] UE_REWRITE friend auto begin(const FPropertyTypeNameBuilder& Builder) { return Builder.Nodes.begin(); }
+	[[nodiscard]] UE_REWRITE friend auto end(const FPropertyTypeNameBuilder& Builder) { return Builder.Nodes.end(); }
+
 	TArray<FPropertyTypeNameNode, TInlineAllocator<10>> Nodes;
 	TArray<int32, TInlineAllocator<10>> OuterNodeIndex;
 	int32 ActiveIndex = INDEX_NONE;
