FAssetPackageData memory footprint reduction through re-ordering to reduce padding.
#rnx

diff --git a/Engine/Source/Runtime/CoreUObject/Public/AssetRegistry/AssetData.h b/Engine/Source/Runtime/CoreUObject/Public/AssetRegistry/AssetData.h
--- a/Engine/Source/Runtime/CoreUObject/Public/AssetRegistry/AssetData.h
+++ b/Engine/Source/Runtime/CoreUObject/Public/AssetRegistry/AssetData.h
@@ -948,6 +948,14 @@ public:
 	/** MD5 of the cooked package on disk, for tracking nondeterministic changes */
 	FMD5Hash CookedHash;
 
+	EPackageExtension Extension;
+
+private:
+	/** Bit storage for flags */
+	uint32 Flags;
+
+public:
+
 	/** 
 		The hash of all the chunks written by the package writer for this package. This is only available after cooking.
 		(Note, this only includes the ExportBundleData chunk if using ZenStoreWriter, as that chunk contains both the uasset and uexp
@@ -964,27 +972,26 @@ public:
 	/** UE file version that the package was saved with */
 	FPackageFileVersion FileVersionUE;
 
-	/** Licensee file version that the package was saved with */
-	int32 FileVersionLicenseeUE;
-
 private:
 #if WITH_EDITORONLY_DATA
 	FIoHash PackageSavedHash;
 #endif
-	UE::AssetRegistry::FPackageCustomVersionsHandle CustomVersions;
-	/** Bit storage for flags */
-	uint32 Flags;
 
 public:
-	EPackageExtension Extension;
+	/** Licensee file version that the package was saved with */
+	int32 FileVersionLicenseeUE;
+
+private:
+	UE::AssetRegistry::FPackageCustomVersionsHandle CustomVersions;
 
 public:
 
+
 	FAssetPackageData()
-		: DiskSize(0)
-		, FileVersionLicenseeUE(-1)
+		: Extension(EPackageExtension::Unspecified)
 		, Flags(0)
-		, Extension(EPackageExtension::Unspecified)
+		, DiskSize(0)
+		, FileVersionLicenseeUE(-1)
 	{
 	}
 
