Moving remote object construction overrides to a global stack to allow access to the overrides when constructing UObjects from random functions that may be called during UObject serialization

diff --git a/Engine/Source/Runtime/CoreUObject/Private/UObject/RemoteObjectSerialization.cpp b/Engine/Source/Runtime/CoreUObject/Private/UObject/RemoteObjectSerialization.cpp
--- a/Engine/Source/Runtime/CoreUObject/Private/UObject/RemoteObjectSerialization.cpp
+++ b/Engine/Source/Runtime/CoreUObject/Private/UObject/RemoteObjectSerialization.cpp
@@ -155,6 +155,28 @@ const FRemoteObjectConstructionParams* FRemoteObjectConstructionOverrides::Find(
 	return nullptr;
 }
 
+FRemoteObjectConstructionOverridesStack& FRemoteObjectConstructionOverridesStack::Get()
+{
+	checkf(IsInGameThread(), TEXT("Currently FRemoteObjectConstructionOverridesStack singleton can only be accessed from the game thread"));
+	static FRemoteObjectConstructionOverridesStack Singleton;
+	return Singleton;
+}
+
+FRemoteObjectConstructionOverridesStack::~FRemoteObjectConstructionOverridesStack()
+{
+	checkf(Stack.Num() == 0, TEXT("RemoteObjectConstructionOverridesStack still contains overrides"));
+}
+
+const FRemoteObjectConstructionParams* FRemoteObjectConstructionOverridesStack::Find(FName InName, UObject* InOuter) const
+{
+	const FRemoteObjectConstructionParams* Result = nullptr;
+	for (int32 StackIndex = Stack.Num() - 1; !Result && StackIndex >= 0; --StackIndex)
+	{
+		Result = Stack[StackIndex]->Find(InName, InOuter);
+	}
+	return Result;
+}
+
 enum class ERemoteReferenceType : uint8
 {
 	None = 0,
@@ -953,7 +975,7 @@ FRemoteObjectConstructionOverrides::FRemoteObjectConstructionOverrides(const TAr
 	}
 }
 
-UObject* ConstructRemoteObject(const FRemoteObjectHeader& Header, FRemoteObjectConstructionOverrides* ConstructionOverrides, ERemoteObjectSerializationFlags DeserializeFlags)
+UObject* ConstructRemoteObject(const FRemoteObjectHeader& Header, ERemoteObjectSerializationFlags DeserializeFlags)
 {
 	using namespace UE::RemoteObject::Handle;
 	using namespace UE::RemoteObject::Private;
@@ -1007,7 +1029,6 @@ UObject* ConstructRemoteObject(const FRemoteObjectHeader& Header, FRemoteObjectC
 		Params.Template = Header.Archetype.Get();
 #if UE_WITH_REMOTE_OBJECT_HANDLE
 		Params.RemoteId = Header.RemoteId;
-		Params.RemoteSubObjectOverrides = ConstructionOverrides;
 #endif
 
 		{
@@ -1278,12 +1299,14 @@ int32 DeserializeObjectData(FRemoteObjectData& ObjectData,
 			TRACE_CPUPROFILER_EVENT_SCOPE(ConstructRemoteObjects);
 
 			FRemoteObjectConstructionOverrides ConstructionOverrides(ObjectHeaders);
+			FRemoteObjectConstructionOverridesScope OverridesScope(ConstructionOverrides);
+
 			OutObjectRemoteIds.Reserve(ObjectHeaders.Num());
 			OutObjects.Reserve(ObjectHeaders.Num());
 			for (int32 ObjectIndex = 0; ObjectIndex < ObjectHeaders.Num(); ++ObjectIndex)
 			{
 				OutObjectRemoteIds.Add(ObjectHeaders[ObjectIndex].RemoteId);
-				OutObjects.Add(ConstructRemoteObject(ObjectHeaders[ObjectIndex], &ConstructionOverrides, DeserializeFlags));
+				OutObjects.Add(ConstructRemoteObject(ObjectHeaders[ObjectIndex], DeserializeFlags));
 			}
 		}
 
diff --git a/Engine/Source/Runtime/CoreUObject/Private/UObject/RemoteObjectSerialization.h b/Engine/Source/Runtime/CoreUObject/Private/UObject/RemoteObjectSerialization.h
--- a/Engine/Source/Runtime/CoreUObject/Private/UObject/RemoteObjectSerialization.h
+++ b/Engine/Source/Runtime/CoreUObject/Private/UObject/RemoteObjectSerialization.h
@@ -68,6 +68,58 @@ namespace UE::RemoteObject::Serialization
 		const FRemoteObjectConstructionParams* Find(FName InName, UObject* InOuter) const;
 	};
 
+	/**
+	* TLS singleton that stores the current stack of remote object construction overrides. 
+	* This singleton can be used to access construction overrides when constructing remote objects when they're being deserialized from a remote server
+	* or (in the future) when constructing objects from loaded packaeges with remote object ids baked in
+	*/
+	class FRemoteObjectConstructionOverridesStack final
+	{
+		friend TThreadSingleton<FRemoteObjectConstructionOverridesStack>;
+
+		TArray<const FRemoteObjectConstructionOverrides*> Stack;
+
+	public:
+
+		FRemoteObjectConstructionOverridesStack() = default;
+		virtual ~FRemoteObjectConstructionOverridesStack();
+
+		COREUOBJECT_API static FRemoteObjectConstructionOverridesStack& Get();
+
+		void Push(const FRemoteObjectConstructionOverrides& InOverrides)
+		{
+			Stack.Push(&InOverrides);
+		}
+		void Pop()
+		{
+			Stack.Pop(EAllowShrinking::No);
+		}
+		bool IsEmpty() const
+		{
+			return !Stack.Num();
+		}
+		COREUOBJECT_API const FRemoteObjectConstructionParams* Find(FName InName, UObject* InOuter) const;
+	};
+
+	/**
+	* Pushes remote object construction overrides onto the FRemoteObjectConstructionOverridesStack on construction and pops them on destruction. 
+	*/
+	class FRemoteObjectConstructionOverridesScope final
+	{
+		FRemoteObjectConstructionOverridesStack& Stack;
+
+	public:
+		explicit FRemoteObjectConstructionOverridesScope(FRemoteObjectConstructionOverrides& InOverrides)
+			: Stack(FRemoteObjectConstructionOverridesStack::Get())
+		{
+			Stack.Push(InOverrides);
+		}
+		~FRemoteObjectConstructionOverridesScope()
+		{
+			Stack.Pop();
+		}
+	};
+
 	/**
 	* Serializes an object and its subobject (or if the object is a default subobject, its parent and the parent's subobjects)
 	* @param InObject Object to be serialized
diff --git a/Engine/Source/Runtime/CoreUObject/Private/UObject/UObjectGlobals.cpp b/Engine/Source/Runtime/CoreUObject/Private/UObject/UObjectGlobals.cpp
--- a/Engine/Source/Runtime/CoreUObject/Private/UObject/UObjectGlobals.cpp
+++ b/Engine/Source/Runtime/CoreUObject/Private/UObject/UObjectGlobals.cpp
@@ -3991,16 +3991,13 @@ FObjectInitializer::FObjectInitializer()
 	Construct_Internal();
 }	
 
-FObjectInitializer::FObjectInitializer(UObject* InObj, const FStaticConstructObjectParameters& StaticConstructParams, UE::RemoteObject::Serialization::FRemoteObjectConstructionOverrides* InRemoteSubObjectOverides)
+FObjectInitializer::FObjectInitializer(UObject* InObj, const FStaticConstructObjectParameters& StaticConstructParams)
 	: Obj(InObj)
 	, ObjectArchetype(StaticConstructParams.Template)
 	, bCopyTransientsFromClassDefaults(StaticConstructParams.bCopyTransientsFromClassDefaults)
 	, bShouldInitializePropsFromArchetype(true)
 	, InstanceGraph(StaticConstructParams.InstanceGraph)
 	, PropertyInitCallback(StaticConstructParams.PropertyInitCallback)
-#if UE_WITH_REMOTE_OBJECT_HANDLE
-	, RemoteSubObjectOverrides(InRemoteSubObjectOverides)
-#endif
 {
 	if (StaticConstructParams.SubobjectOverrides)
 	{
@@ -4877,38 +4874,31 @@ FStaticConstructObjectParameters::FStaticConstructObjectParameters(const UClass*
 #if UE_WITH_REMOTE_OBJECT_HANDLE
 namespace
 {
-	UE::RemoteObject::Serialization::FRemoteObjectConstructionOverrides* GetRemoteSubObjectOverridesFromObjectInitializer(FName InName, UObject* InOuter, int32& InOutSerialNumber, FRemoteObjectId& InOutRemoteId)
+	FORCENOINLINE void TryGetRemoteObjectOverrides(FName InName, UObject* InOuter, int32& InOutSerialNumber, FRemoteObjectId& InOutRemoteId)
 	{
 		using namespace UE::RemoteObject::Serialization;
 
-		FUObjectThreadContext& ThreadContext = FUObjectThreadContext::Get();
-		if (ThreadContext.IsInConstructor)
+		FRemoteObjectConstructionOverridesStack& OverridesStack = FRemoteObjectConstructionOverridesStack::Get();
+		if (!OverridesStack.IsEmpty())
 		{
-			// Try to get remote subobject overrides from FObjectInitializer so that if the object that's about to be constructed is a subobject
+			// Try to get remote subobject overrides so that if the object that's about to be constructed is a subobject
 			// it gets the right RemoteId and SerialNumber at the time of its construction because its owner constructor may be reading the id and/or 
 			// creating a weak pointer to the subobject (if not the owner directly maybe some code that gets called from the owner's constructor)
 
-			FObjectInitializer& ObjectInitializer = ThreadContext.TopInitializerChecked();
-			if (FRemoteObjectConstructionOverrides* RemoteSubObjectOverrides = ObjectInitializer.GetRemoteSubObjectOverrides())
+			if (const FRemoteObjectConstructionParams* ConstructionOverrides = OverridesStack.Find(InName, InOuter))
 			{
-				if (const FRemoteObjectConstructionParams* ConstructionOverrides = RemoteSubObjectOverrides->Find(InName, InOuter))
+				if (ConstructionOverrides->SerialNumber != 0)
 				{
-					if (ConstructionOverrides->SerialNumber != 0)
-					{
-						checkf(InOutSerialNumber == 0 || InOutSerialNumber == ConstructionOverrides->SerialNumber, TEXT("Attempting to override valid SerialNumber of %s.%s (%d -> %d)"), *GetPathNameSafe(InOuter), *InName.ToString(), InOutSerialNumber, ConstructionOverrides->SerialNumber);
-						InOutSerialNumber = ConstructionOverrides->SerialNumber;
-					}
-					if (ConstructionOverrides->RemoteId.IsValid())
-					{
-						checkf(!InOutRemoteId.IsValid() || InOutRemoteId == ConstructionOverrides->RemoteId, TEXT("Attempting to override valid RemoteId of %s.%s (%s -> %s)"), *GetPathNameSafe(InOuter), *InName.ToString(), *InOutRemoteId.ToString(), *ConstructionOverrides->RemoteId.ToString());
-						InOutRemoteId = ConstructionOverrides->RemoteId;
-					}
+					checkf(InOutSerialNumber == 0 || InOutSerialNumber == ConstructionOverrides->SerialNumber, TEXT("Attempting to override valid SerialNumber of %s.%s (%d -> %d)"), *GetPathNameSafe(InOuter), *InName.ToString(), InOutSerialNumber, ConstructionOverrides->SerialNumber);
+					InOutSerialNumber = ConstructionOverrides->SerialNumber;
+				}
+				if (ConstructionOverrides->RemoteId.IsValid())
+				{
+					checkf(!InOutRemoteId.IsValid() || InOutRemoteId == ConstructionOverrides->RemoteId, TEXT("Attempting to override valid RemoteId of %s.%s (%s -> %s)"), *GetPathNameSafe(InOuter), *InName.ToString(), *InOutRemoteId.ToString(), *ConstructionOverrides->RemoteId.ToString());
+					InOutRemoteId = ConstructionOverrides->RemoteId;
 				}
-				return RemoteSubObjectOverrides;
 			}
-
 		}
-		return nullptr;
 	}
 }
 #endif // UE_WITH_REMOTE_OBJECT_HANDLE
@@ -4922,7 +4912,6 @@ UObject* StaticConstructObject_Internal(const FStaticConstructObjectParameters&
 	UObject* InTemplate = Params.Template;
 	int32 SerialNumber = Params.SerialNumber;
 	FRemoteObjectId RemoteId;
-	UE::RemoteObject::Serialization::FRemoteObjectConstructionOverrides* RemoteSubObjectOverrides = nullptr;
 
 	LLM_SCOPE(ELLMTag::UObject);
 	LLM_SCOPE_BYTAG(UObject_StaticConstructObjectInternal);
@@ -4952,11 +4941,7 @@ UObject* StaticConstructObject_Internal(const FStaticConstructObjectParameters&
 	
 #if UE_WITH_REMOTE_OBJECT_HANDLE
 	RemoteId = Params.RemoteId;
-	RemoteSubObjectOverrides = Params.RemoteSubObjectOverrides;
-	if (!RemoteSubObjectOverrides)
-	{
-		RemoteSubObjectOverrides = GetRemoteSubObjectOverridesFromObjectInitializer(InName, InOuter, SerialNumber, RemoteId);
-	}
+	TryGetRemoteObjectOverrides(InName, InOuter, SerialNumber, RemoteId);
 #endif
 
 	FGCReconstructionGuard GCGuard;
@@ -4967,7 +4952,7 @@ UObject* StaticConstructObject_Internal(const FStaticConstructObjectParameters&
 	if (!bRecycledSubobject)
 	{		
 		STAT(FScopeCycleCounterUObject ConstructorScope(InClass->GetFName().IsNone() ? nullptr : InClass, GET_STATID(STAT_ConstructObject)));
-		(*InClass->ClassConstructor)(FObjectInitializer(Result, Params, RemoteSubObjectOverrides));
+		(*InClass->ClassConstructor)(FObjectInitializer(Result, Params));
 	}
 	// StaticAllocateObject might have locked GCGuard but it can only be unlocked after the object has been fully constructed so unlock it here
 	GCGuard.Unlock();
diff --git a/Engine/Source/Runtime/CoreUObject/Public/UObject/UObjectGlobals.h b/Engine/Source/Runtime/CoreUObject/Public/UObject/UObjectGlobals.h
--- a/Engine/Source/Runtime/CoreUObject/Public/UObject/UObjectGlobals.h
+++ b/Engine/Source/Runtime/CoreUObject/Public/UObject/UObjectGlobals.h
@@ -1306,7 +1306,13 @@ public:
 	COREUOBJECT_API FObjectInitializer(UObject* InObj, UObject* InObjectArchetype, EObjectInitializerOptions InOptions, struct FObjectInstancingGraph* InInstanceGraph = nullptr);
 
 	/** Special constructor for static construct object internal that passes along the params block directly */
-	COREUOBJECT_API FObjectInitializer(UObject* InObj, const FStaticConstructObjectParameters& StaticConstructParams, UE::RemoteObject::Serialization::FRemoteObjectConstructionOverrides* RemoteObjectOverides = nullptr);
+	COREUOBJECT_API FObjectInitializer(UObject* InObj, const FStaticConstructObjectParameters& StaticConstructParams);
+
+	UE_DEPRECATED(5.7, "FObjectInitializer constructor that takes remote object overrides is deprecated. Use FRemoteObjectConstructionOverridesScope instead.")
+	FObjectInitializer(UObject* InObj, const FStaticConstructObjectParameters& StaticConstructParams, UE::RemoteObject::Serialization::FRemoteObjectConstructionOverrides* RemoteObjectOverides)
+		: FObjectInitializer(InObj, StaticConstructParams)
+	{
+	}
 
 private:
 	/** Helper for the common behaviors in the constructors */
@@ -1735,13 +1741,11 @@ private:
 #endif // WITH_EDITORONLY_DATA
 
 #if UE_WITH_REMOTE_OBJECT_HANDLE
-	/** Existing remote id for the object */
-	UE::RemoteObject::Serialization::FRemoteObjectConstructionOverrides* RemoteSubObjectOverrides = nullptr;
-
 public:
+	UE_DEPRECATED(5.7, "Use FRemoteObjectConstructionOverridesScope instead.")
 	UE::RemoteObject::Serialization::FRemoteObjectConstructionOverrides* GetRemoteSubObjectOverrides() const
 	{
-		return RemoteSubObjectOverrides;
+		return nullptr;
 	}
 #endif // UE_WITH_REMOTE_OBJECT_HANDLE
 };
@@ -1795,6 +1799,7 @@ struct FStaticConstructObjectParameters
 	/** Existing remote id for the object */
 	FRemoteObjectId RemoteId;
 
+	UE_DEPRECATED(5.7, "Use FRemoteObjectConstructionOverridesScope instead.")
 	UE::RemoteObject::Serialization::FRemoteObjectConstructionOverrides* RemoteSubObjectOverrides = nullptr;
 #endif // UE_WITH_REMOTE_OBJECT_HANDLE
 
