Removing unnecessary virtual and fixing comment (leftover from the initial version of my change)

diff --git a/Engine/Source/Runtime/CoreUObject/Private/UObject/RemoteObjectSerialization.h b/Engine/Source/Runtime/CoreUObject/Private/UObject/RemoteObjectSerialization.h
--- a/Engine/Source/Runtime/CoreUObject/Private/UObject/RemoteObjectSerialization.h
+++ b/Engine/Source/Runtime/CoreUObject/Private/UObject/RemoteObjectSerialization.h
@@ -69,9 +69,10 @@ namespace UE::RemoteObject::Serialization
 	};
 
 	/**
-	* TLS singleton that stores the current stack of remote object construction overrides. 
+	* Singleton that stores the current stack of remote object construction overrides. 
 	* This singleton can be used to access construction overrides when constructing remote objects when they're being deserialized from a remote server
-	* or (in the future) when constructing objects from loaded packaeges with remote object ids baked in
+	* or (in the future) when constructing objects from loaded packaeges with remote object ids baked in.
+	* Note that atm this is a game-thread only object.
 	*/
 	class FRemoteObjectConstructionOverridesStack final
 	{
@@ -82,7 +83,7 @@ namespace UE::RemoteObject::Serialization
 	public:
 
 		FRemoteObjectConstructionOverridesStack() = default;
-		virtual ~FRemoteObjectConstructionOverridesStack();
+		~FRemoteObjectConstructionOverridesStack();
 
 		COREUOBJECT_API static FRemoteObjectConstructionOverridesStack& Get();
 
