Comment typo fix + remove unnecessary friend declaration

diff --git a/Engine/Source/Runtime/CoreUObject/Private/UObject/RemoteObjectSerialization.h b/Engine/Source/Runtime/CoreUObject/Private/UObject/RemoteObjectSerialization.h
--- a/Engine/Source/Runtime/CoreUObject/Private/UObject/RemoteObjectSerialization.h
+++ b/Engine/Source/Runtime/CoreUObject/Private/UObject/RemoteObjectSerialization.h
@@ -71,13 +71,11 @@ namespace UE::RemoteObject::Serialization
 	/**
 	* Singleton that stores the current stack of remote object construction overrides. 
 	* This singleton can be used to access construction overrides when constructing remote objects when they're being deserialized from a remote server
-	* or (in the future) when constructing objects from loaded packaeges with remote object ids baked in.
+	* or (in the future) when constructing objects from loaded packages with remote object ids baked in.
 	* Note that atm this is a game-thread only object.
 	*/
 	class FRemoteObjectConstructionOverridesStack final
 	{
-		friend TThreadSingleton<FRemoteObjectConstructionOverridesStack>;
-
 		TArray<const FRemoteObjectConstructionOverrides*> Stack;
 
 	public:
