[BugFix] FPropertyPath::ToPropertyPathName was treating containers slightly wrong
[BugFix] Changed FPropertyPathNames to use internal indices for sets and maps rather than logical ones.
This is less bug prone and matches the format of FPropertyPath
[Feature] Added FPropertyPathNameTree::Remove
[Feature] Added FPropertyPathName::FromJsonValue and FPropertyPathName::ToJsonValue to facilitate saving FPropertyPaths in IDT config files

diff --git a/Engine/Source/Editor/PropertyEditor/Private/PropertyPath.cpp b/Engine/Source/Editor/PropertyEditor/Private/PropertyPath.cpp
--- a/Engine/Source/Editor/PropertyEditor/Private/PropertyPath.cpp
+++ b/Engine/Source/Editor/PropertyEditor/Private/PropertyPath.cpp
@@ -4,6 +4,7 @@
 
 #include "UObject/PropertyPathName.h"
 #include "UObject/PropertyTypeName.h"
+#include "UObject/PropertyPathFunctions.h"
 
 TArray<UE::FPropertyPathName> FPropertyPath::ToPropertyPathName() const
 {
@@ -12,11 +13,42 @@ TArray<UE::FPropertyPathName> FPropertyPath::ToPropertyPathName() const
 	TArray<FPropertyPathName> PathNames;
 	FPropertyPathName PathName;
 
-	for (const FPropertyInfo& Info : Properties)
+	for (int32 PropIndex = 0; PropIndex < Properties.Num(); ++PropIndex)
 	{
+		const FPropertyInfo& Info = Properties[PropIndex];
 		FPropertyTypeNameBuilder TypeName;
 		Info.Property->SaveTypeName(TypeName);
-		PathName.Push({Info.Property->GetFName(), TypeName.Build(), Info.ArrayIndex});
+		int32 ArrayIndex = Info.ArrayIndex;
+
+		if (Info.Property->IsA<FArrayProperty>() || Info.Property->IsA<FSetProperty>() || Info.Property->IsA<FMapProperty>())
+		{
+			// FPropertyPaths store the index of containers on separate entries but FPropertyPathNames don't
+			++PropIndex;
+
+			if (PropIndex < Properties.Num())
+			{
+				ArrayIndex = Properties[PropIndex].ArrayIndex;
+			}
+		}
+
+		PathName.Push({ Info.Property->GetFName(), TypeName.Build(), ArrayIndex });
+
+		if (PropIndex < Properties.Num())
+		{
+			if (FMapProperty* AsMapProperty = CastField<FMapProperty>(Info.Property.Get()))
+			{
+				// add either the key or value flag to the path
+				const FPropertyInfo& InnerInfo = Properties[PropIndex];
+				if (InnerInfo.Property.Get() == AsMapProperty->KeyProp)
+				{
+					PathName.Push({ UE::NAME_Key });
+				}
+				else if (InnerInfo.Property.Get() == AsMapProperty->ValueProp)
+				{
+					PathName.Push({ UE::NAME_Value });
+				}
+			}
+		}
 
 		if (Info.Property->IsA<FObjectPropertyBase>())
 		{
diff --git a/Engine/Source/Runtime/CoreUObject/Internal/UObject/PropertyPathNameTree.h b/Engine/Source/Runtime/CoreUObject/Internal/UObject/PropertyPathNameTree.h
--- a/Engine/Source/Runtime/CoreUObject/Internal/UObject/PropertyPathNameTree.h
+++ b/Engine/Source/Runtime/CoreUObject/Internal/UObject/PropertyPathNameTree.h
@@ -94,6 +94,9 @@ public:
 	/** Adds the path to the tree. Keeps any existing nodes that match both name and type. */
 	UE_API FNode Add(const FPropertyPathName& Path, int32 StartIndex = 0);
 
+	/** Removes the path from the tree. This may create a parent with no children. If that parent also has no tag, remove it too */
+	UE_API bool Remove(const FPropertyPathName& Path, int32 StartIndex = 0);
+
 	/**
 	 * Finds the path within the tree.
 	 *
diff --git a/Engine/Source/Runtime/CoreUObject/Private/UObject/PropertyPathFunctions.cpp b/Engine/Source/Runtime/CoreUObject/Private/UObject/PropertyPathFunctions.cpp
--- a/Engine/Source/Runtime/CoreUObject/Private/UObject/PropertyPathFunctions.cpp
+++ b/Engine/Source/Runtime/CoreUObject/Private/UObject/PropertyPathFunctions.cpp
@@ -6,6 +6,7 @@
 #include "UObject/PropertyOptional.h"
 #include "UObject/PropertyPathName.h"
 #include "UObject/UnrealType.h"
+#include "Templates/FunctionWithContext.h"
 
 namespace UE
 {
@@ -34,7 +35,8 @@ inline static const UStruct* FindStructFromProperty(const FProperty* Property)
 	return nullptr;
 }
 
-FPropertyValueInContainer TryResolvePropertyPath(const FPropertyPathName& Path, UObject* Object)
+// call the OnSegment method for every property resolved in the path. If a segment couldn't be resolved, break out early.
+static FPropertyValueInContainer TryIteratePropertyPath(const FPropertyPathName& Path, UObject* Object, TFunctionWithContext<void(const FPropertyPathNameSegment&, FPropertyValueInContainer&)> OnSegment)
 {
 	const UStruct* NextStruct = Object ? Object->GetClass() : nullptr;
 	void* NextContainer = Object;
@@ -83,6 +85,11 @@ FPropertyValueInContainer TryResolvePropertyPath(const FPropertyPathName& Path,
 			NextContainer = OptionalHelper.GetValuePointerForReadOrReplaceIfSet(NextContainer);
 		}
 
+		if (OnSegment)
+		{
+			OnSegment(Segment, Value);
+		}
+
 		// Scalar values and static containers are finished resolving.
 		if (Property->ArrayDim > 1 || Segment.Index == INDEX_NONE)
 		{
@@ -108,20 +115,19 @@ FPropertyValueInContainer TryResolvePropertyPath(const FPropertyPathName& Path,
 		else if (const FSetProperty* SetProperty = CastField<FSetProperty>(Property))
 		{
 			FScriptSetHelper SetHelper(SetProperty, NextContainer);
-			if (!SetHelper.IsValidIndex(Segment.Index))
+			NextStruct = FindStructFromProperty(SetProperty->ElementProp);
+			NextContainer = SetHelper.FindNthElementPtr(Segment.Index);
+			if (!NextContainer)
 			{
 				return {};
 			}
-
-			NextStruct = FindStructFromProperty(SetProperty->ElementProp);
-			NextContainer = SetHelper.GetElementPtr(Segment.Index);
 			Value.Property = SetProperty->ElementProp;
 			Value.Container = NextContainer;
 		}
 		else if (const FMapProperty* MapProperty = CastField<FMapProperty>(Property))
 		{
 			FScriptMapHelper MapHelper(MapProperty, NextContainer);
-			if (!MapHelper.IsValidIndex(Segment.Index) || ++Index == Count)
+			if (++Index == Count)
 			{
 				return {};
 			}
@@ -136,13 +142,21 @@ FPropertyValueInContainer TryResolvePropertyPath(const FPropertyPathName& Path,
 			if (MapSegment.Name == NAME_Key)
 			{
 				NextStruct = FindStructFromProperty(MapProperty->KeyProp);
-				NextContainer = MapHelper.GetKeyPtr(Segment.Index);
+				NextContainer = MapHelper.FindNthKeyPtr(Segment.Index);
+				if (!NextContainer)
+				{
+					return {};
+				}
 				Value.Property = MapProperty->KeyProp;
 			}
 			else if (MapSegment.Name == NAME_Value)
 			{
 				NextStruct = FindStructFromProperty(MapProperty->ValueProp);
-				NextContainer = MapHelper.GetValuePtr(Segment.Index);
+				NextContainer = MapHelper.FindNthValuePtr(Segment.Index);
+				if (!NextContainer)
+				{
+					return {};
+				}
 				Value.Property = MapProperty->ValueProp;
 			}
 			else
@@ -151,14 +165,45 @@ FPropertyValueInContainer TryResolvePropertyPath(const FPropertyPathName& Path,
 			}
 
 			// The key and value property both have an offset relative to the pair.
-			Value.Container = MapHelper.GetPairPtr(Segment.Index);
+			Value.Container = MapHelper.FindNthPairPtr(Segment.Index);
 		}
 		else
 		{
 			return {};
 		}
+
+		if (OnSegment)
+		{
+			OnSegment(Segment, Value);
+		}
 	}
 	return Value;
 }
 
+FPropertyValueInContainer TryResolvePropertyPath(const FPropertyPathName& Path, UObject* Object)
+{
+	return TryIteratePropertyPath(Path, Object, nullptr);
+}
+
+FPropertyValueInContainer TryResolvePropertyPathForChangeEvent(const FPropertyPathName& Path, UObject* Object, uint32 ChangeType, TSharedPtr<FPropertyChangedChainEvent>& OutEvent, FEditPropertyChain& OutChain)
+{
+	TMap<FString, int32> ArrayIndeces;
+	FPropertyValueInContainer Result = TryIteratePropertyPath(Path, Object, [&OutChain, &ArrayIndeces](const FPropertyPathNameSegment& Segment, FPropertyValueInContainer& Value)
+		{
+			OutChain.AddTail(const_cast<FProperty*>(Value.Property));
+			if (Value.ArrayIndex != INDEX_NONE)
+			{
+				ArrayIndeces.Add(Value.Property->GetName(), Value.ArrayIndex);
+			}
+		});
+
+	if (Result)
+	{
+		FPropertyChangedEvent InnerEvent(const_cast<FProperty*>(Result.Property), ChangeType, {Object});
+		OutEvent = MakeShared<FPropertyChangedChainEvent>(OutChain, InnerEvent);
+		OutEvent->SetArrayIndexPerObject({ ArrayIndeces });
+	}
+	return Result;
+}
+
 } // UE
diff --git a/Engine/Source/Runtime/CoreUObject/Private/UObject/PropertyPathName.cpp b/Engine/Source/Runtime/CoreUObject/Private/UObject/PropertyPathName.cpp
--- a/Engine/Source/Runtime/CoreUObject/Private/UObject/PropertyPathName.cpp
+++ b/Engine/Source/Runtime/CoreUObject/Private/UObject/PropertyPathName.cpp
@@ -4,6 +4,9 @@
 
 #include "Misc/StringBuilder.h"
 #include "Templates/TypeHash.h"
+#include "Internationalization/Regex.h"
+#include "Dom/JsonObject.h"
+#include "Dom/JsonValue.h"
 
 namespace UE
 {
@@ -92,6 +95,40 @@ void FPropertyPathName::ToString(FStringBuilderBase& Out, FStringView Separator)
 	}
 }
 
+TSharedRef<FJsonValue> FPropertyPathName::ToJsonValue() const
+{
+	TArray<TSharedPtr<FJsonValue>> PathArray;
+	for (const FSegment& Segment : Segments)
+	{
+		FPropertyPathNameSegment UnpackedSegment = Segment.Unpack();
+		PathArray.Add(UnpackedSegment.ToJsonValue());
+	}
+
+	return StaticCastSharedRef<FJsonValue>(MakeShared<FJsonValueArray>(PathArray));
+}
+
+bool FPropertyPathName::FromJsonValue(const TSharedPtr<FJsonValue>& JsonValue, FPropertyPathName& OutPath)
+{
+	TArray<TSharedPtr<FJsonValue>>* PathArray = nullptr;
+	if (JsonValue && JsonValue->TryGetArray(PathArray) && PathArray)
+	{
+		for (TSharedPtr<FJsonValue> SegmentValue : *PathArray)
+		{
+			FPropertyPathNameSegment UnpackedSegment;
+			if (FPropertyPathNameSegment::FromJsonValue(SegmentValue, UnpackedSegment))
+			{
+				OutPath.Segments.Add(FSegment::Pack(UnpackedSegment));
+			}
+			else
+			{
+				return false; // segment failed to parse
+			}
+		}
+		return true;
+	}
+	return false;
+}
+
 [[nodiscard]] uint32 GetTypeHash(const FPropertyPathName& Path)
 {
 	uint32 Hash = 0;
@@ -103,4 +140,54 @@ void FPropertyPathName::ToString(FStringBuilderBase& Out, FStringView Separator)
 	return Hash;
 }
 
+TSharedRef<FJsonValue> FPropertyPathNameSegment::ToJsonValue() const
+{
+	TSharedRef<FJsonObject> Result = MakeShared<FJsonObject>();
+	Result->SetStringField(TEXT("Name"), Name.ToString());
+	Result->SetStringField(TEXT("Type"), *WriteToString<128>(Type));
+	if (Index != INDEX_NONE)
+	{
+		Result->SetNumberField(TEXT("Index"), Index);
+	}
+	return StaticCastSharedRef<FJsonValue>(MakeShared<FJsonValueObject>(Result));
+}
+
+bool FPropertyPathNameSegment::FromJsonValue(const TSharedPtr<FJsonValue>& JsonValue, FPropertyPathNameSegment& OutSegment)
+{
+	TSharedPtr<FJsonObject>* AsObject = nullptr;
+	if (!JsonValue || !JsonValue->TryGetObject(AsObject) || !AsObject || !AsObject->IsValid())
+	{
+		return false; // value must be a non-null object
+	}
+
+	// Name Field (required)
+	FString Name;
+	if (!(*AsObject)->TryGetStringField(TEXT("Name"), Name))
+	{
+		return false; // Name field is required
+	}
+	OutSegment.Name = FName(Name);
+
+	// Type Field (required)
+	FString Type;
+	if (!(*AsObject)->TryGetStringField(TEXT("Type"), Type))
+	{
+		return false; // Type field is required
+	}
+	FPropertyTypeNameBuilder Builder;
+	if (!Builder.TryParse(Type))
+	{
+		return false; // type failed to parse
+	}
+	OutSegment.Type = Builder.Build();
+	
+
+	// Index Field (optional)
+	if (!(*AsObject)->TryGetNumberField(TEXT("Index"), OutSegment.Index))
+	{
+		OutSegment.Index = INDEX_NONE;
+	}
+	return true;
+}
+
 } // UE
diff --git a/Engine/Source/Runtime/CoreUObject/Private/UObject/PropertyPathNameTree.cpp b/Engine/Source/Runtime/CoreUObject/Private/UObject/PropertyPathNameTree.cpp
--- a/Engine/Source/Runtime/CoreUObject/Private/UObject/PropertyPathNameTree.cpp
+++ b/Engine/Source/Runtime/CoreUObject/Private/UObject/PropertyPathNameTree.cpp
@@ -33,6 +33,32 @@ FPropertyPathNameTree::FNode FPropertyPathNameTree::Add(const FPropertyPathName&
 	return {};
 }
 
+bool FPropertyPathNameTree::Remove(const FPropertyPathName& Path, int32 StartIndex)
+{
+	if (const int32 SegmentCount = Path.GetSegmentCount(); StartIndex < SegmentCount)
+	{
+		FPropertyPathNameSegment Segment = Path.GetSegment(StartIndex);
+		if (FValue* Child = Nodes.Find({ Segment.Name, Segment.Type }))
+		{
+			if (++StartIndex >= SegmentCount)
+			{
+				Nodes.Remove({ Segment.Name, Segment.Type });
+				return true;
+			}
+			else if (Child->SubTree && Child->SubTree->Remove(Path, StartIndex))
+			{
+				// if the node is tagless and childless after the operation, remove it too
+				if (Child->Tag.Name.IsNone() && Child->SubTree->Nodes.IsEmpty())
+				{
+					Nodes.Remove({ Segment.Name, Segment.Type });
+				}
+				return true;
+			}
+		}
+	}
+	return false;
+}
+
 FPropertyPathNameTree::FNode FPropertyPathNameTree::Find(const FPropertyPathName& Path, int32 StartIndex)
 {
 	if (const int32 SegmentCount = Path.GetSegmentCount(); StartIndex < SegmentCount)
diff --git a/Engine/Source/Runtime/CoreUObject/Public/UObject/PropertyPathFunctions.h b/Engine/Source/Runtime/CoreUObject/Public/UObject/PropertyPathFunctions.h
--- a/Engine/Source/Runtime/CoreUObject/Public/UObject/PropertyPathFunctions.h
+++ b/Engine/Source/Runtime/CoreUObject/Public/UObject/PropertyPathFunctions.h
@@ -3,6 +3,7 @@
 #pragma once
 
 #include "CoreTypes.h"
+#include "Templates/SharedPointerFwd.h"
 
 #define UE_API COREUOBJECT_API
 
@@ -10,6 +11,8 @@ class FName;
 class FProperty;
 class UObject;
 class UStruct;
+struct FPropertyChangedChainEvent;
+class FEditPropertyChain;
 
 namespace UE { class FPropertyPathName; }
 namespace UE { class FPropertyTypeName; }
@@ -76,6 +79,19 @@ struct FPropertyValueInContainer
  */
 UE_API FPropertyValueInContainer TryResolvePropertyPath(const FPropertyPathName& Path, UObject* Object);
 
+/**
+ * Try to resolve the property path to a property value within the object 
+ * and construct an FPropertyChangeChainEvent that can be used to notify changes to the resolved data
+ *
+ * @param Path The non-empty property path to resolve.
+ * @param Object The object representing the root of the property path to resolve.
+ * @param ChangeType The change type to assign to the FPropertyChangeChainEvent
+ * @param OutEvent The resulting change event
+ * @param OutChain The property chain associated with the event
+ * @return A valid property value reference if resolved, otherwise an invalid reference.
+ */
+UE_API FPropertyValueInContainer TryResolvePropertyPathForChangeEvent(const FPropertyPathName& Path, UObject* Object, uint32 ChangeType, TSharedPtr<FPropertyChangedChainEvent>& OutEvent, FEditPropertyChain& OutChain);
+
 } // UE
 
 #undef UE_API
diff --git a/Engine/Source/Runtime/CoreUObject/Public/UObject/PropertyPathName.h b/Engine/Source/Runtime/CoreUObject/Public/UObject/PropertyPathName.h
--- a/Engine/Source/Runtime/CoreUObject/Public/UObject/PropertyPathName.h
+++ b/Engine/Source/Runtime/CoreUObject/Public/UObject/PropertyPathName.h
@@ -5,11 +5,14 @@
 #include "CoreTypes.h"
 #include "Containers/Array.h"
 #include "Containers/StringFwd.h"
+#include "Templates/SharedPointerFwd.h"
 #include "UObject/NameTypes.h"
 #include "UObject/PropertyTypeName.h"
 
 #define UE_API COREUOBJECT_API
 
+class FJsonValue;
+
 namespace UE
 {
 
@@ -22,7 +25,8 @@ struct FPropertyPathNameSegment
 	FName Name;
 	/** Type of the property, if available. */
 	FPropertyTypeName Type;
-	/** Index within a container, or INDEX_NONE if not in a container. Uses ElementId for associative containers. */
+	/** Index within a container, or INDEX_NONE if not in a container. 
+		Associative containers (Sets and Maps) use virtual indices that skip gaps instead of logical indices. */
 	int32 Index = INDEX_NONE;
 
 	[[nodiscard]] inline FName PackNameWithIndex() const
@@ -39,12 +43,15 @@ struct FPropertyPathNameSegment
 	{
 		return {Name, NewType, Index};
 	}
+
+	UE_API TSharedRef<FJsonValue> ToJsonValue() const;
+	UE_API static bool FromJsonValue(const TSharedPtr<FJsonValue>& JsonValue, FPropertyPathNameSegment& OutSegment);
 };
 
 /**
  * Represents the path to a property, by name, including the property type and the container index.
  *
- * Sequenced containers use the index directly and associative containers use their ElementId.
+ * Sequenced containers use the index directly and associative containers (Sets and Maps) use virtual indices that skip gaps instead of logical indices.
  */
 class FPropertyPathName
 {
@@ -130,6 +137,9 @@ public:
 
 	UE_API void ToString(FStringBuilderBase& Out, FStringView Separator = TEXTVIEW(" -> ")) const;
 
+	UE_API TSharedRef<FJsonValue> ToJsonValue() const;
+	UE_API static bool FromJsonValue(const TSharedPtr<FJsonValue>& JsonValue, FPropertyPathName& OutPath);
+
 	UE_API friend uint32 GetTypeHash(const FPropertyPathName& Path);
 
 private:
diff --git a/Engine/Source/Runtime/CoreUObject/Tests/UObject/PropertyPathNameTreeTest.cpp b/Engine/Source/Runtime/CoreUObject/Tests/UObject/PropertyPathNameTreeTest.cpp
--- a/Engine/Source/Runtime/CoreUObject/Tests/UObject/PropertyPathNameTreeTest.cpp
+++ b/Engine/Source/Runtime/CoreUObject/Tests/UObject/PropertyPathNameTreeTest.cpp
@@ -243,6 +243,31 @@ TEST_CASE_NAMED(FPropertyPathNameTreeTest, "CoreUObject::PropertyPathNameTree",
 
 		CHECK_FALSE(Tree.Find(PathNameA, 1));
 	}
+
+	SECTION("Remove")
+	{
+		FPropertyPathName ParentPathName;
+		ParentPathName.Push({ NAME_Vector, VectorType });
+		FPropertyPathName PathNameA = ParentPathName;
+		PathNameA.Push({ CountName, IntType });
+		FPropertyPathName PathNameB = ParentPathName;
+		PathNameB.Push({ SizeName, FloatType });
+		FPropertyPathNameTree Tree;
+		Tree.Add(PathNameA);
+		Tree.Add(PathNameB);
+		CHECKED_IF(Tree.Remove(PathNameA))
+		{
+			// Because B is also child of Parent, it shouldn'tve been removed yet.
+			CHECK(Tree.Find(ParentPathName));
+
+			CHECKED_IF(Tree.Remove(PathNameB))
+			{
+				// Because B was the last child of ParentPathName, the parent should also get removed.
+				CHECK_FALSE(Tree.Find(ParentPathName));
+
+			}
+		}
+	}
 }
 
 } // UE
