[BugFix] TryResolvePropertyPathForChangeEvent was failing to set indices in the change event for non-static containers

diff --git a/Engine/Source/Runtime/CoreUObject/Private/UObject/PropertyPathFunctions.cpp b/Engine/Source/Runtime/CoreUObject/Private/UObject/PropertyPathFunctions.cpp
--- a/Engine/Source/Runtime/CoreUObject/Private/UObject/PropertyPathFunctions.cpp
+++ b/Engine/Source/Runtime/CoreUObject/Private/UObject/PropertyPathFunctions.cpp
@@ -191,9 +191,9 @@ FPropertyValueInContainer TryResolvePropertyPathForChangeEvent(const FPropertyPa
 	FPropertyValueInContainer Result = TryIteratePropertyPath(Path, Object, [&OutChain, &ArrayIndeces](const FPropertyPathNameSegment& Segment, FPropertyValueInContainer& Value)
 		{
 			OutChain.AddTail(const_cast<FProperty*>(Value.Property));
-			if (Value.ArrayIndex != INDEX_NONE)
+			if (Segment.Index != INDEX_NONE)
 			{
-				ArrayIndeces.Add(Value.Property->GetName(), Value.ArrayIndex);
+				ArrayIndeces.Add(Value.Property->GetName(), Segment.Index);
 			}
 		});
 
