Changing FRemoteObjectConstructionOverridesScope to accept a pointer instead of reference to the overrides to allow use cases where the overrides are not always present

#preflght

diff --git a/Engine/Source/Runtime/CoreUObject/Private/UObject/RemoteObjectSerialization.cpp b/Engine/Source/Runtime/CoreUObject/Private/UObject/RemoteObjectSerialization.cpp
--- a/Engine/Source/Runtime/CoreUObject/Private/UObject/RemoteObjectSerialization.cpp
+++ b/Engine/Source/Runtime/CoreUObject/Private/UObject/RemoteObjectSerialization.cpp
@@ -9,6 +9,7 @@
 #include "UObject/PropertyPortFlags.h"
 #include "UObject/UObjectHash.h"
 #include "UObject/UObjectMigrationContext.h"
+#include "UObject/UObjectThreadContext.h"
 #include "UObject/UnrealType.h"
 #include "Serialization/ArchiveUObject.h"
 #include "Serialization/MemoryWriter.h"
@@ -1299,7 +1300,7 @@ int32 DeserializeObjectData(FRemoteObjectData& ObjectData,
 			TRACE_CPUPROFILER_EVENT_SCOPE(ConstructRemoteObjects);
 
 			FRemoteObjectConstructionOverrides ConstructionOverrides(ObjectHeaders);
-			FRemoteObjectConstructionOverridesScope OverridesScope(ConstructionOverrides);
+			FRemoteObjectConstructionOverridesScope OverridesScope(&ConstructionOverrides);
 
 			OutObjectRemoteIds.Reserve(ObjectHeaders.Num());
 			OutObjects.Reserve(ObjectHeaders.Num());
diff --git a/Engine/Source/Runtime/CoreUObject/Private/UObject/RemoteObjectSerialization.h b/Engine/Source/Runtime/CoreUObject/Private/UObject/RemoteObjectSerialization.h
--- a/Engine/Source/Runtime/CoreUObject/Private/UObject/RemoteObjectSerialization.h
+++ b/Engine/Source/Runtime/CoreUObject/Private/UObject/RemoteObjectSerialization.h
@@ -60,6 +60,7 @@ namespace UE::RemoteObject::Serialization
 
 	public:
 
+		FRemoteObjectConstructionOverrides() = default;
 		explicit FRemoteObjectConstructionOverrides(const TArray<FRemoteObjectHeader>& InObjectHeaders);
 
 		/**
@@ -106,16 +107,24 @@ namespace UE::RemoteObject::Serialization
 	class FRemoteObjectConstructionOverridesScope final
 	{
 		FRemoteObjectConstructionOverridesStack& Stack;
+		FRemoteObjectConstructionOverrides* Overrides = nullptr;
 
 	public:
-		explicit FRemoteObjectConstructionOverridesScope(FRemoteObjectConstructionOverrides& InOverrides)
+		explicit FRemoteObjectConstructionOverridesScope(FRemoteObjectConstructionOverrides* InOverrides)
 			: Stack(FRemoteObjectConstructionOverridesStack::Get())
+			, Overrides(InOverrides)
 		{
-			Stack.Push(InOverrides);
+			if (Overrides)
+			{
+				Stack.Push(*Overrides);
+			}
 		}
 		~FRemoteObjectConstructionOverridesScope()
 		{
-			Stack.Pop();
+			if (Overrides)
+			{
+				Stack.Pop();
+			}
 		}
 	};
 
