GitHub 13450 : Fix StaticDuplicateObject when Source Object class and Destination class are different

#13450

diff --git a/Engine/Source/Runtime/CoreUObject/Private/UObject/UObjectGlobals.cpp b/Engine/Source/Runtime/CoreUObject/Private/UObject/UObjectGlobals.cpp
--- a/Engine/Source/Runtime/CoreUObject/Private/UObject/UObjectGlobals.cpp
+++ b/Engine/Source/Runtime/CoreUObject/Private/UObject/UObjectGlobals.cpp
@@ -3205,6 +3205,8 @@ UObject* StaticDuplicateObjectEx( FObjectDuplicationParameters& Parameters )
 		}
 	}
 
+	const bool bShouldForceTaggedPropertySerialization = Parameters.SourceObject->GetClass() != Parameters.DestClass;
+
 	// Read from the source object(s)
 	FDuplicateDataWriter Writer(
 		DuplicatedObjectAnnotation,				// Ref: Object annotation which stores the duplicated object for each source object
@@ -3219,6 +3221,11 @@ UObject* StaticDuplicateObjectEx( FObjectDuplicationParameters& Parameters )
 		Parameters.PortFlags,					// PortFlags	
 		Parameters.bAssignExternalPackages);	// Assign duplicate external packages
 
+	if (bShouldForceTaggedPropertySerialization)
+	{
+		Writer.SetUseUnversionedPropertySerialization(false);
+	}
+
 	TArray<UObject*> SerializedObjects;
 
 	
@@ -3240,6 +3247,12 @@ UObject* StaticDuplicateObjectEx( FObjectDuplicationParameters& Parameters )
 
 	TRefCountPtr<FUObjectSerializeContext> LoadContext(FUObjectThreadContext::Get().GetSerializeContext());
 	FDuplicateDataReader Reader(DuplicatedObjectAnnotation, ObjectData.Get(), Parameters.PortFlags, Parameters.DestOuter);
+
+	if (bShouldForceTaggedPropertySerialization)
+	{
+		Reader.SetUseUnversionedPropertySerialization(false);
+	}
+
 	for(int32 ObjectIndex = 0;ObjectIndex < SerializedObjects.Num();ObjectIndex++)
 	{
 		UObject* SerializedObject = SerializedObjects[ObjectIndex];
