FRedirectionSummary:: Add/Remove : remove parameter bIsWildcardMatch. Use the private function FCoreRedirect::IsWildcardMatch instead.
Make FRedirectionSummary a friend class of FCoreRedirect so I can call functions like IsWildcardMatch. It will be usefull for future changes as well to call IsPrefixedMatch.

#rnx

diff --git a/Engine/Source/Runtime/CoreUObject/Internal/UObject/CoreRedirects/RedirectionSummary.h b/Engine/Source/Runtime/CoreUObject/Internal/UObject/CoreRedirects/RedirectionSummary.h
--- a/Engine/Source/Runtime/CoreUObject/Internal/UObject/CoreRedirects/RedirectionSummary.h
+++ b/Engine/Source/Runtime/CoreUObject/Internal/UObject/CoreRedirects/RedirectionSummary.h
@@ -24,8 +24,8 @@ public:
 	FRedirectionSummary(FRedirectionSummary&& Other);
 	FRedirectionSummary& operator=(FRedirectionSummary&& Other);
 
-	void Add(const FCoreRedirect& CoreRedirect, bool bIsWildcardMatch);
-	void Remove(const FCoreRedirect& CoreRedirect, bool bIsWildcardMatch);
+	void Add(const FCoreRedirect& CoreRedirect);
+	void Remove(const FCoreRedirect& CoreRedirect);
 
 	UE_DEPRECATED(5.6, "Use GetHashAffectingPackages instead.")
 	void AppendHashAffectingPackages(FBlake3& Hasher, TConstArrayView<FName> PackageNames);
diff --git a/Engine/Source/Runtime/CoreUObject/Private/UObject/CoreRedirects.cpp b/Engine/Source/Runtime/CoreUObject/Private/UObject/CoreRedirects.cpp
--- a/Engine/Source/Runtime/CoreUObject/Private/UObject/CoreRedirects.cpp
+++ b/Engine/Source/Runtime/CoreUObject/Private/UObject/CoreRedirects.cpp
@@ -2047,7 +2047,7 @@ void FCoreRedirects::RecordAddedObjectRedirector(const FSoftObjectPath& Source,
 	{
 		FCoreRedirect ConvertedToCoreRedirect(ECoreRedirectFlags::Type_Object,
 			FCoreRedirectObjectName(Source), FCoreRedirectObjectName(Dest));
-		Context.GetRedirectionSummary().Add(ConvertedToCoreRedirect, false /* bIsWildcardMatch */);
+		Context.GetRedirectionSummary().Add(ConvertedToCoreRedirect);
 	}
 }
 
@@ -2058,7 +2058,7 @@ void FCoreRedirects::RecordRemovedObjectRedirector(const FSoftObjectPath& Source
 	{
 		FCoreRedirect ConvertedToCoreRedirect(ECoreRedirectFlags::Type_Object,
 			FCoreRedirectObjectName(Source), FCoreRedirectObjectName(Dest));
-		Context.GetRedirectionSummary().Remove(ConvertedToCoreRedirect, false /* bIsWildcardMatch */);
+		Context.GetRedirectionSummary().Remove(ConvertedToCoreRedirect);
 	}
 }
 
@@ -2846,7 +2846,7 @@ bool FCoreRedirects::AddSingleRedirectUnderWriteLock(const FCoreRedirect& NewRed
 #if WITH_EDITOR
 	if (!!(LockedContext.Get().GetFlags() & FCoreRedirectsContext::EFlags::UseRedirectionSummary))
 	{
-		LockedContext.Get().GetRedirectionSummary().Add(NewRedirect, bIsWildcardMatch);
+		LockedContext.Get().GetRedirectionSummary().Add(NewRedirect);
 	}
 #endif
 
@@ -2975,7 +2975,7 @@ bool FCoreRedirects::RemoveSingleRedirectUnderWriteLock(const FCoreRedirect& Red
 #if WITH_EDITOR
 		if (!!(LockedContext.Get().GetFlags() & FCoreRedirectsContext::EFlags::UseRedirectionSummary))
 		{
-			LockedContext.Get().GetRedirectionSummary().Remove(RedirectToRemove, bIsWildcardMatch);
+			LockedContext.Get().GetRedirectionSummary().Remove(RedirectToRemove);
 		}
 #endif
 	}
diff --git a/Engine/Source/Runtime/CoreUObject/Private/UObject/RedirectionSummary.cpp b/Engine/Source/Runtime/CoreUObject/Private/UObject/RedirectionSummary.cpp
--- a/Engine/Source/Runtime/CoreUObject/Private/UObject/RedirectionSummary.cpp
+++ b/Engine/Source/Runtime/CoreUObject/Private/UObject/RedirectionSummary.cpp
@@ -41,11 +41,11 @@ FRedirectionSummary& FRedirectionSummary::operator=(FRedirectionSummary&& Other)
 	return *this;
 }
 
-void FRedirectionSummary::Add(const FCoreRedirect& CoreRedirect, bool bIsWildcardMatch)
+void FRedirectionSummary::Add(const FCoreRedirect& CoreRedirect)
 {
 	FWriteScopeLock ScopeLock(Lock);
 
-	if (bIsWildcardMatch || CoreRedirect.OldName.PackageName.IsNone())
+	if (CoreRedirect.IsWildcardMatch() || CoreRedirect.OldName.PackageName.IsNone())
 	{
 		GlobalRedirects.Add(FCoreRedirect(CoreRedirect));
 	}
@@ -58,11 +58,11 @@ void FRedirectionSummary::Add(const FCoreRedirect& CoreRedirect, bool bIsWildcar
 	}
 }
 
-void FRedirectionSummary::Remove(const FCoreRedirect& CoreRedirect, bool bIsWildcardMatch)
+void FRedirectionSummary::Remove(const FCoreRedirect& CoreRedirect)
 {
 	FWriteScopeLock ScopeLock(Lock);
 
-	if (bIsWildcardMatch || CoreRedirect.OldName.PackageName.IsNone())
+	if (CoreRedirect.IsWildcardMatch() || CoreRedirect.OldName.PackageName.IsNone())
 	{
 		GlobalRedirects.Remove(CoreRedirect);
 	}
diff --git a/Engine/Source/Runtime/CoreUObject/Public/UObject/CoreRedirects.h b/Engine/Source/Runtime/CoreUObject/Public/UObject/CoreRedirects.h
--- a/Engine/Source/Runtime/CoreUObject/Public/UObject/CoreRedirects.h
+++ b/Engine/Source/Runtime/CoreUObject/Public/UObject/CoreRedirects.h
@@ -285,6 +285,7 @@ struct FCoreRedirect
 
 private:
 	friend struct FCoreRedirects;
+	friend class FRedirectionSummary;
 
 	/* Returns the updated name after redirection. If bIsKnownToMatch is true, OldObjectName must have 
 	been validated previously to be acceptable for redirection */
