Fix Mass holding on to observer classes after their module has unloaded

- Add self-cleanup for MassObserverManager, MassObserverRegistry
- Update OnModulePackagesLoaded to provide a package list, as this is becoming widely used
- Add missing Reserve call
- Cleanup

#rnx

diff --git a/Engine/Source/Runtime/CoreUObject/Private/UObject/UObjectBase.cpp b/Engine/Source/Runtime/CoreUObject/Private/UObject/UObjectBase.cpp
--- a/Engine/Source/Runtime/CoreUObject/Private/UObject/UObjectBase.cpp
+++ b/Engine/Source/Runtime/CoreUObject/Private/UObject/UObjectBase.cpp
@@ -112,7 +112,7 @@ static TMap<FName, TArray<FPendingRegistrant*>>& GetPerModuleBootstrapMap()
 #endif
 
 /** Must be called before modules will unload that contains UObject classes */
-static void RemoveLoadedUObjects(TConstArrayView<FName> InModuleNames);
+static void RemoveCompiledInUObjects(TConstArrayView<FName> InModuleNames);
 
 /** Can be called to confirm full removal of UObjects in a list of modules */
 static bool CheckLiveObjectsInModules(TConstArrayView<FName> InModuleNames);
@@ -999,7 +999,7 @@ void RegisterModularObjectsProcessing()
 	UE_CALL_ONCE([]()
 	{
 		FModuleManager::Get().OnProcessLoadedObjectsCallback().AddStatic(ProcessNewlyLoadedUObjects);
-		FModuleManager::Get().OnModulesUnloaded().AddStatic(RemoveLoadedUObjects);
+		FModuleManager::Get().OnModulesUnloaded().AddStatic(RemoveCompiledInUObjects);
 		FModuleManager::Get().OnCheckLiveObjectsInModulesCallback().BindStatic(CheckLiveObjectsInModules);
 
 		FCoreUObjectDelegates::GarbageCollectComplete.AddLambda([]()
@@ -1108,6 +1108,8 @@ static TArray<UPackage*> GetScriptPackagesFromModuleNames(TConstArrayView<FName>
 {
 	TArray<UPackage*> Packages;
 
+	Packages.Reserve(2 * InModuleNames.Num());
+
 	for (const FName ModuleName : InModuleNames)
 	{
 		if (UPackage* Package = FindObjectFast<UPackage>(/* Outer */ nullptr, WriteToString<256>(TEXT("/Script/"), ModuleName).ToString()))
@@ -1124,7 +1126,7 @@ static TArray<UPackage*> GetScriptPackagesFromModuleNames(TConstArrayView<FName>
 	return Packages;
 }
 
-static void RemoveObjectsInModules(TConstArrayView<FName> InModuleNames)
+static void RemoveObjectsInModules(TConstArrayView<UPackage*> Packages)
 {
 	LLM_SCOPE(ELLMTag::UObject);
 	LLM_SCOPE_BYTAG(UObject_UObjectBase);
@@ -1132,7 +1134,6 @@ static void RemoveObjectsInModules(TConstArrayView<FName> InModuleNames)
 
 	FPermanentObjectPoolExtents PermanentPool;
 
-	const TArray<UPackage*> Packages = GetScriptPackagesFromModuleNames(InModuleNames);
 	for (UPackage* Package : Packages)
 	{
 		UE_LOG(LogUObjectBase, Log, TEXT("RemoveObjectsInModules: removing package %s"), *Package->GetName());
@@ -1210,14 +1211,17 @@ static bool CheckLiveObjectsInModules(TConstArrayView<FName> InModuleNames)
 	}
 }
 
-static void RemoveLoadedUObjects(TConstArrayView<FName> InModuleNames)
+static void RemoveCompiledInUObjects(TConstArrayView<FName> InModuleNames)
 {
 	LLM_SCOPE(ELLMTag::UObject);
 	LLM_SCOPE_BYTAG(UObject_UObjectBase);
-	DECLARE_SCOPE_CYCLE_COUNTER(TEXT("RemoveLoadedUObjects"), STAT_RemoveLoadedUObjects, STATGROUP_ObjectVerbose);
+	DECLARE_SCOPE_CYCLE_COUNTER(TEXT("RemoveCompiledInUObjects"), STAT_RemoveCompiledInUObjects, STATGROUP_ObjectVerbose);
 
 	FlushAsyncLoading();
-	RemoveObjectsInModules(InModuleNames);
+
+	const TArray<UPackage*> Packages = GetScriptPackagesFromModuleNames(InModuleNames);
+	FCoreUObjectDelegates::CompiledInUObjectsRemovedDelegate.Broadcast(Packages);
+	RemoveObjectsInModules(Packages);
 }
 
 static int32 GVarMaxObjectsNotConsideredByGC;
diff --git a/Engine/Source/Runtime/CoreUObject/Private/UObject/UObjectGlobals.cpp b/Engine/Source/Runtime/CoreUObject/Private/UObject/UObjectGlobals.cpp
--- a/Engine/Source/Runtime/CoreUObject/Private/UObject/UObjectGlobals.cpp
+++ b/Engine/Source/Runtime/CoreUObject/Private/UObject/UObjectGlobals.cpp
@@ -187,6 +187,7 @@ FCoreUObjectDelegates::FReloadReinstancingCompleteDelegate FCoreUObjectDelegates
 FCoreUObjectDelegates::FReloadCompleteDelegate FCoreUObjectDelegates::ReloadCompleteDelegate;
 FCoreUObjectDelegates::FReloadAddedClassesDelegate FCoreUObjectDelegates::ReloadAddedClassesDelegate;
 FCoreUObjectDelegates::FCompiledInUObjectsRegisteredDelegate FCoreUObjectDelegates::CompiledInUObjectsRegisteredDelegate;
+FCoreUObjectDelegates::FCompiledInUObjectsRemovedDelegate FCoreUObjectDelegates::CompiledInUObjectsRemovedDelegate;
 FCoreUObjectDelegates::FIsPackageOKToSaveDelegate FCoreUObjectDelegates::IsPackageOKToSaveDelegate;
 FCoreUObjectDelegates::FOnPostInitSparseClassData FCoreUObjectDelegates::OnPostInitSparseClassData;
 FCoreUObjectDelegates::FOnPackageReloaded FCoreUObjectDelegates::OnPackageReloaded;
diff --git a/Engine/Source/Runtime/CoreUObject/Public/UObject/UObjectGlobals.h b/Engine/Source/Runtime/CoreUObject/Public/UObject/UObjectGlobals.h
--- a/Engine/Source/Runtime/CoreUObject/Public/UObject/UObjectGlobals.h
+++ b/Engine/Source/Runtime/CoreUObject/Public/UObject/UObjectGlobals.h
@@ -3385,6 +3385,10 @@ struct FCoreUObjectDelegates
 	DECLARE_MULTICAST_DELEGATE_TwoParams(FCompiledInUObjectsRegisteredDelegate, FName /*Package*/, ECompiledInUObjectsRegisteredStatus Status);
 	static COREUOBJECT_API FCompiledInUObjectsRegisteredDelegate CompiledInUObjectsRegisteredDelegate;
 
+	/* Multicast delegate that is executed when any UObjects in packages need destroying before modules are unloaded. */
+	DECLARE_MULTICAST_DELEGATE_OneParam(FCompiledInUObjectsRemovedDelegate, TConstArrayView<UPackage*>);
+	static COREUOBJECT_API FCompiledInUObjectsRemovedDelegate CompiledInUObjectsRemovedDelegate;
+
 	/** Sent at the very beginning of LoadMap */
 	DECLARE_MULTICAST_DELEGATE_OneParam(FPreLoadMapDelegate, const FString& /* MapName */);
 	static COREUOBJECT_API FPreLoadMapDelegate PreLoadMap;
diff --git a/Engine/Source/Runtime/MassEntity/Private/MassObserverManager.cpp b/Engine/Source/Runtime/MassEntity/Private/MassObserverManager.cpp
--- a/Engine/Source/Runtime/MassEntity/Private/MassObserverManager.cpp
+++ b/Engine/Source/Runtime/MassEntity/Private/MassObserverManager.cpp
@@ -339,10 +339,21 @@ void FMassObserverManager::Initialize()
 		}
 	});
 #endif // WITH_MASSENTITY_DEBUG
+
+	if (!ModulesUnloadedHandle.IsValid())
+	{
+		ModulesUnloadedHandle = FCoreUObjectDelegates::CompiledInUObjectsRemovedDelegate.AddRaw(this, &FMassObserverManager::OnModulePackagesUnloaded);
+	}
 }
 
 void FMassObserverManager::DeInitialize()
 {
+	if (ModulesUnloadedHandle.IsValid())
+	{
+		FCoreUObjectDelegates::CompiledInUObjectsRemovedDelegate.Remove(ModulesUnloadedHandle);
+		ModulesUnloadedHandle.Reset();
+	}
+
 	for (int32 i = 0; i < (int32)EMassObservedOperation::MAX; ++i)
 	{
 		(*FragmentObservers[i]).Empty();
@@ -789,6 +800,32 @@ void FMassObserverManager::OnPostFork(EForkProcessRole)
 	}
 }
 
+void FMassObserverManager::OnModulePackagesUnloaded(TConstArrayView<UPackage*> Packages)
+{
+	TRACE_CPUPROFILER_EVENT_SCOPE(FMassObserverManager::OnModulePackagesUnloaded);
+
+	const auto ProcessObserversMap = [&Packages](FMassObserversMap& ObserversMap)
+	{
+		for (auto MapIt = (*ObserversMap).CreateIterator(); MapIt; ++MapIt)
+		{
+			TObjectPtr<const UScriptStruct> ObserverClass = MapIt->Key;
+			const UPackage* ObserverPackage = ObserverClass->GetPackage();
+
+			if (Packages.Contains(ObserverPackage))
+			{
+				UE_LOG(LogMass, Verbose, TEXT("%hs: removed observer %s (%s)"), __FUNCTION__, *ObserverClass->GetName(), *ObserverPackage->GetName());
+				MapIt.RemoveCurrent();
+			}
+		}
+	};
+
+	for (int32 OperationIndex = 0; OperationIndex < static_cast<uint32>(EMassObservedOperation::MAX); ++OperationIndex)
+	{
+		ProcessObserversMap(FragmentObservers[OperationIndex]);
+		ProcessObserversMap(TagObservers[OperationIndex]);
+	}
+}
+
 void FMassObserverManager::ResumeExecution(FObserverLock& LockBeingReleased)
 {
 	using namespace UE::Mass::ObserverManager;
@@ -961,4 +998,4 @@ PRAGMA_ENABLE_DEPRECATION_WARNINGS
 bool FMassObserverManager::OnCollectionsCreatedImpl(UE::Mass::FProcessingContext& ProcessingContext, TArray<FMassArchetypeEntityCollection>&& EntityCollections)
 {
 	return OnCollectionsCreatedImpl(ProcessingContext, MakeArrayView(EntityCollections));
-}
\ No newline at end of file
+}
diff --git a/Engine/Source/Runtime/MassEntity/Private/MassObserverRegistry.cpp b/Engine/Source/Runtime/MassEntity/Private/MassObserverRegistry.cpp
--- a/Engine/Source/Runtime/MassEntity/Private/MassObserverRegistry.cpp
+++ b/Engine/Source/Runtime/MassEntity/Private/MassObserverRegistry.cpp
@@ -1,6 +1,7 @@
 // Copyright Epic Games, Inc. All Rights Reserved.
 
 #include "MassObserverRegistry.h"
+#include "Algo/Find.h"
 
 #include UE_INLINE_GENERATED_CPP_BY_NAME(MassObserverRegistry)
 
@@ -11,6 +12,11 @@ UMassObserverRegistry::UMassObserverRegistry()
 {
 	// there can be only one!
 	check(HasAnyFlags(RF_ClassDefaultObject));
+
+	if (!ModulesUnloadedHandle.IsValid())
+	{
+		ModulesUnloadedHandle = FCoreUObjectDelegates::CompiledInUObjectsRemovedDelegate.AddUObject(this, &UMassObserverRegistry::OnModulePackagesUnloaded);
+	}
 }
 
 void UMassObserverRegistry::RegisterObserver(TNotNull<const UScriptStruct*> ObservedType, const uint8 OperationFlags, TSubclassOf<UMassProcessor> ObserverClass)
@@ -36,3 +42,41 @@ void UMassObserverRegistry::RegisterObserver(const UScriptStruct& ObservedType,
 {
 	RegisterObserver(&ObservedType, static_cast<uint8>(1 << static_cast<uint8>(Operation)), ObserverClass);
 }
+
+void UMassObserverRegistry::OnModulePackagesUnloaded(TConstArrayView<UPackage*> Packages)
+{
+	TRACE_CPUPROFILER_EVENT_SCOPE(UMassObserverRegistry::OnModulePackagesUnloaded);
+
+	const auto ProcessObserversMap = [&Packages](FObserverClassesMap* ObserversMap)
+	{
+		for (auto MapIt = ObserversMap->CreateIterator(); MapIt; ++MapIt)
+		{
+			TArray<FSoftClassPath>& ObserverClasses = MapIt->Value;
+
+			for (auto ObserverIt = ObserverClasses.CreateIterator(); ObserverIt; ++ObserverIt)
+			{
+				const FSoftClassPath& ObservedClass = *ObserverIt;
+				const FName PackageName = ObservedClass.GetLongPackageFName();
+
+				bool bRemove = !!Algo::FindByPredicate(Packages, [PackageName](const UPackage* Package)
+				{
+					return PackageName == Package->GetFName();
+				});
+
+				if (bRemove)
+				{
+					UE_LOG(LogMass, Log, TEXT("%hs: removed observer %s (%s)"), __FUNCTION__, *ObservedClass.ToString(), *PackageName.ToString());
+					ObserverIt.RemoveCurrent();
+				}
+			}
+
+			if (ObserverClasses.Num() == 0)
+			{
+				MapIt.RemoveCurrent();
+			}
+		}
+	};
+
+	ProcessObserversMap(FragmentObserverMaps);
+	ProcessObserversMap(TagObserverMaps);
+}
diff --git a/Engine/Source/Runtime/MassEntity/Public/MassObserverManager.h b/Engine/Source/Runtime/MassEntity/Public/MassObserverManager.h
--- a/Engine/Source/Runtime/MassEntity/Public/MassObserverManager.h
+++ b/Engine/Source/Runtime/MassEntity/Public/MassObserverManager.h
@@ -216,6 +216,8 @@ protected:
 
 	MASSENTITY_API void OnPostFork(EForkProcessRole);
 
+	MASSENTITY_API void OnModulePackagesUnloaded(TConstArrayView<UPackage*> Packages);
+
 	//TSharedRef<FObserverLock> GetOrMakeObserverLock(TConstArrayView<FMassEntityHandle> ReservedEntities, FMassArchetypeEntityCollection&& EntityCollection);
 	MASSENTITY_API TSharedRef<FObserverLock> GetOrMakeObserverLock();
 	TSharedPtr<FObserverLock> GetObserverLock() const;
@@ -257,6 +259,8 @@ protected:
 	UPROPERTY()
 	FMassObserversMap TagObservers[(uint8)EMassObservedOperation::MAX];
 
+	FDelegateHandle ModulesUnloadedHandle;
+
 	/** 
 	 * The owning EntityManager. No need for it to be a UPROPERTY since by design we don't support creation of 
 	 * FMassObserverManager outside a FMassEntityManager instance 
diff --git a/Engine/Source/Runtime/MassEntity/Public/MassObserverRegistry.h b/Engine/Source/Runtime/MassEntity/Public/MassObserverRegistry.h
--- a/Engine/Source/Runtime/MassEntity/Public/MassObserverRegistry.h
+++ b/Engine/Source/Runtime/MassEntity/Public/MassObserverRegistry.h
@@ -70,6 +70,11 @@ public:
 
 	using FObserverClassesMap = TMap<TObjectKey<const UScriptStruct>, TArray<FSoftClassPath>>;
 protected:
+
+	void OnModulePackagesUnloaded(TConstArrayView<UPackage*> Packages);
+
+	FDelegateHandle ModulesUnloadedHandle;
+
 	friend FMassObserverManager;
 	FObserverClassesMap FragmentObserverMaps[static_cast<uint8>(EMassObservedOperation::MAX)];
 	FObserverClassesMap TagObserverMaps[static_cast<uint8>(EMassObservedOperation::MAX)];
