PR #13724: Update SoftObjectPtr.h param name to match documentation.

#rnx
#13450

- The code file is fairly evenly split on using rhs vs other but the class in question tends to use other so in this case the PR is correct and it makes more sense to change the code to match the documentation than the other way around.

diff --git a/Engine/Source/Runtime/CoreUObject/Public/UObject/SoftObjectPtr.h b/Engine/Source/Runtime/CoreUObject/Public/UObject/SoftObjectPtr.h
--- a/Engine/Source/Runtime/CoreUObject/Public/UObject/SoftObjectPtr.h
+++ b/Engine/Source/Runtime/CoreUObject/Public/UObject/SoftObjectPtr.h
@@ -848,9 +848,9 @@ public:
 	 *
 	 * @param Other soft pointer to compare to 
 	 */
-	[[nodiscard]] UE_FORCEINLINE_HINT bool operator==(const TSoftClassPtr& Rhs) const
+	[[nodiscard]] UE_FORCEINLINE_HINT bool operator==(const TSoftClassPtr& Other) const
 	{
-		return SoftObjectPtr == Rhs.SoftObjectPtr;
+		return SoftObjectPtr == Other.SoftObjectPtr;
 	}
 #if !PLATFORM_COMPILER_HAS_GENERATED_COMPARISON_OPERATORS
 	/**  
@@ -859,9 +859,9 @@ public:
 	 *
 	 * @param Other soft pointer to compare to
 	 */
-	[[nodiscard]] UE_FORCEINLINE_HINT bool operator!=(const TSoftClassPtr& Rhs) const
+	[[nodiscard]] UE_FORCEINLINE_HINT bool operator!=(const TSoftClassPtr& Other) const
 	{
-		return SoftObjectPtr != Rhs.SoftObjectPtr;
+		return SoftObjectPtr != Other.SoftObjectPtr;
 	}
 #endif
 
