Clean-up some editor-only material localization that was being gathered as game localization data

* Consider Optional objects as editor-only, as they are currently only used for editor-only sidecar data.
* Added the "EditorOnlyLocalization" meta-data tag for properties that can't be marked WITH_EDITORONLY_DATA, but should still be considered as editor-only for localization.
* This should be a last resort, for data that cannot be made WITH_EDITORONLY_DATA (eg, because it's needed by a cooked editor), and that also cannot be moved to an optional object.
* Updated some material channel names to use "EditorOnlyLocalization" meta-data on their properties.

diff --git a/Engine/Source/Runtime/CoreUObject/Private/Serialization/PropertyLocalizationDataGathering.cpp b/Engine/Source/Runtime/CoreUObject/Private/Serialization/PropertyLocalizationDataGathering.cpp
--- a/Engine/Source/Runtime/CoreUObject/Private/Serialization/PropertyLocalizationDataGathering.cpp
+++ b/Engine/Source/Runtime/CoreUObject/Private/Serialization/PropertyLocalizationDataGathering.cpp
@@ -196,7 +196,7 @@ bool FPropertyLocalizationDataGatherer::CanGatherFromInnerProperty(const FProper
 void FPropertyLocalizationDataGatherer::GatherLocalizationDataFromObjectWithCallbacks(const UObject* Object, const EPropertyLocalizationGathererTextFlags GatherTextFlags)
 {
 	EPropertyLocalizationGathererTextFlags ObjectGatherTextFlags = GatherTextFlags;
-	if (Object->GetClass()->GetPackage()->HasAnyPackageFlags(PKG_EditorOnly))
+	if (Object->GetClass()->GetPackage()->HasAnyPackageFlags(PKG_EditorOnly) || Object->GetClass()->HasAnyClassFlags(CLASS_Optional))
 	{
 		ObjectGatherTextFlags |= EPropertyLocalizationGathererTextFlags::ForceEditorOnlyProperties;
 	}
@@ -290,7 +290,7 @@ void FPropertyLocalizationDataGatherer::GatherLocalizationDataFromObjectFields(c
 		}
 		const void* ValueAddress = PropertyField->ContainerPtrToValuePtr<void>(Object);
 		const void* DefaultValueAddress = (ArchetypeObject && ArchetypeObject->IsA(PropertyField->GetOwnerClass())) ? PropertyField->ContainerPtrToValuePtr<void>(ArchetypeObject) : nullptr;
-		GatherLocalizationDataFromChildTextProperties(PathToParent, PropertyField, ValueAddress, DefaultValueAddress, GatherTextFlags | (PropertyField->HasAnyPropertyFlags(CPF_EditorOnly) ? EPropertyLocalizationGathererTextFlags::ForceEditorOnly : EPropertyLocalizationGathererTextFlags::None));
+		GatherLocalizationDataFromChildTextProperties(PathToParent, PropertyField, ValueAddress, DefaultValueAddress, GatherTextFlags | (IsEditorOnlyProperty(PropertyField) ? EPropertyLocalizationGathererTextFlags::ForceEditorOnly : EPropertyLocalizationGathererTextFlags::None));
 	}
 
 	// Gather text from the script bytecode
@@ -331,7 +331,7 @@ void FPropertyLocalizationDataGatherer::GatherLocalizationDataFromStructFields(c
 	{
 		const void* ValueAddress = PropertyField->ContainerPtrToValuePtr<void>(StructData);
 		const void* DefaultValueAddress = DefaultStructData ? PropertyField->ContainerPtrToValuePtr<void>(DefaultStructData) : nullptr;
-		GatherLocalizationDataFromChildTextProperties(PathToParent, PropertyField, ValueAddress, DefaultValueAddress, GatherTextFlags | (PropertyField->HasAnyPropertyFlags(CPF_EditorOnly) ? EPropertyLocalizationGathererTextFlags::ForceEditorOnly : EPropertyLocalizationGathererTextFlags::None));
+		GatherLocalizationDataFromChildTextProperties(PathToParent, PropertyField, ValueAddress, DefaultValueAddress, GatherTextFlags | (IsEditorOnlyProperty(PropertyField) ? EPropertyLocalizationGathererTextFlags::ForceEditorOnly : EPropertyLocalizationGathererTextFlags::None));
 	}
 
 	// Gather text from the script bytecode
@@ -362,7 +362,7 @@ void FPropertyLocalizationDataGatherer::GatherLocalizationDataFromChildTextPrope
 	const FObjectPropertyBase* const ObjectProperty = CastField<const FObjectPropertyBase>(Property);
 	const FOptionalProperty* const OptionalProperty = CastField<const FOptionalProperty>(Property);
 
-	const EPropertyLocalizationGathererTextFlags FixedChildPropertyGatherTextFlags = GatherTextFlags | (Property->HasAnyPropertyFlags(CPF_EditorOnly) ? EPropertyLocalizationGathererTextFlags::ForceEditorOnly : EPropertyLocalizationGathererTextFlags::None);
+	const EPropertyLocalizationGathererTextFlags FixedChildPropertyGatherTextFlags = GatherTextFlags | (IsEditorOnlyProperty(Property) ? EPropertyLocalizationGathererTextFlags::ForceEditorOnly : EPropertyLocalizationGathererTextFlags::None);
 
 	FString PathToElementRoot;
 	{
@@ -423,7 +423,7 @@ void FPropertyLocalizationDataGatherer::GatherLocalizationDataFromChildTextPrope
 					PropertyPackage->MarkPackageDirty();
 				}
 
-				GatherTextInstance(*TextElementValueAddress, PathToElement, !!(GatherTextFlags & EPropertyLocalizationGathererTextFlags::ForceEditorOnlyProperties) || TextProperty->HasAnyPropertyFlags(CPF_EditorOnly));
+				GatherTextInstance(*TextElementValueAddress, PathToElement, !!(GatherTextFlags & EPropertyLocalizationGathererTextFlags::ForceEditorOnlyProperties) || IsEditorOnlyProperty(TextProperty));
 			}
 		}
 		// Property is a DYNAMIC array property.
@@ -816,6 +816,26 @@ void FPropertyLocalizationDataGatherer::MarkDefaultTextInstance(const FText& Tex
 	}
 }
 
+bool FPropertyLocalizationDataGatherer::IsEditorOnlyProperty(const FProperty* Property)
+{
+	if (Property->HasAnyPropertyFlags(CPF_EditorOnly))
+	{
+		return true;
+	}
+
+#if WITH_METADATA
+	{
+		static const FName NAME_EditorOnlyLocalization = "EditorOnlyLocalization";
+		if (Property->HasMetaData(NAME_EditorOnlyLocalization))
+		{
+			return true;
+		}
+	}
+#endif
+	
+	return false;
+}
+
 bool FPropertyLocalizationDataGatherer::ExtractTextIdentity(const FText& Text, FString& OutNamespace, FString& OutKey, const bool bCleanNamespace)
 {
 	const FTextId TextId = FTextInspector::GetTextId(Text);
diff --git a/Engine/Source/Runtime/CoreUObject/Public/Serialization/PropertyLocalizationDataGathering.h b/Engine/Source/Runtime/CoreUObject/Public/Serialization/PropertyLocalizationDataGathering.h
--- a/Engine/Source/Runtime/CoreUObject/Public/Serialization/PropertyLocalizationDataGathering.h
+++ b/Engine/Source/Runtime/CoreUObject/Public/Serialization/PropertyLocalizationDataGathering.h
@@ -152,6 +152,8 @@ public:
 
 	COREUOBJECT_API const FGatherableFieldsForType& GetGatherableFieldsForType(const UStruct* InType);
 
+	static COREUOBJECT_API bool IsEditorOnlyProperty(const FProperty* Property);
+
 	static COREUOBJECT_API bool ExtractTextIdentity(const FText& Text, FString& OutNamespace, FString& OutKey, const bool bCleanNamespace);
 
 	static COREUOBJECT_API FLocalizationDataObjectGatheringCallbackMap& GetTypeSpecificLocalizationDataObjectGatheringCallbacks();
diff --git a/Engine/Source/Runtime/Engine/Public/Materials/MaterialExpressionTextureSample.h b/Engine/Source/Runtime/Engine/Public/Materials/MaterialExpressionTextureSample.h
--- a/Engine/Source/Runtime/Engine/Public/Materials/MaterialExpressionTextureSample.h
+++ b/Engine/Source/Runtime/Engine/Public/Materials/MaterialExpressionTextureSample.h
@@ -80,7 +80,7 @@ class UMaterialExpressionTextureSample : public UMaterialExpressionTextureBase
 	UPROPERTY(EditAnywhere, Category = MaterialExpressionTextureSample)
 	uint8 AutomaticViewMipBias : 1;
 
-	UPROPERTY(EditAnywhere, Category = ParameterCustomization)
+	UPROPERTY(EditAnywhere, Category = ParameterCustomization, meta=(EditorOnlyLocalization))
 	FParameterChannelNames ChannelNames;
 
 protected:
diff --git a/Engine/Source/Runtime/Engine/Public/Materials/MaterialExpressionVectorParameter.h b/Engine/Source/Runtime/Engine/Public/Materials/MaterialExpressionVectorParameter.h
--- a/Engine/Source/Runtime/Engine/Public/Materials/MaterialExpressionVectorParameter.h
+++ b/Engine/Source/Runtime/Engine/Public/Materials/MaterialExpressionVectorParameter.h
@@ -23,7 +23,7 @@ class UMaterialExpressionVectorParameter : public UMaterialExpressionParameter
 	UPROPERTY(EditAnywhere, Category=CustomPrimitiveData, meta=(ClampMin="0"))
 	uint8 PrimitiveDataIndex = 0;
 
-	UPROPERTY(EditAnywhere, Category = ParameterCustomization)
+	UPROPERTY(EditAnywhere, Category = ParameterCustomization, meta=(EditorOnlyLocalization))
 	FParameterChannelNames ChannelNames;
 
 	//~ Begin UMaterialExpression Interface
