Localization:
Prevent text properties initialized with FText::FromString() from having the localization data cached in the header of packages.
- Because the way keying of the text works, when the package is first saved, the FText::FromString() is not keyed properly. Thus the results of the localization package cache is inaccurate.
- This prevents the package header from being generated and forces a load of the package to gather the text.
- This should fix issues that require double resaves of packages to be properly gathered for localization

diff --git a/Engine/Source/Runtime/CoreUObject/Private/Serialization/PropertyLocalizationDataGathering.cpp b/Engine/Source/Runtime/CoreUObject/Private/Serialization/PropertyLocalizationDataGathering.cpp
--- a/Engine/Source/Runtime/CoreUObject/Private/Serialization/PropertyLocalizationDataGathering.cpp
+++ b/Engine/Source/Runtime/CoreUObject/Private/Serialization/PropertyLocalizationDataGathering.cpp
@@ -597,6 +597,13 @@ void FPropertyLocalizationDataGatherer::GatherTextInstance(const FText& Text, co
 		}
 	};
 
+	if (Text.IsInitializedFromString())
+	{
+		// We early out as a text initialized from a string does not have a namespace or key and thus can't be gathered. into the array
+		ResultFlags |= EPropertyLocalizationGathererResultFlags::HasTextInitializedFromString;
+		return;
+	}
+
 	FString Namespace;
 	FString Key;
 	if (!ExtractTextIdentity(Text, Namespace, Key, /*bCleanNamespace*/false))
diff --git a/Engine/Source/Runtime/CoreUObject/Private/UObject/SavePackage2.cpp b/Engine/Source/Runtime/CoreUObject/Private/UObject/SavePackage2.cpp
--- a/Engine/Source/Runtime/CoreUObject/Private/UObject/SavePackage2.cpp
+++ b/Engine/Source/Runtime/CoreUObject/Private/UObject/SavePackage2.cpp
@@ -2214,7 +2214,8 @@ void WriteGatherableText(FStructuredArchive::FRecord& StructuredArchiveRoot, FSa
 		// We can only cache packages that:
 		//	1) Don't contain script data, as script data is very volatile and can only be safely gathered after it's been compiled (which happens automatically on asset load).
 		//	2) Don't contain text keyed with an incorrect package localization ID, as these keys will be changed later during save.
-		&& !EnumHasAnyFlags(SaveContext.GatherableTextResultFlags, EPropertyLocalizationGathererResultFlags::HasScript | EPropertyLocalizationGathererResultFlags::HasTextWithInvalidPackageLocalizationID))
+		//	3) Don't contain text initialized from a string which includes text properties pasted into the property matrix 
+		&& !EnumHasAnyFlags(SaveContext.GatherableTextResultFlags, EPropertyLocalizationGathererResultFlags::HasScript | EPropertyLocalizationGathererResultFlags::HasTextWithInvalidPackageLocalizationID | EPropertyLocalizationGathererResultFlags::HasTextInitializedFromString))
 	{
 		FLinkerSave* Linker = SaveContext.GetLinker();
 
diff --git a/Engine/Source/Runtime/CoreUObject/Public/Serialization/PropertyLocalizationDataGathering.h b/Engine/Source/Runtime/CoreUObject/Public/Serialization/PropertyLocalizationDataGathering.h
--- a/Engine/Source/Runtime/CoreUObject/Public/Serialization/PropertyLocalizationDataGathering.h
+++ b/Engine/Source/Runtime/CoreUObject/Public/Serialization/PropertyLocalizationDataGathering.h
@@ -97,6 +97,11 @@ enum class EPropertyLocalizationGathererResultFlags : uint8
 	 * The call resulted in text with an invalid package localization ID being added to the array.
 	 */
 	HasTextWithInvalidPackageLocalizationID = 1<<2,
+	
+	/**
+	* The call resulted in text initialized from a string being detected. The text is NOT added to the array.
+	*/
+	HasTextInitializedFromString = 1 << 3,
 };
 ENUM_CLASS_FLAGS(EPropertyLocalizationGathererResultFlags);
 
