Suppressing TSAN warnings related to GIsIncrementalReachabilityPending. It's not thread safe but incremental reachability analysis is currently an experimental feature and we need a better solution than an atomic.

diff --git a/Engine/Source/Runtime/CoreUObject/Public/UObject/GarbageCollectionGlobals.h b/Engine/Source/Runtime/CoreUObject/Public/UObject/GarbageCollectionGlobals.h
--- a/Engine/Source/Runtime/CoreUObject/Public/UObject/GarbageCollectionGlobals.h
+++ b/Engine/Source/Runtime/CoreUObject/Public/UObject/GarbageCollectionGlobals.h
@@ -23,5 +23,5 @@ namespace UE::GC
 	extern COREUOBJECT_API EInternalObjectFlags GMaybeUnreachableObjectFlag;
 
 	/** true if incremental reachability analysis is in progress (global for faster access in low level structs and functions otherwise use IsIncrementalReachabilityAnalisysPending()) */
-	extern COREUOBJECT_API bool GIsIncrementalReachabilityPending;
+	extern COREUOBJECT_API TSAN_ATOMIC(bool) GIsIncrementalReachabilityPending;
 }
