Shrink back asset registry tag allocations to reduce cooker memory
#rnx

diff --git a/Engine/Source/Runtime/CoreUObject/Private/UObject/Obj.cpp b/Engine/Source/Runtime/CoreUObject/Private/UObject/Obj.cpp
--- a/Engine/Source/Runtime/CoreUObject/Private/UObject/Obj.cpp
+++ b/Engine/Source/Runtime/CoreUObject/Private/UObject/Obj.cpp
@@ -2351,6 +2351,7 @@ static void GetAssetRegistryTagFromProperty(const void* BaseMemoryLocation, cons
 			PropertyStr.Reserve(8192);
 			const uint8* PropertyAddr = Prop->ContainerPtrToValuePtr<uint8>(BaseMemoryLocation);
 			Prop->ExportTextItem_Direct(PropertyStr, PropertyAddr, PropertyAddr, nullptr, PPF_None);
+			PropertyStr.Shrink();
 		}
 
 		Context.AddTag(UObject::FAssetRegistryTag(Prop->GetFName(), MoveTemp(PropertyStr), TagType));
