Fix ReferenceChainSearch giving the wrong property name when references are within nested schema (i.e. arrays of structs).
Pass UE::GC::FWorkerContext in to reference processors and use SchemaStack to get property name/path.
Change FReferenceChainSearch::TReferneceInfo::ReferencerName to FString
Allow FPropertyStack::ConvertPathToProperties to take FStringView rather than FName.

diff --git a/Engine/Source/Runtime/CoreUObject/Private/UObject/GarbageCollectionSchema.cpp b/Engine/Source/Runtime/CoreUObject/Private/UObject/GarbageCollectionSchema.cpp
--- a/Engine/Source/Runtime/CoreUObject/Private/UObject/GarbageCollectionSchema.cpp
+++ b/Engine/Source/Runtime/CoreUObject/Private/UObject/GarbageCollectionSchema.cpp
@@ -1,6 +1,8 @@
 // Copyright Epic Games, Inc. All Rights Reserved.
 
 #include "UObject/GarbageCollectionSchema.h"
+
+#include "Misc/AsciiSet.h"
 #include "UObject/FastReferenceCollector.h"
 #include "UObject/UObjectIterator.h"
 #include "UObject/UnrealType.h"
@@ -567,15 +569,20 @@ FString FPropertyStack::GetPropertyPath() const
 
 bool FPropertyStack::ConvertPathToProperties(UClass* ObjectClass, FName InPropertyPath, TArray<FProperty*>& OutProperties)
 {
-	const TCHAR DelimiterChar = TEXT('.');
-	FString PropertyNameOrPath = InPropertyPath.ToString();
-	int32 DelimiterIndex = -1;
+	FNameBuilder Buffer(InPropertyPath);
+	return ConvertPathToProperties(ObjectClass, Buffer.ToView(), OutProperties);
+}
+
+bool FPropertyStack::ConvertPathToProperties(UClass* ObjectClass, FStringView PropertyNameOrPath, TArray<FProperty*>& OutProperties)
+{
+	constexpr FAsciiSet Delimiters(".");
 	bool bFullPathConstructed = true;
 
-	if (!PropertyNameOrPath.FindChar(DelimiterChar, DelimiterIndex))
+	FStringView PropertyName = FAsciiSet::FindPrefixWithout(PropertyNameOrPath, Delimiters);
+	if (PropertyName.Len() == PropertyNameOrPath.Len())
 	{
 		// 99% of the time we're be dealing with just a single property
-		if (FProperty* FoundProperty = ObjectClass->FindPropertyByName(*PropertyNameOrPath))
+		if (FProperty* FoundProperty = ObjectClass->FindPropertyByName(FName(PropertyNameOrPath)))
 		{
 			OutProperties.Add(FoundProperty);
 		}
@@ -587,28 +594,24 @@ bool FPropertyStack::ConvertPathToProperties(UClass* ObjectClass, FName InProper
 	else
 	{
 		// Try and find the first property as we can't start processing the rest of the path without it
-		FString PropertyName = PropertyNameOrPath.Left(DelimiterIndex);
-		if (FProperty* FoundProperty = ObjectClass->FindPropertyByName(*PropertyName))
+		if (FProperty* FoundProperty = ObjectClass->FindPropertyByName(FName(PropertyName)))
 		{
+			FStringView Remaining = FAsciiSet::TrimPrefixWith(PropertyNameOrPath.RightChop(PropertyName.Len()), Delimiters);
 			OutProperties.Add(FoundProperty);
-
-			int32 StartIndex = DelimiterIndex + 1;
-			const TCHAR DelimiterStr[] = { DelimiterChar, TEXT('\0') };
 			do
 			{
 				// Determine the next property name
-				DelimiterIndex = PropertyNameOrPath.Find(DelimiterStr, ESearchCase::CaseSensitive, ESearchDir::FromStart, StartIndex);
-				PropertyName = PropertyNameOrPath.Mid(StartIndex, DelimiterIndex >= 0 ? (DelimiterIndex - StartIndex) : (PropertyNameOrPath.Len() - StartIndex));
-
+				PropertyName = FAsciiSet::FindPrefixWithout(Remaining, Delimiters);
+				Remaining = FAsciiSet::TrimPrefixWith(Remaining.RightChop(PropertyName.Len()), Delimiters);
 				if (FStructProperty* StructProp = CastField<FStructProperty>(FoundProperty))
 				{
 					// If the previous property was a struct property, the next one belongs to the struct the previous property represented
-					FoundProperty = StructProp->Struct->FindPropertyByName(*PropertyName);
+					FoundProperty = StructProp->Struct->FindPropertyByName(FName(PropertyName));
 				}
 				else
 				{
 					// In all other case (though in reality it should only be a TMap) find the inner property
-					FoundProperty = CastField<FProperty>(FoundProperty->GetInnerFieldByName(*PropertyName));
+					FoundProperty = CastField<FProperty>(FoundProperty->GetInnerFieldByName(FName(PropertyName)));
 				}
 
 				if (FoundProperty)
@@ -619,7 +622,7 @@ bool FPropertyStack::ConvertPathToProperties(UClass* ObjectClass, FName InProper
 				{
 					bFullPathConstructed = false;
 				}
-			} while (DelimiterIndex >= 0 && bFullPathConstructed);
+			} while (Remaining.Len() > 0 && bFullPathConstructed);
 		}
 		else
 		{
diff --git a/Engine/Source/Runtime/CoreUObject/Private/UObject/ReferenceChainSearch.cpp b/Engine/Source/Runtime/CoreUObject/Private/UObject/ReferenceChainSearch.cpp
--- a/Engine/Source/Runtime/CoreUObject/Private/UObject/ReferenceChainSearch.cpp
+++ b/Engine/Source/Runtime/CoreUObject/Private/UObject/ReferenceChainSearch.cpp
@@ -261,7 +261,7 @@ namespace UE::ReferenceChainSearch
 
 					if (Object != ReferencingObject)
 					{
-						This->HandleObjectReference(ThreadIndex, Object, ReferencingObject, MemberId, Origin);
+						This->HandleObjectReference(Context, ThreadIndex, Object, ReferencingObject, MemberId, Origin);
 					}
 				}
 			}
@@ -506,6 +506,7 @@ namespace UE::ReferenceChainSearch
 		}
 
 		void HandleObjectReference(
+			UE::GC::FWorkerContext& Context,
 			int32 ThreadIndex, UObject* Object, UObject* ReferencingObject, UE::GC::FTokenId TokenIndex, UE::GC::EOrigin TokenType)
 		{
 			FThreadData& ThreadData = AllThreadData[ThreadIndex];
@@ -571,6 +572,7 @@ namespace UE::ReferenceChainSearch
 		}
 
 		void HandleObjectReference(
+			UE::GC::FWorkerContext& Context,
 			int32 ThreadIndex,
 			UObject* Object,
 			UObject* ReferencingObject,
@@ -1073,6 +1075,7 @@ namespace UE::ReferenceChainSearch
 		}
 
 		void HandleObjectReference(
+			UE::GC::FWorkerContext& Context,
 			int32 ThreadIndex, UObject* Object, UObject* ReferencingObject, UE::GC::FMemberId MemberId, UE::GC::EOrigin Origin)
 		{
 			if (!Object || Object == ReferencingObject) // Skip self-references just in case
@@ -1099,8 +1102,7 @@ namespace UE::ReferenceChainSearch
 			*RefInfo = FReferenceChainSearch::FObjectReferenceInfo(ObjectInfo);
 			if (MemberId != UE::GC::EMemberlessId::Collector)
 			{
-				FName MemberName = GetMemberDebugInfo(ReferencingObject->GetClass()->ReferenceSchema.Get(), MemberId).Name;
-				RefInfo->ReferencerName = MemberName;
+				RefInfo->ReferencerName = Context.SchemaStack->ToString();
 				RefInfo->Type = FReferenceChainSearch::EReferenceType::Property;
 			}
 			else
@@ -1112,27 +1114,17 @@ namespace UE::ReferenceChainSearch
 
 				if (FGCObject::GGCObjectReferencer && (!ReferencingObject || ReferencingObject == FGCObject::GGCObjectReferencer))
 				{
-					FString RefName;
-					if (FGCObject::GGCObjectReferencer->GetReferencerName(Object, RefName, true))
+					if (FGCObject::GGCObjectReferencer->GetReferencerName(Object, RefInfo->ReferencerName, true))
 					{
-						// In case FGCObjects misbehave and give us long strings, truncate so we can make them an FName.
-						if (RefName.Len() >= NAME_SIZE)
-						{
-							RefName.LeftInline(NAME_SIZE - 1);
-						}
-						else
-						{
-							RefInfo->ReferencerName = FName(*RefName);
-						}
 					}
 					else if (FGCObject* ReferencingFGCObject = FGCObject::GGCObjectReferencer->GetCurrentlySerializingObject())
 					{
-						RefInfo->ReferencerName = *ReferencingFGCObject->GetReferencerName();
+						RefInfo->ReferencerName = ReferencingFGCObject->GetReferencerName();
 					}
 				}
 				else if (ReferencingObject)
 				{
-					RefInfo->ReferencerName = *ReferencingObject->GetFullName();
+					RefInfo->ReferencerName = ReferencingObject->GetFullName();
 				}
 			}
 		}
@@ -1171,7 +1163,7 @@ namespace UE::ReferenceChainSearch
 					if (RefInfo)
 					{
 						FReferenceChainSearch::EReferenceType ReferenceType = FReferenceChainSearch::EReferenceType::Unknown;
-						if (SourceVertex == GCObjReferencerVertex || RefInfo->ReferencerName == NAME_None)
+						if (SourceVertex == GCObjReferencerVertex || RefInfo->ReferencerName.IsNone())
 						{
 							ReferenceType = FReferenceChainSearch::EReferenceType::AddReferencedObjects;
 						}
@@ -1179,7 +1171,7 @@ namespace UE::ReferenceChainSearch
 						{
 							ReferenceType = FReferenceChainSearch::EReferenceType::Property;
 						}
-						TargetPair.Value = FReferenceChainSearch::FObjectReferenceInfo(TargetInfo, ReferenceType, RefInfo->ReferencerName);
+						TargetPair.Value = FReferenceChainSearch::FObjectReferenceInfo(TargetInfo, ReferenceType, RefInfo->ReferencerName.ToString());
 					}
 				}
 			}
@@ -1451,13 +1443,13 @@ void FReferenceChainSearch::DumpChain(FReferenceChainSearch::FReferenceChain* Ch
 						*InnermostProperty->GetCPPType(),
 						OutermostProperty->GetOwnerClass()->GetPrefixCPP(),
 						*OutermostProperty->GetOwnerClass()->GetName(),
-						*ReferenceInfo->ReferencerName.ToString());
+						*ReferenceInfo->ReferencerName);
 				}
 				else
 				{
 					// Handle base UObject referencer info (it's only exposed to the GC token stream and not to the reflection system)
-					static const FName ClassPropertyName(TEXT("Class"));
-					static const FName OuterPropertyName(TEXT("Outer"));
+					static const FStringView ClassPropertyName = TEXTVIEW("Class");
+					static const FStringView OuterPropertyName = TEXTVIEW("Outer");
 
 					FString ClassName;
 					if (ReferenceInfo->ReferencerName == ClassPropertyName || ReferenceInfo->ReferencerName == OuterPropertyName)
@@ -1475,7 +1467,7 @@ void FReferenceChainSearch::DumpChain(FReferenceChainSearch::FReferenceChain* Ch
 						// Revert to the internal class name if not
 						ClassName = ReferencerObject->GetClassName();
 					}
-					ReferencingPropertyName = FString::Printf(TEXT("UObject* %s::%s"), *ClassName, *ReferenceInfo->ReferencerName.ToString());
+					ReferencingPropertyName = FString::Printf(TEXT("UObject* %s::%s"), *ClassName, *ReferenceInfo->ReferencerName);
 				}
 
 				Out.Logf(InVerbosityForPrint, TEXT("%s-> %s = %s %s"),
@@ -1487,7 +1479,7 @@ void FReferenceChainSearch::DumpChain(FReferenceChainSearch::FReferenceChain* Ch
 			else if (ReferenceInfo && ReferenceInfo->Type == EReferenceType::AddReferencedObjects)
 			{
 				FString UObjectOrGCObjectName;
-				if (ReferenceInfo->ReferencerName.IsNone())
+				if (ReferenceInfo->ReferencerName.IsEmpty())
 				{
 					UClass* ReferencerClass = Cast<UClass>(ReferencerObject->GetClass()->TryResolveObject());
 					if (ReferencerClass)
@@ -1502,7 +1494,7 @@ void FReferenceChainSearch::DumpChain(FReferenceChainSearch::FReferenceChain* Ch
 				}
 				else
 				{
-					UObjectOrGCObjectName = ReferenceInfo->ReferencerName.ToString();
+					UObjectOrGCObjectName = ReferenceInfo->ReferencerName;
 				}
 
 				Out.Logf(InVerbosityForPrint, TEXT("%s-> %s::AddReferencedObjects(%s %s)"),
diff --git a/Engine/Source/Runtime/CoreUObject/Public/UObject/GarbageCollectionSchema.h b/Engine/Source/Runtime/CoreUObject/Public/UObject/GarbageCollectionSchema.h
--- a/Engine/Source/Runtime/CoreUObject/Public/UObject/GarbageCollectionSchema.h
+++ b/Engine/Source/Runtime/CoreUObject/Public/UObject/GarbageCollectionSchema.h
@@ -76,6 +76,7 @@ public:
 	 * @returns true if the conversion was successful, false otherwise
 	 */
 	COREUOBJECT_API static bool ConvertPathToProperties(UClass* ObjectClass, FName InPropertyPath, TArray<FProperty*>& OutProperties);
+	COREUOBJECT_API static bool ConvertPathToProperties(UClass* ObjectClass, FStringView InPropertyPath, TArray<FProperty*>& OutProperties);
 private:
 	friend class FPropertyStackScope;
 	TArray<const FProperty*> Props;
diff --git a/Engine/Source/Runtime/CoreUObject/Public/UObject/ReferenceChainSearch.h b/Engine/Source/Runtime/CoreUObject/Public/UObject/ReferenceChainSearch.h
--- a/Engine/Source/Runtime/CoreUObject/Public/UObject/ReferenceChainSearch.h
+++ b/Engine/Source/Runtime/CoreUObject/Public/UObject/ReferenceChainSearch.h
@@ -91,7 +91,7 @@ public:
 		/** Type of reference to the object being referenced */
 		EReferenceType Type;
 		/** Name of the object or property that is referencing this object */
-		FName ReferencerName;
+		FString ReferencerName;
 
 		/** Default ctor */
 		TReferenceInfo()
@@ -110,13 +110,13 @@ public:
 		TReferenceInfo(
 			T* InObject, 
 			EReferenceType InType, 
-			FName InReferencerName,
+			FString InReferencerName,
 			TConstArrayView<uint64> InStackFrames = {}
 		)
 			: Object(InObject)
 			, StackFrames(InStackFrames)
 			, Type(InType)
-			, ReferencerName(InReferencerName)
+			, ReferencerName(MoveTemp(InReferencerName))
 		{
 		}
 
@@ -138,13 +138,13 @@ public:
 		{
 			if (Type == EReferenceType::Property)
 			{
-				return FString::Printf(TEXT("->%s"), *ReferencerName.ToString());
+				return FString::Printf(TEXT("->%s"), *ReferencerName);
 			}
 			else if (Type == EReferenceType::AddReferencedObjects)
 			{
-				if (!ReferencerName.IsNone())
+				if (!ReferencerName.IsEmpty())
 				{
-					return FString::Printf(TEXT("::AddReferencedObjects(): %s"), *ReferencerName.ToString());
+					return FString::Printf(TEXT("::AddReferencedObjects(): %s"), *ReferencerName);
 				}
 				else
 				{
