FSaveConfigContext takes a set of PropertyNames instead of just a SinglePropertyName

diff --git a/Engine/Source/Editor/DetailCustomizations/Private/ConfigEditorPropertyDetails.cpp b/Engine/Source/Editor/DetailCustomizations/Private/ConfigEditorPropertyDetails.cpp
--- a/Engine/Source/Editor/DetailCustomizations/Private/ConfigEditorPropertyDetails.cpp
+++ b/Engine/Source/Editor/DetailCustomizations/Private/ConfigEditorPropertyDetails.cpp
@@ -106,7 +106,7 @@ void FConfigPropertyHelperDetails::OnPropertyValueChanged(UObject* Object, FProp
 			UObject::FSaveConfigContext Context;
 			Context.Filename = ConfigIniName;
 			Context.bSaveToLayer = true;
-			Context.SinglePropertyName = OriginalProperty->GetFName();
+			Context.PropertyNames = { OriginalProperty->GetFName() };
 			Context.OverrideConfigSectionName = OriginalProperty->GetOwnerClass()->GetPathName();
 			Object->SaveConfig(Context);
 
diff --git a/Engine/Source/Runtime/CoreUObject/Private/UObject/Obj.cpp b/Engine/Source/Runtime/CoreUObject/Private/UObject/Obj.cpp
--- a/Engine/Source/Runtime/CoreUObject/Private/UObject/Obj.cpp
+++ b/Engine/Source/Runtime/CoreUObject/Private/UObject/Obj.cpp
@@ -3556,7 +3556,7 @@ void UObject::SaveConfig(const FSaveConfigContext& Context)
 	for (FProperty* Property = GetClass()->PropertyLink; Property; Property = Property->PropertyLinkNext)
 	{
 		// we just check names, in case this property is a duplicated but same property as this object's class
-		if (Context.SinglePropertyName != NAME_None && Property->GetFName() != Context.SinglePropertyName)
+		if (!Context.PropertyNames.IsEmpty() && !Context.PropertyNames.Contains(Property->GetFName()))
 		{
 			continue;
 		}
@@ -3891,7 +3891,7 @@ void UObject::UpdateSinglePropertyInConfigFile(const FProperty* InProperty, cons
 		FSaveConfigContext Context;
 		Context.Filename = InConfigIniName;
 		Context.bSaveToLayer = true;
-		Context.SinglePropertyName = InProperty->GetFName();
+		Context.PropertyNames = { InProperty->GetFName() };
 		SaveConfig(Context);
 		return;
 	}
diff --git a/Engine/Source/Runtime/CoreUObject/Public/UObject/Object.h b/Engine/Source/Runtime/CoreUObject/Public/UObject/Object.h
--- a/Engine/Source/Runtime/CoreUObject/Public/UObject/Object.h
+++ b/Engine/Source/Runtime/CoreUObject/Public/UObject/Object.h
@@ -1298,7 +1298,7 @@ public:
 		bool bAllowCopyToDefaultObject = true;
 
 		bool bSaveToLayer = false;
-		FName SinglePropertyName;
+		TSet<FName> PropertyNames = {};
 		FString OverrideConfigSectionName;
 	};
 	COREUOBJECT_API void SaveConfig(const FSaveConfigContext& Context);
