Initial support for creating Property Bag Placeholder types for entity prefabs.

To achieve this a new UE5 object version is introduced that adds a new table to packages: ImportTypeHierarchy. Each entry of that able contains hierarchical information for every FObjectImport entry that is a type (e.g., UClass).

This change only introduces the object version bump and the corresponding package file summary header fields. The code to serialize this new table will follow as a separate change.

#rnx

diff --git a/Engine/Source/Runtime/Core/Public/UObject/ObjectVersion.h b/Engine/Source/Runtime/Core/Public/UObject/ObjectVersion.h
--- a/Engine/Source/Runtime/Core/Public/UObject/ObjectVersion.h
+++ b/Engine/Source/Runtime/Core/Public/UObject/ObjectVersion.h
@@ -99,6 +99,9 @@ enum class EUnrealEngineObjectUE5Version : uint32
 	// OS shadow serialization of subobjects
 	OS_SUB_OBJECT_SHADOW_SERIALIZATION,
 
+	// Adds a table of hierarchical type information for imports in a package
+	IMPORT_TYPE_HIERARCHIES,
+
 	// -----<new versions can be added before this line>-------------------------------------------------
 	// - this needs to be the last line (see note below)
 	AUTOMATIC_VERSION_PLUS_ONE,
diff --git a/Engine/Source/Runtime/CoreUObject/Private/UObject/PackageFileSummary.cpp b/Engine/Source/Runtime/CoreUObject/Private/UObject/PackageFileSummary.cpp
--- a/Engine/Source/Runtime/CoreUObject/Private/UObject/PackageFileSummary.cpp
+++ b/Engine/Source/Runtime/CoreUObject/Private/UObject/PackageFileSummary.cpp
@@ -324,6 +324,16 @@ void operator<<(FStructuredArchive::FSlot Slot, FPackageFileSummary& Sum)
 
 		Record << SA_VALUE(TEXT("ThumbnailTableOffset"), Sum.ThumbnailTableOffset);
 
+		if (Sum.FileVersionUE >= EUnrealEngineObjectUE5Version::IMPORT_TYPE_HIERARCHIES)
+		{
+			Record << SA_VALUE(TEXT("ImportTypeHierarchiesCount"), Sum.ImportTypeHierarchiesCount) << SA_VALUE(TEXT("ImportTypeHierarchiesOffset"), Sum.ImportTypeHierarchiesOffset);
+		}
+		else
+		{
+			Sum.ImportTypeHierarchiesCount = 0;
+			Sum.ImportTypeHierarchiesOffset = 0;
+		}
+
 		if (Sum.GetFileVersionUE() < EUnrealEngineObjectUE5Version::PACKAGE_SAVED_HASH)
 		{
 			FGuid LegacyGuid;
diff --git a/Engine/Source/Runtime/CoreUObject/Public/UObject/PackageFileSummary.h b/Engine/Source/Runtime/CoreUObject/Public/UObject/PackageFileSummary.h
--- a/Engine/Source/Runtime/CoreUObject/Public/UObject/PackageFileSummary.h
+++ b/Engine/Source/Runtime/CoreUObject/Public/UObject/PackageFileSummary.h
@@ -192,6 +192,16 @@ public:
 	*/
 	int32		ThumbnailTableOffset;
 
+	/**
+	* Number of import type hierarchy entries
+	*/
+	int32		ImportTypeHierarchiesCount;
+
+	/**
+	* Location into the file on disk for the import type hierarchy map
+	*/
+	int32		ImportTypeHierarchiesOffset;
+
 private:
 	/**
 	 * Hash of the Package's bytes when it was saved to disk. Accessed through GetSavedHash/SetSavedHash.
