Added IsFromAutoSave() method to object save context classes (FObjectPreSaveContext, FObjectPostSaveContext, FObjectSavePackageSerializeContext)
* To mimick what's already available in the FSaveContext class

diff --git a/Engine/Plugins/Cameras/GameplayCameras/Source/GameplayCameras/Private/Core/CameraAsset.cpp b/Engine/Plugins/Cameras/GameplayCameras/Source/GameplayCameras/Private/Core/CameraAsset.cpp
--- a/Engine/Plugins/Cameras/GameplayCameras/Source/GameplayCameras/Private/Core/CameraAsset.cpp
+++ b/Engine/Plugins/Cameras/GameplayCameras/Source/GameplayCameras/Private/Core/CameraAsset.cpp
@@ -175,7 +175,7 @@ void UCameraAsset::PreSave(FObjectPreSaveContext ObjectSaveContext)
 #if WITH_EDITOR
 
 	const bool bIsUserObject = !HasAnyFlags(RF_ClassDefaultObject | RF_ArchetypeObject);
-	const bool bIsEditorAutoSave = ((ObjectSaveContext.GetSaveFlags() & SAVE_FromAutosave) != 0);
+	const bool bIsEditorAutoSave = ObjectSaveContext.IsFromAutoSave();
 	if (bIsUserObject && !bIsEditorAutoSave)
 	{
 		// Build when saving/cooking.
diff --git a/Engine/Plugins/Cameras/GameplayCameras/Source/GameplayCameras/Private/Core/CameraRigAsset.cpp b/Engine/Plugins/Cameras/GameplayCameras/Source/GameplayCameras/Private/Core/CameraRigAsset.cpp
--- a/Engine/Plugins/Cameras/GameplayCameras/Source/GameplayCameras/Private/Core/CameraRigAsset.cpp
+++ b/Engine/Plugins/Cameras/GameplayCameras/Source/GameplayCameras/Private/Core/CameraRigAsset.cpp
@@ -123,7 +123,7 @@ void UCameraRigAsset::PreSave(FObjectPreSaveContext ObjectSaveContext)
 #if WITH_EDITOR
 
 	const bool bIsUserObject = !HasAnyFlags(RF_ClassDefaultObject | RF_ArchetypeObject);
-	const bool bIsEditorAutoSave = ((ObjectSaveContext.GetSaveFlags() & SAVE_FromAutosave) != 0);
+	const bool bIsEditorAutoSave = ObjectSaveContext.IsFromAutoSave();
 	if (bIsUserObject && !bIsEditorAutoSave)
 	{
 		// Build when saving/cooking.
diff --git a/Engine/Plugins/Cameras/GameplayCameras/Source/GameplayCameras/Private/Core/CameraShakeAsset.cpp b/Engine/Plugins/Cameras/GameplayCameras/Source/GameplayCameras/Private/Core/CameraShakeAsset.cpp
--- a/Engine/Plugins/Cameras/GameplayCameras/Source/GameplayCameras/Private/Core/CameraShakeAsset.cpp
+++ b/Engine/Plugins/Cameras/GameplayCameras/Source/GameplayCameras/Private/Core/CameraShakeAsset.cpp
@@ -63,7 +63,7 @@ void UCameraShakeAsset::PreSave(FObjectPreSaveContext ObjectSaveContext)
 #if WITH_EDITOR
 
 	const bool bIsUserObject = !HasAnyFlags(RF_ClassDefaultObject | RF_ArchetypeObject);
-	const bool bIsEditorAutoSave = ((ObjectSaveContext.GetSaveFlags() & SAVE_FromAutosave) != 0);
+	const bool bIsEditorAutoSave = ObjectSaveContext.IsFromAutoSave();
 	if (bIsUserObject && !bIsEditorAutoSave)
 	{
 		// Build when saving/cooking.
diff --git a/Engine/Plugins/Importers/USDImporter/Source/USDStageEditor/Private/USDStageEditorModule.cpp b/Engine/Plugins/Importers/USDImporter/Source/USDStageEditor/Private/USDStageEditorModule.cpp
--- a/Engine/Plugins/Importers/USDImporter/Source/USDStageEditor/Private/USDStageEditorModule.cpp
+++ b/Engine/Plugins/Importers/USDImporter/Source/USDStageEditor/Private/USDStageEditorModule.cpp
@@ -517,7 +517,7 @@ public:
 			[](UWorld* World, FObjectPreSaveContext InContext)
 			{
 				// Detect if we should actually do anything (check for autosaves, cooking, etc.)
-				if (InContext.GetSaveFlags() & ESaveFlags::SAVE_FromAutosave || InContext.IsProceduralSave())
+				if (InContext.IsFromAutoSave() || InContext.IsProceduralSave())
 				{
 					return;
 				}
diff --git a/Engine/Source/Developer/UncontrolledChangelists/Private/UncontrolledChangelistsModule.cpp b/Engine/Source/Developer/UncontrolledChangelists/Private/UncontrolledChangelistsModule.cpp
--- a/Engine/Source/Developer/UncontrolledChangelists/Private/UncontrolledChangelistsModule.cpp
+++ b/Engine/Source/Developer/UncontrolledChangelists/Private/UncontrolledChangelistsModule.cpp
@@ -600,7 +600,7 @@ void FUncontrolledChangelistsModule::OnObjectPreSaved(UObject* InObject, FObject
 	}
 
 	// Ignore procedural save and autosaves
-	if (InPreSaveContext.IsProceduralSave() || ((InPreSaveContext.GetSaveFlags() & SAVE_FromAutosave) != 0))
+	if (InPreSaveContext.IsProceduralSave() || InPreSaveContext.IsFromAutoSave())
 	{
 		return;
 	}
diff --git a/Engine/Source/Developer/UnsavedAssetsTracker/Source/Private/UnsavedAssetsTracker.cpp b/Engine/Source/Developer/UnsavedAssetsTracker/Source/Private/UnsavedAssetsTracker.cpp
--- a/Engine/Source/Developer/UnsavedAssetsTracker/Source/Private/UnsavedAssetsTracker.cpp
+++ b/Engine/Source/Developer/UnsavedAssetsTracker/Source/Private/UnsavedAssetsTracker.cpp
@@ -307,7 +307,7 @@ void FUnsavedAssetsTracker::OnPackageDirtyStateUpdated(UPackage* Package)
 
 void FUnsavedAssetsTracker::OnPackageSaved(const FString& PackagePathname, UPackage* Package, FObjectPostSaveContext ObjectSaveContext)
 {
-	if (ObjectSaveContext.IsProceduralSave() || (ObjectSaveContext.GetSaveFlags() & SAVE_FromAutosave) != 0)
+	if (ObjectSaveContext.IsProceduralSave() || ObjectSaveContext.IsFromAutoSave())
 	{
 		return; // Don't track procedural save (during cooking) nor package auto-saved as backup in case of crash
 	}
diff --git a/Engine/Source/Editor/UnrealEd/Private/EditorEngine.cpp b/Engine/Source/Editor/UnrealEd/Private/EditorEngine.cpp
--- a/Engine/Source/Editor/UnrealEd/Private/EditorEngine.cpp
+++ b/Engine/Source/Editor/UnrealEd/Private/EditorEngine.cpp
@@ -4603,7 +4603,7 @@ void UEditorEngine::OnPreSaveWorld(UWorld* World, FObjectPreSaveContext ObjectSa
 	// Update cull distance volumes (and associated primitives).
 	World->UpdateCullDistanceVolumes();
 
-	const bool bAutosave = (ObjectSaveContext.GetSaveFlags() & SAVE_FromAutosave) != 0;
+	const bool bAutosave = ObjectSaveContext.IsFromAutoSave();
 	if (bAutosave)
 	{
 		// Temporarily flag packages saved under a PIE filename as PKG_PlayInEditor for serialization so loading
@@ -4681,7 +4681,7 @@ void UEditorEngine::OnPostSaveWorld(UWorld* World, FObjectPostSaveContext Object
 	if ( !IsRunningCommandlet() )
 	{
 		UPackage* WorldPackage = World->GetOutermost();
-		const bool bAutosave = (ObjectSaveContext.GetSaveFlags() & SAVE_FromAutosave) != 0;
+		const bool bAutosave = ObjectSaveContext.IsFromAutoSave();
 		if ( bAutosave )
 		{
 			// Restore original value of PKG_PlayInEditor if we changed it during PIE saving
diff --git a/Engine/Source/Editor/UnrealEd/Private/EditorTransaction.cpp b/Engine/Source/Editor/UnrealEd/Private/EditorTransaction.cpp
--- a/Engine/Source/Editor/UnrealEd/Private/EditorTransaction.cpp
+++ b/Engine/Source/Editor/UnrealEd/Private/EditorTransaction.cpp
@@ -59,7 +59,7 @@ private:
 
 	void OnPackageSaved(const FString& Filename, UPackage* Pkg, FObjectPostSaveContext ObjectSaveContext)
 	{
-		if (!(ObjectSaveContext.GetSaveFlags() & SAVE_FromAutosave) && !ObjectSaveContext.IsProceduralSave())
+		if (!ObjectSaveContext.IsFromAutoSave() && !ObjectSaveContext.IsProceduralSave())
 		{
 			IncrementCount(Pkg);
 		}
diff --git a/Engine/Source/Editor/UnrealEd/Private/UnrealEdMisc.cpp b/Engine/Source/Editor/UnrealEd/Private/UnrealEdMisc.cpp
--- a/Engine/Source/Editor/UnrealEd/Private/UnrealEdMisc.cpp
+++ b/Engine/Source/Editor/UnrealEd/Private/UnrealEdMisc.cpp
@@ -1233,7 +1233,7 @@ void FUnrealEdMisc::CB_RefreshEditor()
 void FUnrealEdMisc::PreSaveWorld(UWorld* World, FObjectPreSaveContext ObjectSaveContext)
 {
 	LogAssetUpdate(World, ObjectSaveContext);
-	const bool bAutosaveOrPIE = (ObjectSaveContext.GetSaveFlags() & SAVE_FromAutosave) != 0;
+	const bool bAutosaveOrPIE = ObjectSaveContext.IsFromAutoSave();
 	if (bAutosaveOrPIE || !World || World != GEditor->GetEditorWorldContext().World() || !FEngineAnalytics::IsAvailable())
 	{
 		return;
diff --git a/Engine/Source/Runtime/CoreUObject/Public/UObject/ObjectSaveContext.h b/Engine/Source/Runtime/CoreUObject/Public/UObject/ObjectSaveContext.h
--- a/Engine/Source/Runtime/CoreUObject/Public/UObject/ObjectSaveContext.h
+++ b/Engine/Source/Runtime/CoreUObject/Public/UObject/ObjectSaveContext.h
@@ -304,6 +304,11 @@ public:
 	 */
 	bool IsProceduralSave() const { return Data.bProceduralSave; }
 
+	/**
+	 * Return whether the package is being saved due to the editor auto save feature.
+	 */
+	bool IsFromAutoSave() const { return (Data.SaveFlags & SAVE_FromAutosave) != 0; }
+
 	/**
 	 * Return whether LoadedPath of the package being saved is being updated.
 	 * This allows us to update the in-memory package when it is saved in editor to match its new save file.
@@ -425,6 +430,11 @@ public:
 	 */
 	bool IsProceduralSave() const { return Data.bProceduralSave; }
 
+	/**
+	 * Return whether the package is being saved due to the editor auto save feature.
+	 */
+	bool IsFromAutoSave() const { return (Data.SaveFlags & SAVE_FromAutosave) != 0; }
+
 	/**
 	 * Return whether LoadedPath of the package being saved is being updated.
 	 * This allows us to update the in-memory package when it is saved in editor to match its new save file.
@@ -486,6 +496,11 @@ public:
 	 */
 	bool IsProceduralSave() const { return Data.bProceduralSave; }
 
+	/**
+	 * Return whether the package is being saved due to the editor auto save feature.
+	 */
+	bool IsFromAutoSave() const { return (Data.SaveFlags & SAVE_FromAutosave) != 0;	}
+
 	/**
 	 * Return whether LoadedPath of the package being saved is being updated.
 	 * This allows us to update the in-memory package when it is saved in editor to match its new save file.
diff --git a/Engine/Source/Runtime/Engine/Private/LevelInstance/LevelInstanceEditorObject.cpp b/Engine/Source/Runtime/Engine/Private/LevelInstance/LevelInstanceEditorObject.cpp
--- a/Engine/Source/Runtime/Engine/Private/LevelInstance/LevelInstanceEditorObject.cpp
+++ b/Engine/Source/Runtime/Engine/Private/LevelInstance/LevelInstanceEditorObject.cpp
@@ -91,7 +91,7 @@ void ULevelInstanceEditorObject::OnPackageChanged(UPackage* Package)
 
 void ULevelInstanceEditorObject::OnObjectPreSave(UObject* Object, FObjectPreSaveContext SaveContext)
 {
-	if (!SaveContext.IsProceduralSave() && !(SaveContext.GetSaveFlags() & SAVE_FromAutosave))
+	if (!SaveContext.IsProceduralSave() && !SaveContext.IsFromAutoSave())
 	{
 		OnPackageChanged(Object->GetPackage());
 	}
diff --git a/Engine/Source/Runtime/Engine/Private/LevelInstance/LevelInstanceSubsystem.cpp b/Engine/Source/Runtime/Engine/Private/LevelInstance/LevelInstanceSubsystem.cpp
--- a/Engine/Source/Runtime/Engine/Private/LevelInstance/LevelInstanceSubsystem.cpp
+++ b/Engine/Source/Runtime/Engine/Private/LevelInstance/LevelInstanceSubsystem.cpp
@@ -435,7 +435,7 @@ void ULevelInstanceSubsystem::OnAssetsPreDelete(const TArray<UObject*>& Objects)
 
 void ULevelInstanceSubsystem::OnPreSaveWorldWithContext(UWorld* InWorld, FObjectPreSaveContext ObjectSaveContext)
 {
-	if (!(ObjectSaveContext.GetSaveFlags() & SAVE_FromAutosave) && !ObjectSaveContext.IsProceduralSave())
+	if (!ObjectSaveContext.IsFromAutoSave() && !ObjectSaveContext.IsProceduralSave())
 	{
 		if (const UPackage* WorldPackage = InWorld->GetPackage())
 		{
diff --git a/Engine/Source/Runtime/Engine/Private/World.cpp b/Engine/Source/Runtime/Engine/Private/World.cpp
--- a/Engine/Source/Runtime/Engine/Private/World.cpp
+++ b/Engine/Source/Runtime/Engine/Private/World.cpp
@@ -1802,7 +1802,7 @@ void UWorld::PostSaveRoot( FObjectPostSaveRootContext ObjectSaveContext )
 		}
 	}
 
-	if (!ObjectSaveContext.IsProceduralSave() && !(ObjectSaveContext.GetSaveFlags() & SAVE_FromAutosave))
+	if (!ObjectSaveContext.IsProceduralSave() && !ObjectSaveContext.IsFromAutoSave())
 	{
 		// Once saved, OriginalWorldName must match World's name
 		OriginalWorldName = GetFName();
diff --git a/Engine/Source/Runtime/Engine/Private/WorldPartition/ActorDescContainer.cpp b/Engine/Source/Runtime/Engine/Private/WorldPartition/ActorDescContainer.cpp
--- a/Engine/Source/Runtime/Engine/Private/WorldPartition/ActorDescContainer.cpp
+++ b/Engine/Source/Runtime/Engine/Private/WorldPartition/ActorDescContainer.cpp
@@ -490,7 +490,7 @@ void UActorDescContainer::OnDeletedObjectPlaceholderCreated(const UDeletedObject
 
 void UActorDescContainer::OnObjectPreSave(UObject* Object, FObjectPreSaveContext SaveContext)
 {
-	if (!SaveContext.IsProceduralSave() && !(SaveContext.GetSaveFlags() & SAVE_FromAutosave))
+	if (!SaveContext.IsProceduralSave() && !SaveContext.IsFromAutoSave())
 	{
 		if (const AActor* Actor = Cast<AActor>(Object))
 		{
diff --git a/Engine/Source/Runtime/Landscape/Private/Landscape.cpp b/Engine/Source/Runtime/Landscape/Private/Landscape.cpp
--- a/Engine/Source/Runtime/Landscape/Private/Landscape.cpp
+++ b/Engine/Source/Runtime/Landscape/Private/Landscape.cpp
@@ -3942,7 +3942,7 @@ void ALandscapeProxy::PreSave(FObjectPreSaveContext ObjectSaveContext)
 	}
 
 	// Update nanite (and block to wait for it).  Don't update nanite on auto-save, since its so slow.
-	if ((ObjectSaveContext.GetSaveFlags() & SAVE_FromAutosave) == 0)
+	if (!ObjectSaveContext.IsFromAutoSave())
 	{
 		if (ULandscapeInfo* LandscapeInfo = GetLandscapeInfo())
 		{
