Net
* Added clearer documentation around GetLifetimeReplicatedProps now that it's optional
* Cleaned up and added documentation to FLifetimeProperty variables

#rnx
#DisabledValidations Preflight, checkUnacceptableWords

diff --git a/Engine/Source/Runtime/CoreUObject/Public/UObject/CoreNet.h b/Engine/Source/Runtime/CoreUObject/Public/UObject/CoreNet.h
--- a/Engine/Source/Runtime/CoreUObject/Public/UObject/CoreNet.h
+++ b/Engine/Source/Runtime/CoreUObject/Public/UObject/CoreNet.h
@@ -300,25 +300,25 @@ class FLifetimeProperty
 {
 public:
 
-	uint16 RepIndex;
-	ELifetimeCondition Condition;
-	ELifetimeRepNotifyCondition RepNotifyCondition;
-	bool bIsPushBased;
+	/** The FProperty's RepIndex this struct is representing */
+	uint16 RepIndex = 0;
+
+	/** The replication condition assigned to this property. Default is 'None' which means no conditions are needed. */
+	ELifetimeCondition Condition = COND_None;
+
+	/** The RepNotify behavior of the property. Only useful if CPF_RepNotify is set on the property.  */
+	ELifetimeRepNotifyCondition RepNotifyCondition = REPNOTIFY_OnChanged;
+
+	/** Is this property using push model suggestions to help the RepSystem tell when it's really dirty. */
+	bool bIsPushBased = false;
+
+	/** Iris only: optional functor that will be used  to create and register a ReplicationFragment for the property */
 	UE::Net::CreateAndRegisterReplicationFragmentFunc CreateAndRegisterReplicationFragmentFunction = nullptr;
 
-	FLifetimeProperty()
-		: RepIndex(0)
-		, Condition(COND_None)
-		, RepNotifyCondition(REPNOTIFY_OnChanged)
-		, bIsPushBased(false)
-	{
-	}
+	FLifetimeProperty() = default;
 
 	FLifetimeProperty(int32 InRepIndex)
 		: RepIndex((uint16)InRepIndex)
-		, Condition(COND_None)
-		, RepNotifyCondition(REPNOTIFY_OnChanged)
-		, bIsPushBased(false)
 	{
 		check(InRepIndex <= 65535);
 	}
diff --git a/Engine/Source/Runtime/CoreUObject/Public/UObject/Object.h b/Engine/Source/Runtime/CoreUObject/Public/UObject/Object.h
--- a/Engine/Source/Runtime/CoreUObject/Public/UObject/Object.h
+++ b/Engine/Source/Runtime/CoreUObject/Public/UObject/Object.h
@@ -1042,7 +1042,13 @@ public:
 	 */
 	COREUOBJECT_API virtual void TagSubobjects(EObjectFlags NewFlags);
 
-	/** Returns properties that are replicated for the lifetime of the actor channel */
+	/** 
+	 * Returns the list of replicated properties for the class and their replication configuration.
+	 * This function used to require every replicated property of the class to be registered in the array. 
+	 * But now you can leave it empty (or call the Super) and the class's replicated properties will still be replicated.
+	 * The NetDriver will automatically add all properties that are CPF_Net and not found in the array.
+	 * Note that these auto-added properties will use default configs: No conditions, PushModel = off, etc.
+	 */
 	COREUOBJECT_API virtual void GetLifetimeReplicatedProps( TArray< class FLifetimeProperty > & OutLifetimeProps ) const;
 
 	/** Called when this object begins replicating to initialize the state of custom property conditions */
