Support for creating Property Bag Placeholder types for entity prefabs.

This is done via a new table added to package files: ImportTypeHierarchy table. Each entry of that table contains hierarchical information for every FObjectImport entry that is a type (e.g., UClass).

This is the second part of a 2 part change. The first change introduced the package file header change along with the object version bump. This change adds the remaining logic to save/load the new table and the change to use it when creating placeholders.

[RN] Support for creating Property Bag Placeholder types for entity prefabs.

diff --git a/Engine/Source/Runtime/CoreUObject/Internal/UObject/InstanceDataObjectUtils.h b/Engine/Source/Runtime/CoreUObject/Internal/UObject/InstanceDataObjectUtils.h
--- a/Engine/Source/Runtime/CoreUObject/Internal/UObject/InstanceDataObjectUtils.h
+++ b/Engine/Source/Runtime/CoreUObject/Internal/UObject/InstanceDataObjectUtils.h
@@ -18,6 +18,7 @@ class UStruct;
 
 namespace UE { class FPropertyPathNameTree; }
 namespace UE { class FUnknownEnumNames; }
+namespace UE::Serialization::Private { class FImportTypeHierarchy; }
 
 namespace UE
 {
@@ -45,7 +46,9 @@ UE_API bool StructIsInstanceDataObjectStruct(const UStruct* Struct);
 UE_API bool IsPropertyLoose(const FProperty* Property);
 
 /** Query if placeholder support is enabled for a specific import class type. */
-bool CanCreatePropertyBagPlaceholderTypeForImportClass(const UClass* ImportType);
+bool CanCreatePropertyBagPlaceholderTypeForImportClass(const UClass* ClassKind, const UE::Serialization::Private::FImportTypeHierarchy* ImportTypeHierarchy);
+/** Query if placeholder support is enabled for a specific type. */
+bool CanCreatePropertyBagPlaceholderForType(const UClass* ClassKind, const UStruct* Type);
 
 /** Helper to check if a class is an IDO class. */
 UE_API bool IsClassOfInstanceDataObjectClass(const UStruct* Class);
diff --git a/Engine/Source/Runtime/CoreUObject/Private/Serialization/ImportTypeHierarchy.cpp b/Engine/Source/Runtime/CoreUObject/Private/Serialization/ImportTypeHierarchy.cpp
new file mode 100644
--- /dev/null
+++ b/Engine/Source/Runtime/CoreUObject/Private/Serialization/ImportTypeHierarchy.cpp
@@ -0,0 +1,57 @@
+// Copyright Epic Games, Inc. All Rights Reserved.
+
+#include "Serialization/ImportTypeHierarchy.h"
+
+namespace UE::Serialization::Private
+{
+
+FArchive& operator<<(FArchive& Ar, FTypeResource& TypeResource)
+{
+	FStructuredArchiveFromArchive(Ar).GetSlot() << TypeResource;
+	return Ar;
+}
+
+void operator<<(FStructuredArchive::FSlot Slot, FTypeResource& TypeResource)
+{
+	FStructuredArchive::FRecord Record = Slot.EnterRecord();
+
+	Record << SA_VALUE(TEXT("TypeName"), TypeResource.TypeName);
+	Record << SA_VALUE(TEXT("PackageName"), TypeResource.PackageName);
+	Record << SA_VALUE(TEXT("ClassName"), TypeResource.ClassName);
+	Record << SA_VALUE(TEXT("ClassPackageName"), TypeResource.ClassPackageName);
+}
+
+FArchive& operator<<(FArchive& Ar, FImportTypeHierarchy& I)
+{
+	FStructuredArchiveFromArchive(Ar).GetSlot() << I;
+	return Ar;
+}
+
+void operator<<(FStructuredArchive::FSlot Slot, FImportTypeHierarchy& I)
+{
+	FStructuredArchive::FRecord Record = Slot.EnterRecord();
+
+	int32 SuperTypeCount = I.SuperTypes.Num();
+	Record << SA_VALUE(TEXT("SuperTypeCount"), SuperTypeCount);
+
+	if (!ensure(SuperTypeCount >= 0))
+	{
+		SuperTypeCount = 0;
+	}
+
+	FArchive& BaseArchive = Slot.GetUnderlyingArchive();
+	if (BaseArchive.IsLoading())
+	{
+		I.SuperTypes.SetNum(SuperTypeCount);
+	}
+
+	FStructuredArchive::FStream SuperTypeStream = Record.EnterStream(TEXT("SuperTypes"));
+	for (int32 i = 0; i < SuperTypeCount; ++i)
+	{
+		FStructuredArchiveRecord NameRecord = SuperTypeStream.EnterElement().EnterRecord();
+
+		NameRecord << SA_VALUE(TEXT("SuperType"), I.SuperTypes[i]);
+	}
+}
+
+} // namespace UE::Serialization::Private
diff --git a/Engine/Source/Runtime/CoreUObject/Private/UObject/InstanceDataObjectUtils.cpp b/Engine/Source/Runtime/CoreUObject/Private/UObject/InstanceDataObjectUtils.cpp
--- a/Engine/Source/Runtime/CoreUObject/Private/UObject/InstanceDataObjectUtils.cpp
+++ b/Engine/Source/Runtime/CoreUObject/Private/UObject/InstanceDataObjectUtils.cpp
@@ -15,6 +15,7 @@
 #include "Serialization/ObjectWriter.h"
 #include "Serialization/StructuredArchive.h"
 #include "Serialization/ArchiveCountMem.h"
+#include "Serialization/ImportTypeHierarchy.h"
 #include "UObject/Class.h"
 #include "UObject/EnumProperty.h"
 #include "UObject/Field.h"
@@ -269,7 +270,9 @@ namespace UE
 	static const FName NAME_IsInstanceDataObjectStruct(ANSITEXTVIEW("IsInstanceDataObjectClass"));
 	static const FName NAME_ContainsLoosePropertiesMetadata(ANSITEXTVIEW("ContainsLooseProperties"));
 	static const FName NAME_VerseClass(ANSITEXTVIEW("VerseClass"));
+	static const FName NAME_BlueprintGeneratedClass(ANSITEXTVIEW("BlueprintGeneratedClass"));
 	static const FName NAME_VerseDevice(ANSITEXTVIEW("VerseDevice_C"));
+	static const FName NAME_Entity(ANSITEXTVIEW("entity"));
 	static const FName NAME_IDOMapKey(ANSITEXTVIEW("Key"));
 	static const FName NAME_IDOMapValue(ANSITEXTVIEW("Value"));
 
@@ -672,7 +675,6 @@ namespace UE
 
 	static const UObject* GetBlueprintGeneratedObject(const UObject* InObject)
 	{
-		static const FName NAME_BlueprintGeneratedClass(TEXT("BlueprintGeneratedClass"));
 		const UObject* Current = InObject;
 		while (Current && !Current->IsA<UPackage>())
 		{
@@ -772,10 +774,62 @@ namespace UE
 		}
 	}
 
-	bool CanCreatePropertyBagPlaceholderTypeForImportClass(const UClass* ImportClass)
+	bool CanCreatePropertyBagPlaceholderTypeForImportClass(const UClass* ClassKind, const UE::Serialization::Private::FImportTypeHierarchy* ImportTypeHierarchy)
 	{
-		// @todo - Expand to other import types (e.g. prefab BPs) later; for now restricted to Verse class objects only.
-		return ImportClass && ImportClass->GetFName() == NAME_VerseClass;
+		bool bResult = false;
+
+		if (ClassKind)
+		{
+			// @todo - Expand to other import types (e.g. non-prefab BPs) later; for now restricted to Verse class objects and entity prefabs.
+
+			if (ClassKind->GetFName() == NAME_VerseClass)
+			{
+				bResult = true;
+			}
+			else if (ImportTypeHierarchy && (ClassKind->GetFName() == NAME_BlueprintGeneratedClass))
+			{
+				const TArray<UE::Serialization::Private::FTypeResource>& SuperTypes = ImportTypeHierarchy->GetSuperTypes();
+
+				for (const UE::Serialization::Private::FTypeResource& SuperType : SuperTypes)
+				{
+					if (SuperType.TypeName == NAME_Entity)
+					{
+						bResult = true;
+						break;
+					}
+				}
+			}
+		}
+
+		return bResult;
+	}
+
+	bool CanCreatePropertyBagPlaceholderForType(const UClass* ClassKind, const UStruct* Type)
+	{
+		bool bResult = false;
+
+		if (CanCreatePropertyBagPlaceholderTypeForImportClass(ClassKind, /*ImportTypeHierarchy*/nullptr))
+		{
+			bResult = true;
+		}
+		else
+		{
+			// See comment in CanCreatePropertyBagPlaceholderTypeForImportClass: for now we only
+			// support BlueprintGeneratedClasses that are entity prefabs.
+			if (Type && ClassKind && (ClassKind->GetFName() == NAME_BlueprintGeneratedClass))
+			{
+				for (const UStruct* Struct : Type->GetSuperStructIterator())
+				{
+					if (Struct->GetFName() == NAME_Entity)
+					{
+						bResult = true;
+						break;
+					}
+				}
+			}
+		}
+
+		return bResult;
 	}
 
 	bool IsClassOfInstanceDataObjectClass(const UStruct* Class)
diff --git a/Engine/Source/Runtime/CoreUObject/Private/UObject/Linker.cpp b/Engine/Source/Runtime/CoreUObject/Private/UObject/Linker.cpp
--- a/Engine/Source/Runtime/CoreUObject/Private/UObject/Linker.cpp
+++ b/Engine/Source/Runtime/CoreUObject/Private/UObject/Linker.cpp
@@ -200,6 +200,7 @@ void FLinker::Serialize( FArchive& Ar )
 		Ar << SoftPackageReferenceList;
 		Ar << GatherableTextDataMap;
 		Ar << SearchableNamesMap;
+		Ar << ImportTypeHierarchies;
 	}
 }
 
diff --git a/Engine/Source/Runtime/CoreUObject/Private/UObject/LinkerLoad.cpp b/Engine/Source/Runtime/CoreUObject/Private/UObject/LinkerLoad.cpp
--- a/Engine/Source/Runtime/CoreUObject/Private/UObject/LinkerLoad.cpp
+++ b/Engine/Source/Runtime/CoreUObject/Private/UObject/LinkerLoad.cpp
@@ -60,6 +60,7 @@
 #include "Serialization/UnversionedPropertySerialization.h"
 #include "Serialization/LoadTimeTracePrivate.h"
 #include "Serialization/EditorBulkData.h"
+#include "Serialization/ImportTypeHierarchy.h"
 #include "HAL/FileManager.h"
 #include "UObject/CoreRedirects.h"
 #include "UObject/ICookInfo.h"
@@ -1033,6 +1034,14 @@ FLinkerLoad::ELinkerStatus FLinkerLoad::ProcessPackageSummary(TMap<TPair<FName,
 		Status = SerializeDataResourceMap();
 	}
 
+#if WITH_EDITORONLY_DATA
+	if (Status == LINKER_Loaded)
+	{
+		SCOPED_LOADTIMER(LinkerLoad_SerializeImportTypeHierarchies);
+		Status = SerializeImportTypeHierarchies();
+	}
+#endif
+
 	TRACE_LOADTIME_END_PROCESS_SUMMARY;
 
 	// Finalize creation process.
@@ -1168,6 +1177,7 @@ FLinkerLoad::FLinkerLoad(UPackage* InParent, const FPackagePath& InPackagePath,
 , bUseTimeLimit(false)
 , bUseFullTimeLimit(false)
 , bLoaderNeedsEngineVersionChecks(true)
+, bHasSerializedImportTypeHierarchies(false)
 #if WITH_EDITOR
 , bExportsDuplicatesFixed(false)
 , bIsPackageRelocated(false)
@@ -1726,6 +1736,8 @@ FLinkerLoad::ELinkerStatus FLinkerLoad::SerializePackageFileSummary()
 		NameMap.Empty(Summary.NameCount);
 		// Depends map gets pre-sized in SerializeDependsMap if used.
 
+		ImportTypeHierarchies.Empty(Summary.ImportTypeHierarchiesCount);
+
 		// Avoid serializing it again.
 		bHasSerializedPackageFileSummary = true;
 	}
@@ -2705,6 +2717,34 @@ FLinkerLoad::ELinkerStatus FLinkerLoad::SerializeDataResourceMap()
 	return LINKER_Loaded;
 }
 
+#if WITH_EDITORONLY_DATA
+
+FLinkerLoad::ELinkerStatus FLinkerLoad::SerializeImportTypeHierarchies()
+{
+	DECLARE_SCOPE_CYCLE_COUNTER(TEXT("FLinkerLoad::SerializeImportTypeHierarchies"), STAT_LinkerLoad_SerializeImportTypeHierarchies, STATGROUP_LinkerLoad);
+
+	// Skip serializing import types if the package does not include any
+	if (bHasSerializedImportTypeHierarchies || (Summary.ImportTypeHierarchiesCount == 0) || (Summary.ImportTypeHierarchiesOffset == 0))
+	{
+		return LINKER_Loaded;
+	}
+
+	if (!IsTextFormat())
+	{
+		Seek(Summary.ImportTypeHierarchiesOffset);
+	}
+
+	FStructuredArchive::FStream Stream = StructuredArchiveRootRecord->EnterStream(TEXT("ImportTypeHierarchies"));
+	Stream << ImportTypeHierarchies;
+
+	bHasSerializedImportTypeHierarchies = true;
+
+	// Return whether we finished this step and it's safe to start with the next.
+	return !IsTimeLimitExceeded(TEXT("serializing import type hierarchies")) ? LINKER_Loaded : LINKER_TimedOut;
+}
+
+#endif // WITH_EDITORONLY_DATA
+
 /**
  * Serializes thumbnails
  */
@@ -4326,7 +4366,10 @@ UClass* FLinkerLoad::TryCreatePlaceholderClassImport(int32 ImportIndex)
 	{
 		if (UClass* ImportClass = FindObjectFast<UClass>(ImportClassPackage, Import.ClassName, EFindObjectFlags::None))
 		{
-			if (UE::CanCreatePropertyBagPlaceholderTypeForImportClass(ImportClass))
+			FPackageIndex ImportPackageIndex = FPackageIndex::FromImport(ImportIndex);
+			const UE::Serialization::Private::FImportTypeHierarchy* ImportTypeHierarchy = ImportTypeHierarchies.Find(ImportPackageIndex);
+
+			if (UE::CanCreatePropertyBagPlaceholderTypeForImportClass(ImportClass, ImportTypeHierarchy))
 			{
 				// If the outer package import is also missing, create it now so that the full path remains the same. 
 				UPackage* ClassObjectPackage = Cast<UPackage>(IndexToObject(Import.OuterIndex));
@@ -5196,6 +5239,22 @@ UObject* FLinkerLoad::CreateExport( int32 Index )
 		if( !Export.SuperIndex.IsNull() )
 		{
 			UStruct* SuperStruct = (UStruct*)IndexToObject( Export.SuperIndex );
+
+#if WITH_EDITOR
+			if (!SuperStruct)
+			{
+				// Try creating a placeholder type for the super struct. 
+				if (Export.SuperIndex.IsImport())
+				{
+					SuperStruct = TryCreatePlaceholderClassImport(Export.SuperIndex.ToImport());
+				}
+				else
+				{
+					SuperStruct = TryCreatePlaceholderClassForExport(Export.SuperIndex.ToExport());
+				}
+			}
+#endif
+
 			if( !SuperStruct )
 			{
 				if( LoadClass->IsChildOf(UFunction::StaticClass()) )
diff --git a/Engine/Source/Runtime/CoreUObject/Private/UObject/SavePackage/PackageHarvester.cpp b/Engine/Source/Runtime/CoreUObject/Private/UObject/SavePackage/PackageHarvester.cpp
--- a/Engine/Source/Runtime/CoreUObject/Private/UObject/SavePackage/PackageHarvester.cpp
+++ b/Engine/Source/Runtime/CoreUObject/Private/UObject/SavePackage/PackageHarvester.cpp
@@ -18,6 +18,11 @@
 #include "VerseVM/VVMPackage.h"
 #endif
 
+#if WITH_EDITORONLY_DATA
+#include "UObject/Class.h"
+#include "UObject/InstanceDataObjectUtils.h"
+#endif
+
 EObjectMark GenerateMarksForObject(const UObject* InObject, FSaveContext& SaveContext)
 {
 	using namespace UE::SavePackageUtilities;
@@ -946,6 +951,13 @@ void FPackageHarvester::ProcessImport(TObjectPtr<UObject> InObjPtr,
 		HarvestPackageHeaderName(ObjClass->GetFName());
 		HarvestPackageHeaderName(ObjClass->GetOuter()->GetFName());
 	}
+
+#if WITH_EDITORONLY_DATA
+	if (!SaveContext.IsCooking())
+	{
+		HarvestImportTypeHierarchyNames(InObject);
+	}
+#endif
 }
 
 FString FPackageHarvester::GetArchiveName() const
@@ -1387,6 +1399,45 @@ void FPackageHarvester::HarvestCellImport(Verse::VCell* InCell)
 }
 #endif
 
+#if WITH_EDITORONLY_DATA
+// Harvests the names of an import Struct object's super type hierarchy (used in: FImportTypeHierarchy).
+void FPackageHarvester::HarvestImportTypeHierarchyNames(const UObject* Import)
+{
+	const UStruct* Struct = Cast<UStruct>(Import);
+	if (Struct)
+	{
+		// Only harvest names for types that can be used to create placeholders with.
+		// See comment in SavePackage2.cpp.
+		if (UE::CanCreatePropertyBagPlaceholderForType(Struct->GetClass(), Struct))
+		{
+			const UStruct* SuperStruct = Struct->GetSuperStruct();
+
+			while (SuperStruct && (SuperStruct != UObject::StaticClass()))
+			{
+				const UClass* SuperStructClass = SuperStruct->GetClass();
+				const UPackage* SuperStructPackage = SuperStruct->GetPackage();
+
+				FName TypeName = SuperStruct->GetFName();
+				FName PackageName = SuperStructPackage ? SuperStructPackage->GetFName() : NAME_None;
+
+				// This is the Class of the Struct, possible values are "Class", "VerseClass", "BlueprintGeneratedClass", "UserDefinedStruct", etc.
+				// We need to know the class of the Struct so we can create a placeholder type for it, when it is not available anymore (e.g., the type
+				// was removed).
+				FName ClassName = SuperStructClass ? SuperStructClass->GetFName() : NAME_None;
+				FName ClassPackageName = (SuperStructClass && SuperStructClass->GetPackage()) ? SuperStructClass->GetPackage()->GetFName() : NAME_None;
+
+				HarvestPackageHeaderName(TypeName);
+				HarvestPackageHeaderName(PackageName);
+				HarvestPackageHeaderName(ClassName);
+				HarvestPackageHeaderName(ClassPackageName);
+
+				SuperStruct = SuperStruct->GetSuperStruct();
+			}
+		}
+	}
+}
+#endif
+
 void FPackageHarvester::AppendCurrentExportDependencies()
 {
 	check(CurrentExportDependencies.CurrentExport || CurrentExportDependencies.CurrentCellExport);
diff --git a/Engine/Source/Runtime/CoreUObject/Private/UObject/SavePackage/PackageHarvester.h b/Engine/Source/Runtime/CoreUObject/Private/UObject/SavePackage/PackageHarvester.h
--- a/Engine/Source/Runtime/CoreUObject/Private/UObject/SavePackage/PackageHarvester.h
+++ b/Engine/Source/Runtime/CoreUObject/Private/UObject/SavePackage/PackageHarvester.h
@@ -164,6 +164,10 @@ private:
 	void HarvestCellImport(Verse::VCell* InCell);
 #endif
 
+#if WITH_EDITORONLY_DATA
+	void HarvestImportTypeHierarchyNames(const UObject* Import);
+#endif
+
 	void AppendCurrentExportDependencies();
 	FString GetUnsaveableReason(UObject* Required, ESaveRealm RealmInWhichItIsUnsaveable);
 	ESaveableStatus GetSaveableStatusForRealm(UObject* Obj, ESaveRealm RealmInWhichItIsUnsaveable,
diff --git a/Engine/Source/Runtime/CoreUObject/Private/UObject/SavePackage2.cpp b/Engine/Source/Runtime/CoreUObject/Private/UObject/SavePackage2.cpp
--- a/Engine/Source/Runtime/CoreUObject/Private/UObject/SavePackage2.cpp
+++ b/Engine/Source/Runtime/CoreUObject/Private/UObject/SavePackage2.cpp
@@ -35,6 +35,7 @@
 #include "Serialization/PackageWriter.h"
 #include "Serialization/PropertyLocalizationDataGathering.h"
 #include "Serialization/UnversionedPropertySerialization.h"
+#include "Serialization/ImportTypeHierarchy.h"
 #include "UObject/DebugSerializationFlags.h"
 #include "UObject/EditorObjectVersion.h"
 #include "UObject/InstanceDataObjectUtils.h"
@@ -1657,6 +1658,63 @@ ESavePackageResult BuildLinker(FSaveContext& SaveContext)
 			}
 		}
 
+#if WITH_EDITORONLY_DATA
+		// Build ImportTypeHierarchies
+		if (!SaveContext.IsCooking())
+		{
+			SCOPED_SAVETIMER(UPackage_Save_BuildImportTypeHierarchies);
+
+			for (FObjectExport& Export : Linker->ExportMap)
+			{
+				if (!Export.ClassIndex.IsNull() && Export.ClassIndex.IsImport())
+				{
+					FObjectImport& ClassImport = Linker->ImportMap[Export.ClassIndex.ToImport()];
+
+					if (ClassImport.XObject)
+					{
+						UStruct* Struct = Cast<UStruct>(ClassImport.XObject);
+						if (Struct)
+						{
+							const FPackageIndex PackageIndex = Linker->MapObject(ClassImport.XObject);
+							// This should always be in the import table already.
+							if (ensureMsgf(!PackageIndex.IsNull(), TEXT("%s: Struct object: '%s' is not mapped"), *Linker->LinkerRoot->GetFullName(), *Struct->GetFullName()))
+							{
+								// Skip already processed class imports.
+								if (!Linker->ImportTypeHierarchies.Contains(PackageIndex))
+								{
+									// To avoid bloating the ImportTypeHierarchy table only create entries for types 
+									// that support PropertyBag placeholders.
+									if (UE::CanCreatePropertyBagPlaceholderForType(Struct->GetClass(), Struct))
+									{
+										UE::Serialization::Private::FImportTypeHierarchy& ImportTypeHierarchy = Linker->ImportTypeHierarchies.Add(PackageIndex);
+
+										UStruct* SuperStruct = Struct->GetSuperStruct();
+
+										while (SuperStruct && (SuperStruct != UObject::StaticClass()))
+										{
+											UClass* SuperStructClass = SuperStruct->GetClass();
+											UPackage* SuperStructPackage = SuperStruct->GetPackage();
+
+											UE::Serialization::Private::FTypeResource SuperType;
+											SuperType.TypeName = SuperStruct->GetFName();
+											SuperType.PackageName = SuperStructPackage ? SuperStructPackage->GetFName() : NAME_None;
+											SuperType.ClassName = SuperStructClass ? SuperStructClass->GetFName() : NAME_None;
+											SuperType.ClassPackageName = (SuperStructClass && SuperStructClass->GetPackage()) ? SuperStructClass->GetPackage()->GetFName() : NAME_None;
+
+											ImportTypeHierarchy.AddSuperType(MoveTemp(SuperType));
+
+											SuperStruct = SuperStruct->GetSuperStruct();
+										}
+									}
+								}
+							}
+						}
+					}
+				}
+			}
+		}
+#endif
+
 #if WITH_VERSE_VM || defined(__INTELLISENSE__)
 		for (FCellImport& CellImport : Linker->CellImportMap)
 		{
@@ -2433,6 +2491,27 @@ ESavePackageResult WritePackageHeader(FStructuredArchive::FRecord& StructuredArc
 		Linker->Summary.SearchableNamesOffset = 0;
 	}
 
+	// Write ImportTypeHierarchies
+#if WITH_EDITORONLY_DATA
+	Linker->Summary.ImportTypeHierarchiesCount = 0;
+	Linker->Summary.ImportTypeHierarchiesOffset = 0;
+
+	if (!SaveContext.IsCooking())
+	{
+		SCOPED_SAVETIMER(UPackage_Save_SaveImportTypeHierarchies);
+		
+		Linker->Summary.ImportTypeHierarchiesOffset = (int32)Linker->Tell();
+		Linker->Summary.ImportTypeHierarchiesCount = Linker->ImportTypeHierarchies.Num();
+
+		// Sort before saving to keep the entry order consistent.
+		Linker->ImportTypeHierarchies.KeySort(TLess<FPackageIndex>());
+
+		FStructuredArchive::FStream ImportTypeHierarchyStream = StructuredArchiveRoot.EnterStream(TEXT("ImportTypeHierarchies"));
+
+		ImportTypeHierarchyStream << Linker->ImportTypeHierarchies;
+	}
+#endif
+
 	// Save thumbnails
 	{
 		SCOPED_SAVETIMER(UPackage_Save_SaveThumbnails);
diff --git a/Engine/Source/Runtime/CoreUObject/Public/Serialization/ImportTypeHierarchy.h b/Engine/Source/Runtime/CoreUObject/Public/Serialization/ImportTypeHierarchy.h
new file mode 100644
--- /dev/null
+++ b/Engine/Source/Runtime/CoreUObject/Public/Serialization/ImportTypeHierarchy.h
@@ -0,0 +1,53 @@
+// Copyright Epic Games, Inc. All Rights Reserved.
+
+#pragma once
+
+#include "Containers/Array.h"
+#include "UObject/NameTypes.h"
+#include "Serialization/StructuredArchive.h"
+#include "Serialization/StructuredArchiveSlots.h"
+
+/*
+* FImportTypeHierarchy are meta-data entries stored as a map in package files (see FPackageFileSummary).
+* Their purpose is to store the full type hierarchy of a package Import entry that is a Struct-based object.
+* NOTE1: This is internal-engine data that should not be used by user code. The format of the data is also subject to change.
+* NOTE2: The class hierarchy stored is a snapshot of the type when the object was saved. It might not correspond to the actual 
+		type hierarchy when the object is loaded.
+* NOTE3: UObject::StaticClass is omitted from the SuperTypes array as it is implied.
+*/
+namespace UE::Serialization::Private
+{
+
+/** Stores object information for a type that is part of a FImportTypeHierarchy. */
+struct FTypeResource
+{
+	FName TypeName;
+	// Package where the type object lives.
+	FName PackageName;
+	// This is the Class 'kind' for the type (e.g., "Class", "VerseClass", "BlueprintGeneratedClass", "UserDefinedStruct", etc).
+	FName ClassName;
+	// Package where the Class 'kind' object lives.
+	FName ClassPackageName;
+
+	/** I/O functions */
+	friend COREUOBJECT_API FArchive& operator<<(FArchive& Ar, FTypeResource& R);
+	friend COREUOBJECT_API void operator<<(FStructuredArchive::FSlot Slot, FTypeResource& R);
+};
+
+/** Stores a snapshot of hierarchical type information for a Struct import entry. */
+class FImportTypeHierarchy
+{
+	TArray<FTypeResource> SuperTypes;
+
+public:
+
+	UE_INTERNAL const TArray<FTypeResource>& GetSuperTypes() const { return SuperTypes; }
+	UE_INTERNAL void AddSuperType(const FTypeResource& SuperType) { SuperTypes.Add(SuperType); }
+	UE_INTERNAL void AddSuperType(FTypeResource&& SuperType) { SuperTypes.Add(SuperType); }
+
+	/** I/O functions */
+	friend COREUOBJECT_API FArchive& operator<<(FArchive& Ar, FImportTypeHierarchy& I);
+	friend COREUOBJECT_API void operator<<(FStructuredArchive::FSlot Slot, FImportTypeHierarchy& I);
+};
+
+} // namespace UE::Serialization::Private
diff --git a/Engine/Source/Runtime/CoreUObject/Public/UObject/Linker.h b/Engine/Source/Runtime/CoreUObject/Public/UObject/Linker.h
--- a/Engine/Source/Runtime/CoreUObject/Public/UObject/Linker.h
+++ b/Engine/Source/Runtime/CoreUObject/Public/UObject/Linker.h
@@ -14,6 +14,7 @@
 #include "Misc/AssertionMacros.h"
 #include "Serialization/StructuredArchive.h"
 #include "Serialization/StructuredArchiveSlots.h"
+#include "Serialization/ImportTypeHierarchy.h"
 #include "Templates/RefCounting.h"
 #include "UObject/LinkerInstancingContext.h"
 #include "UObject/NameTypes.h"
@@ -418,7 +419,10 @@ public:
 	
 	/** Raw/bulk data references */
 	TArray<FObjectDataResource> DataResourceMap;
-	
+
+	/** Map of hierarchical type information for FObjectImport Struct entries in the package */
+	TMap<FPackageIndex, UE::Serialization::Private::FImportTypeHierarchy> ImportTypeHierarchies;
+
 	/** If true, filter out exports that are for clients but not servers */
 	bool					FilterClientButNotServer;
 
diff --git a/Engine/Source/Runtime/CoreUObject/Public/UObject/LinkerLoad.h b/Engine/Source/Runtime/CoreUObject/Public/UObject/LinkerLoad.h
--- a/Engine/Source/Runtime/CoreUObject/Public/UObject/LinkerLoad.h
+++ b/Engine/Source/Runtime/CoreUObject/Public/UObject/LinkerLoad.h
@@ -469,6 +469,8 @@ private:
 	bool					bUseFullTimeLimit:1;
 	/** Whether the loader needs version and correctness checks (see OpenReadPackage)										*/
 	bool					bLoaderNeedsEngineVersionChecks : 1;
+	/** Whether we already serialized import type hierarchies.																*/
+	bool					bHasSerializedImportTypeHierarchies : 1;
 
 #if WITH_EDITOR
 	/** Check to avoid multiple export duplicate fixups in case we don't save asset. */
@@ -1325,6 +1327,13 @@ private:
 	 */
 	ELinkerStatus SerializeDataResourceMap();
 
+	/**
+	 * Serializes the import type hierarchy map.
+	 */
+#if WITH_EDITORONLY_DATA
+	FLinkerLoad::ELinkerStatus SerializeImportTypeHierarchies();
+#endif
+
 	/** Sets the basic linker archive info */
 	void ResetStatusInfo();
 public:
