IncrementalCook: LinkerLoad: Call UE_TRACK_REFERENCING_PACKAGE_SCOPED in CreateExport, so that errors encountered during the attempt to create the export are assigned to the package that owns the export rather than the referencing package. This is necessary so that the errors are not played back if the owning package is incrementally invalidated but the referencing package is not.This fixes the error "CreateExport: Failed to load Class %s as Parent for BlueprintGeneratedClass%s..." spuriously replaying in the incremental cook even after the error was fixed and the owning package was recooked.
#rnx

diff --git a/Engine/Source/Runtime/CoreUObject/Private/UObject/LinkerLoad.cpp b/Engine/Source/Runtime/CoreUObject/Private/UObject/LinkerLoad.cpp
--- a/Engine/Source/Runtime/CoreUObject/Private/UObject/LinkerLoad.cpp
+++ b/Engine/Source/Runtime/CoreUObject/Private/UObject/LinkerLoad.cpp
@@ -5198,6 +5198,7 @@ UObject* FLinkerLoad::CreateExport( int32 Index )
 	if( !Export.Object && !FilterExport(Export) ) // for some acceptable position, it was not "not for" 
 	{
 		TGuardValue<void*> GuardThreadContextAsyncPackage(FUObjectThreadContext::Get().AsyncPackage, AsyncRoot);
+		UE_TRACK_REFERENCING_PACKAGE_SCOPED(LinkerRoot->GetFName(), PackageAccessTrackingOps::NAME_Load);
 		FUObjectSerializeContext* CurrentLoadContext = FUObjectThreadContext::Get().GetSerializeContext();
 		check(!GEventDrivenLoaderEnabled || !bLockoutLegacyOperations || !EVENT_DRIVEN_ASYNC_LOAD_ACTIVE_AT_RUNTIME);
 		check(Export.ObjectName!=NAME_None || !(Export.ObjectFlags&RF_Public));
