[UE] Disable custom property lists during native struct serialization.

diff --git a/Engine/Source/Runtime/CoreUObject/Private/UObject/Class.cpp b/Engine/Source/Runtime/CoreUObject/Private/UObject/Class.cpp
--- a/Engine/Source/Runtime/CoreUObject/Private/UObject/Class.cpp
+++ b/Engine/Source/Runtime/CoreUObject/Private/UObject/Class.cpp
@@ -3396,6 +3396,9 @@ void UScriptStruct::SerializeItem(FStructuredArchive::FSlot Slot, void* Value, v
 	bool bItemSerialized = false;
 	if (bUseNativeSerialization)
 	{
+		// Native serialization is not compatible with custom property lists.
+		FGuardValue_Bitfield(Slot.GetUnderlyingArchive().ArUseCustomPropertyList, false);
+
 		UScriptStruct::ICppStructOps* TheCppStructOps = GetCppStructOps();
 		check(TheCppStructOps); // else should not have STRUCT_SerializeNative
 
