[Misc] Fix debug map to work with sparse and compact forms

#tests Run server with compact sets enabled

diff --git a/Engine/Source/Runtime/CoreUObject/Private/Serialization/AsyncLoading2.cpp b/Engine/Source/Runtime/CoreUObject/Private/Serialization/AsyncLoading2.cpp
--- a/Engine/Source/Runtime/CoreUObject/Private/Serialization/AsyncLoading2.cpp
+++ b/Engine/Source/Runtime/CoreUObject/Private/Serialization/AsyncLoading2.cpp
@@ -6420,12 +6420,18 @@ void FCookedScriptObjectsDebug::TrimDebugData(const TMap<FPackageObjectIndex, UO
 	const int32 Num = Objects.Num();
 	TBitArray<> Remove(false, Num);
 	TBitArray<> Visited(false, Num);
-	TFunction<void(int32)> KeepOuters = [this, &KeepOuters, &Remove, &Visited](int32 I)
+
+	int32 NumRemaining = Num;
+	TFunction<void(int32)> KeepOuters = [this, &KeepOuters, &Remove, &Visited, &NumRemaining](int32 I)
 	{
 		if (!Visited[I])
 		{
 			Visited[I] = true;
-			Remove[I] = false;
+			if (Remove[I])
+			{
+				++NumRemaining;
+				Remove[I] = false;
+			}
 			const TPair<FMinimalName, FPackageObjectIndex>& Entry = Objects.Get(FSetElementId::FromInteger(I)).Value;
 			if (!Entry.Value.IsNull())
 			{
@@ -6441,6 +6447,7 @@ void FCookedScriptObjectsDebug::TrimDebugData(const TMap<FPackageObjectIndex, UO
 		if (Id.IsValidId())
 		{
 			Remove[Id.AsInteger()] = true;
+			--NumRemaining;
 		}
 	}
 	// 2. Keep the outers for all objects that are kept
@@ -6451,17 +6458,28 @@ void FCookedScriptObjectsDebug::TrimDebugData(const TMap<FPackageObjectIndex, UO
 			KeepOuters(I);
 		}
 	}
+
 	// 3. Remove the remaining runtime objects
-	for (int I = 0; I < Num; ++I)
+	if (NumRemaining != Num)
 	{
-		if (Remove[I])
+		TMap<FPackageObjectIndex, TPair<FMinimalName, FPackageObjectIndex>> RemainingObjects;
+		RemainingObjects.Reserve(NumRemaining);
+
+		for (int I = 0; I < Num; ++I)
 		{
-			Objects.Remove(FSetElementId::FromInteger(I));
+			if (!Remove[I])
+			{
+				RemainingObjects.Add(Objects.Get(FSetElementId::FromInteger(I)));
+			}
 		}
+
+		check(RemainingObjects.Num() == NumRemaining);
+
+		Swap(RemainingObjects, Objects);
 	}
-	Objects.Compact();
-	UE_LOG(LogStreaming, Log, TEXT("AsyncLoading2 - TrimDebugData: Kept %d cooked script object entries (%.2f KB)"),
-		Objects.Num(), (float)Objects.GetAllocatedSize()/1024.f);
+
+	UE_LOG(LogStreaming, Log, TEXT("AsyncLoading2 - TrimDebugData: Kept %d/%d cooked script object entries (%.2f KB)"),
+		NumRemaining, Num, (float)Objects.GetAllocatedSize()/1024.f);
 }
 
 void FCookedScriptObjectsDebug::GetPathNameInternal(const TPair<FMinimalName, FPackageObjectIndex>& Entry, FStringBuilderBase& PathName, int32& NumParts) const
