Downgrades the fatal error for checking for leaking CDOs of Blueprint-compiled classes to be an `ensure`, controlled by the new `s.EnableFatalCDOBlueprintClassesLeaks` cvar. This is to avert a current SH regarding Verse hot-reload functionality that can currently be triggered thanks to this behaviour.

#rnx

diff --git a/Engine/Source/Runtime/CoreUObject/Private/Serialization/AsyncLoading2.cpp b/Engine/Source/Runtime/CoreUObject/Private/Serialization/AsyncLoading2.cpp
--- a/Engine/Source/Runtime/CoreUObject/Private/Serialization/AsyncLoading2.cpp
+++ b/Engine/Source/Runtime/CoreUObject/Private/Serialization/AsyncLoading2.cpp
@@ -101,6 +101,14 @@
 
 #include <atomic>
 
+TAutoConsoleVariable<bool> CVar_CheckBlueprintCDOWillBePostLoadedAfterPreload(
+	TEXT("s.CheckBlueprintCDOWillBePostLoadedAfterPreload"),
+	true,
+	TEXT("Determines whether or not to check that Class Default Object (CDO)s that are compiled from Blueprints (including old VerseVM classes) are"
+		"set to be post-loaded and are in the process of being loaded."
+		"When disabled, we will still `ensure` regardless, but will not treat these as fatal errors.")
+);
+
 template<typename PayloadType> using TAsyncAtomic = std::atomic<PayloadType>;
 
 // For now, the partial request behavior is reserved for the editor only.
@@ -345,6 +353,12 @@ static FAutoConsoleVariableRef CEnableScriptImportsDebugging(
 									, (Object ? *LexToString(Object->GetFlags()) : TEXT("None")) \
 									, (Object ? Object->GetInternalFlags() : EInternalObjectFlags::None))
 
+#define ensureObject(Object, expr) ensureMsgf(expr, TEXT("Object='%s' (%p), Flags=%s, InternalFlags=0x%08X") \
+									, (Object ? *Object->GetFullName() : TEXT("null")) \
+									, Object \
+									, (Object ? *LexToString(Object->GetFlags()) : TEXT("None")) \
+									, (Object ? Object->GetInternalFlags() : EInternalObjectFlags::None))
+
 static TSet<FPackageId> GAsyncLoading2_DebugPackageIds;
 static FString GAsyncLoading2_DebugPackageNamesString;
 static TSet<FPackageId> GAsyncLoading2_VerbosePackageIds;
@@ -8900,13 +8914,17 @@ bool FAsyncPackage2::EventDrivenSerializeExport(const FAsyncPackageHeaderData& H
 	Object->SetFlags(RF_LoadCompleted);
 	LoadContext->SerializedObject = PrevSerializedObject;
 
-#if DO_CHECK
 	if (Object->HasAnyFlags(RF_ClassDefaultObject) && Object->GetClass()->HasAnyClassFlags(CLASS_CompiledFromBlueprint))
 	{
-		checkObject(Object, Object->HasAllFlags(RF_NeedPostLoad | RF_WasLoaded));
-		//Object->SetFlags(RF_NeedPostLoad | RF_WasLoaded);
+		if (CVar_CheckBlueprintCDOWillBePostLoadedAfterPreload.GetValueOnAnyThread())
+		{
+			checkObject(Object, Object->HasAllFlags(RF_NeedPostLoad | RF_WasLoaded));
+		}
+		else
+		{
+			ensureObject(Object, Object->HasAllFlags(RF_NeedPostLoad | RF_WasLoaded));
+		}
 	}
-#endif
 
 	UE_ASYNC_PACKAGE_LOG_VERBOSE(VeryVerbose, Desc, TEXT("SerializeExport"), TEXT("Serialized export %s"), *Object->GetPathName());
 
