Removed redundant preallocation in UObjectBaseUtility::GetFullName

diff --git a/Engine/Source/Runtime/CoreUObject/Private/UObject/UObjectBaseUtility.cpp b/Engine/Source/Runtime/CoreUObject/Private/UObject/UObjectBaseUtility.cpp
--- a/Engine/Source/Runtime/CoreUObject/Private/UObject/UObjectBaseUtility.cpp
+++ b/Engine/Source/Runtime/CoreUObject/Private/UObject/UObjectBaseUtility.cpp
@@ -94,7 +94,6 @@ void UObjectBaseUtility::GetPathName(const UObject* StopOuter, FStringBuilderBas
 FString UObjectBaseUtility::GetFullName(const UObject* StopOuter, EObjectFullNameFlags Flags) const
 {
 	FString Result;
-	Result.Empty(128);
 	GetFullName(StopOuter, Result, Flags);
 	return Result;
 }
