Fix syntax error in constinit uobjects build.

diff --git a/Engine/Source/Runtime/CoreUObject/Public/UObject/IntrinsicClassBaseChain.h b/Engine/Source/Runtime/CoreUObject/Public/UObject/IntrinsicClassBaseChain.h
--- a/Engine/Source/Runtime/CoreUObject/Public/UObject/IntrinsicClassBaseChain.h
+++ b/Engine/Source/Runtime/CoreUObject/Public/UObject/IntrinsicClassBaseChain.h
@@ -23,7 +23,7 @@ namespace UE::Private
 	{
 		using BaseChainPtr = const FStructBaseChain*;
 		static constexpr SIZE_T Num = sizeof...(Classes);
-		static inline constinit BaseChainPtr Bases[] = { (static_cast<BaseChainPtr>GetConstInitClass<Classes>())... };
+		static inline constinit BaseChainPtr Bases[] = { static_cast<BaseChainPtr>(GetConstInitClass<Classes>())... };
 
 		static consteval TArrayView<const FStructBaseChain*> ToArrayView()
 		{
diff --git a/Engine/Source/Runtime/CoreUObject/Public/UObject/RegisterCompiledInObjects.h b/Engine/Source/Runtime/CoreUObject/Public/UObject/RegisterCompiledInObjects.h
--- a/Engine/Source/Runtime/CoreUObject/Public/UObject/RegisterCompiledInObjects.h
+++ b/Engine/Source/Runtime/CoreUObject/Public/UObject/RegisterCompiledInObjects.h
@@ -23,14 +23,14 @@ struct FRegisterCompiledInObjects
 
 	// Register an intrinsic class with a special initialization function
 	FRegisterCompiledInObjects(UClass* InClass, FIntrinsicClassConstructor* InConstructor)
-		: FRegisterCompiledInObjects(TInPlaceType<FIntrinsicClassRegistrant>{}, InClass, InConstructor)
+		: FRegisterCompiledInObjects(FIntrinsicClassRegistrant{InClass, InConstructor})
 	{
 	}
 
 	// Register a package and some objects that might otherwise not be registered - e.g. noexport structs and delegate
 	// function signatures. ArrayView should point to static memory
 	FRegisterCompiledInObjects(UPackage* InPackage, TConstArrayView<UObject*> Others)
-		: FRegisterCompiledInObjects(TInPlaceType<FPackageRegistrants>{}, InPackage, Others)
+		: FRegisterCompiledInObjects(FPackageRegistrants{InPackage, Others})
 	{
 	}
 
@@ -41,40 +41,40 @@ struct FRegisterCompiledInObjects
 		TConstArrayView<UClass*> InClasses,
 		TConstArrayView<UScriptStruct*> InStructs,
 		TConstArrayView<UEnum*> InEnums)
-		: FRegisterCompiledInObjects(TInPlaceType<FClassRegistrants>{}, InClasses, InStructs, InEnums)
+		: FRegisterCompiledInObjects(FClassRegistrants{InClasses, InStructs, InEnums})
 	{
 	}
 	FRegisterCompiledInObjects(
 		TConstArrayView<UClass*> InClasses)
-		: FRegisterCompiledInObjects(TInPlaceType<FClassRegistrants>{}, InClasses, {}, {})
+		: FRegisterCompiledInObjects(FClassRegistrants{InClasses, {}, {}})
 	{
 	}
 	FRegisterCompiledInObjects(
 		TConstArrayView<UClass*> InClasses,
 		TConstArrayView<UScriptStruct*> InStructs)
-		: FRegisterCompiledInObjects(TInPlaceType<FClassRegistrants>{}, InClasses, InStructs, {})
+		: FRegisterCompiledInObjects(FClassRegistrants{InClasses, InStructs, {}})
 	{
 	}
 	FRegisterCompiledInObjects(
 		TConstArrayView<UClass*> InClasses,
 		TConstArrayView<UEnum*> InEnums)
-		: FRegisterCompiledInObjects(TInPlaceType<FClassRegistrants>{}, InClasses, {}, InEnums)
+		: FRegisterCompiledInObjects(FClassRegistrants{InClasses, {}, InEnums})
 	{
 	}
 	FRegisterCompiledInObjects(
 		TConstArrayView<UScriptStruct*> InStructs,
 		TConstArrayView<UEnum*> InEnums)
-		: FRegisterCompiledInObjects(TInPlaceType<FClassRegistrants>{}, {}, InStructs, InEnums)
+		: FRegisterCompiledInObjects(FClassRegistrants{{}, InStructs, InEnums})
 	{
 	}
 	FRegisterCompiledInObjects(
 		TConstArrayView<UScriptStruct*> InStructs)
-		: FRegisterCompiledInObjects(TInPlaceType<FClassRegistrants>{}, {}, InStructs, {})
+		: FRegisterCompiledInObjects(FClassRegistrants{{}, InStructs, {}})
 	{
 	}
 	FRegisterCompiledInObjects(
 		TConstArrayView<UEnum*> InEnums)
-		: FRegisterCompiledInObjects(TInPlaceType<FClassRegistrants>{}, {}, {}, InEnums)
+		: FRegisterCompiledInObjects(FClassRegistrants{{}, {}, InEnums})
 	{
 	}
 
@@ -101,9 +101,9 @@ struct FRegisterCompiledInObjects
 	FRegisterCompiledInObjects* ListNext = nullptr;
 	
 	// Delegating constructor
-	template<typename VARIANT, typename... Args>
-	FRegisterCompiledInObjects(TInPlaceType<VARIANT>, Args&&... InArgs)
-	: Registrants(TInPlaceType<VARIANT>{}, Forward<Args>(InArgs)...)
+	template<typename VARIANT>
+	FRegisterCompiledInObjects(VARIANT&& Variant)
+	: Registrants(TInPlaceType<VARIANT>{}, Forward<VARIANT>(Variant))
 	{
 		Register();
 	}
