Memory-mapped bulk data alignment change

Before this, memory-mapped bulk data is always aligned to page size *in actual iostore files*. That results in hundreds of MBs of zeros and runtime memory cost (due to zeros being also mapped, wasting internal page space).

This change sets the alignment to be 1 by default, which removes the padding. The mmapping code handles the offset being unaligned internally, so still mapping in page boundaries.

Added cook.MemoryMappedBulkDataAlignment cvar to override the used alignment, which is set to 1 by default.
Old behaviour can be achieved via cook.MemoryMappedBulkDataAlignment=0.

Pak files alignment is also changed to be 1, so no padding is added there, in similar fashion.

[RN]

diff --git a/Engine/Config/Android/BaseAndroidEngine.ini b/Engine/Config/Android/BaseAndroidEngine.ini
--- a/Engine/Config/Android/BaseAndroidEngine.ini
+++ b/Engine/Config/Android/BaseAndroidEngine.ini
@@ -29,4 +29,8 @@ r.VolumetricCloud.Support=0
 [/Script/AndroidRuntimeSettings.AndroidRuntimeSettings]
 ; Default to latest astc version - projects can manually specify this to retain older versions to prevent
 ; patches.
-ASTCVersion=501
\ No newline at end of file
+ASTCVersion=501
+
+[MemoryMappedFiles]
+Enable=true
+Alignment=1
diff --git a/Engine/Config/IOS/BaseIOSEngine.ini b/Engine/Config/IOS/BaseIOSEngine.ini
--- a/Engine/Config/IOS/BaseIOSEngine.ini
+++ b/Engine/Config/IOS/BaseIOSEngine.ini
@@ -84,7 +84,7 @@ DefaultPlatformService=IOSLocalNotification
 
 [MemoryMappedFiles]
 Enable=true
-Alignment=16384
+Alignment=1
 
 [OnlineSubsystem]
 DefaultPlatformService=IOS
diff --git a/Engine/Source/Runtime/CoreUObject/Private/UObject/SavePackage2.cpp b/Engine/Source/Runtime/CoreUObject/Private/UObject/SavePackage2.cpp
--- a/Engine/Source/Runtime/CoreUObject/Private/UObject/SavePackage2.cpp
+++ b/Engine/Source/Runtime/CoreUObject/Private/UObject/SavePackage2.cpp
@@ -81,6 +81,15 @@ static FAutoConsoleVariableRef CVar_FixupStandaloneFlags(
 	TEXT("If non-zero, when the UAsset of a package is missing RF_Standalone, the flag is added. If zero, the flags are not changed and the save fails.")
 );
 
+int32 GMemoryMappedBulkDataAlignment = 1;
+static FAutoConsoleVariableRef CVar_MemoryMappedBulkDataAlignment(
+	TEXT("cook.MemoryMappedBulkDataAlignment"),
+	GMemoryMappedBulkDataAlignment,
+	TEXT("0 - Target platform memory mapping alignment is respected, might result in big zero-padding")
+	TEXT("1 - No special alignment required, 1-byte alignment used (no padding, default)")															  
+	TEXT("Any other non-zero value - overrides the memory mapped data alignment, implies padding")
+);
+
 ESavePackageResult WriteAdditionalFiles(FSaveContext& SaveContext, FScopedSlowTask& SlowTask, int64& VirtualExportsFileOffset);
 
 ESavePackageResult ReturnSuccessOrCancel()
@@ -1250,7 +1259,8 @@ ESavePackageResult BuildLinker(FSaveContext& SaveContext)
 			{
 				if (TargetPlatform->SupportsFeature(ETargetPlatformFeatures::MemoryMappedFiles))
 				{
-					Linker->SetMemoryMapAlignment(TargetPlatform->GetMemoryMappingAlignment());
+					int64 MemoryMapAlignment = GMemoryMappedBulkDataAlignment > 0 ? GMemoryMappedBulkDataAlignment : TargetPlatform->GetMemoryMappingAlignment();
+					Linker->SetMemoryMapAlignment(MemoryMapAlignment);
 				}
 				Linker->SetFileRegionsEnabled(TargetPlatform->SupportsFeature(ETargetPlatformFeatures::CookFileRegionMetadata));
 			}
diff --git a/Engine/Source/Runtime/PakFile/Private/IoDispatcherFileBackend.cpp b/Engine/Source/Runtime/PakFile/Private/IoDispatcherFileBackend.cpp
--- a/Engine/Source/Runtime/PakFile/Private/IoDispatcherFileBackend.cpp
+++ b/Engine/Source/Runtime/PakFile/Private/IoDispatcherFileBackend.cpp
@@ -1989,7 +1989,11 @@ TIoStatusOr<FIoMappedRegion> FFileIoStore::OpenMapped(const FIoChunkId& ChunkId,
 			IMappedFileRegion* MappedFileRegion = MappedFileHandle->MapRegion(BlockOffset + Options.GetOffset(), ResolvedSize);
 			if (MappedFileRegion != nullptr)
 			{
-				check(IsAligned(MappedFileRegion->GetMappedPtr(), FPlatformMemory::GetStats().PageSize));
+				// This is not required as mapped file region implementation ensure we map an aligned
+				// pointer/size. MappedPtr could be inside of the page, depending on the offset. This
+				// is actually fine, unless client code cares about MappedPtr alignment, which is
+				// public API.
+				// check(IsAligned(MappedFileRegion->GetMappedPtr(), FPlatformMemory::GetStats().PageSize));
 				return FIoMappedRegion{ MappedFileHandle, MappedFileRegion };
 			}
 			else
