AssetRegistry: Fix progation dependencies not being queried correctly: uninitialized variable when constructing FDependencyQuery from EDependencyQuery::Propagation.
#rnx

diff --git a/Engine/Source/Runtime/CoreUObject/Public/Misc/AssetRegistryInterface.h b/Engine/Source/Runtime/CoreUObject/Public/Misc/AssetRegistryInterface.h
--- a/Engine/Source/Runtime/CoreUObject/Public/Misc/AssetRegistryInterface.h
+++ b/Engine/Source/Runtime/CoreUObject/Public/Misc/AssetRegistryInterface.h
@@ -288,6 +288,7 @@ namespace UE::AssetRegistry
 			}
 			else
 			{
+				Required = UE::AssetRegistry::EDependencyProperty::None;
 				EnumRemoveFlags(QueryFlags, EDependencyQuery::Game | EDependencyQuery::NotGame | EDependencyQuery::Build | EDependencyQuery::NotBuild);
 				RequiredUnions.Add(UE::AssetRegistry::EDependencyProperty::Game | UE::AssetRegistry::EDependencyProperty::Build);
 			}
