Simplify critical-error checks in deserialization code.

`State` already points to the archive state of the slot's underlying archive.

diff --git a/Engine/Source/Runtime/CoreUObject/Private/UObject/PropertyTypeName.cpp b/Engine/Source/Runtime/CoreUObject/Private/UObject/PropertyTypeName.cpp
--- a/Engine/Source/Runtime/CoreUObject/Private/UObject/PropertyTypeName.cpp
+++ b/Engine/Source/Runtime/CoreUObject/Private/UObject/PropertyTypeName.cpp
@@ -409,7 +409,7 @@ void operator<<(FStructuredArchiveSlot Slot, FPropertyTypeName& TypeName)
 				Stream.EnterElement() << Node;
 				Remaining += Node.InnerCount - 1;
 
-				if (Slot.GetUnderlyingArchive().IsCriticalError())
+				if (State.IsCriticalError())
 				{
 					break;
 				}
@@ -435,7 +435,7 @@ void operator<<(FStructuredArchiveSlot Slot, FPropertyTypeName& TypeName)
 				Stream.EnterElement() << NodeCopy;
 				Remaining += NodeCopy.InnerCount;
 
-				if (Slot.GetUnderlyingArchive().IsCriticalError())
+				if (State.IsCriticalError())
 				{
 					break;
 				}
