Fix when clearing a SubObject twice, clears all properties of its owner.

diff --git a/Engine/Source/Runtime/CoreUObject/Private/UObject/OverriddenPropertySet.cpp b/Engine/Source/Runtime/CoreUObject/Private/UObject/OverriddenPropertySet.cpp
--- a/Engine/Source/Runtime/CoreUObject/Private/UObject/OverriddenPropertySet.cpp
+++ b/Engine/Source/Runtime/CoreUObject/Private/UObject/OverriddenPropertySet.cpp
@@ -785,7 +785,7 @@ bool FOverriddenPropertySet::ClearOverriddenProperty(FOverriddenPropertyNode& Pa
 			{
 				OverridableManager.ClearOverrides(SubObject);
 			}
-			bClearedOverrides = true;
+			bClearedOverrides = CurrentOverriddenPropertyNode != nullptr;
 		}
 		// Special handling for array of instanced subobjects 
 		else if (const FArrayProperty* ArrayProperty = CastField<FArrayProperty>(CurrentProperty))
