[Core] Revert change from compact to shrink as it was both incorrect and causing carshes down the line in systems that were indexing directly into the set

#tests preflight, cis fix

diff --git a/Engine/Source/Runtime/CoreUObject/Private/UObject/UObjectHash.cpp b/Engine/Source/Runtime/CoreUObject/Private/UObject/UObjectHash.cpp
--- a/Engine/Source/Runtime/CoreUObject/Private/UObject/UObjectHash.cpp
+++ b/Engine/Source/Runtime/CoreUObject/Private/UObject/UObjectHash.cpp
@@ -227,7 +227,7 @@ struct FSetHashBucket
 		TSet<UObjectBase*>* Items = GetItems();
 		if (Items)
 		{
-			Items->Shrink();
+			Items->Compact();
 		}
 	}
 
