Additional error handling when serializing import type hierarchy entries.
Added a CVar so that serialization for these entries can be disabled.

#rnx

diff --git a/Engine/Source/Runtime/CoreUObject/Private/Serialization/ImportTypeHierarchy.cpp b/Engine/Source/Runtime/CoreUObject/Private/Serialization/ImportTypeHierarchy.cpp
--- a/Engine/Source/Runtime/CoreUObject/Private/Serialization/ImportTypeHierarchy.cpp
+++ b/Engine/Source/Runtime/CoreUObject/Private/Serialization/ImportTypeHierarchy.cpp
@@ -1,10 +1,25 @@
 // Copyright Epic Games, Inc. All Rights Reserved.
 
 #include "Serialization/ImportTypeHierarchy.h"
+#include "HAL/IConsoleManager.h"
 
 namespace UE::Serialization::Private
 {
 
+static int32 GImportTypeHierarchyEnabled = 1;
+static FAutoConsoleVariableRef CVarImportTypeHierarchyEnabled(
+	TEXT("s.ImportTypeHierarchyEnabled"),
+	GImportTypeHierarchyEnabled,
+	TEXT("Controls whether Import Type Hierarchy entries will be loaded/saved."),
+	ECVF_Default
+);
+
+bool FImportTypeHierarchy::IsSerializationEnabled()
+{
+	bool Result = (GImportTypeHierarchyEnabled > 0);
+	return Result;
+}
+
 FArchive& operator<<(FArchive& Ar, FTypeResource& TypeResource)
 {
 	FStructuredArchiveFromArchive(Ar).GetSlot() << TypeResource;
@@ -29,6 +44,11 @@ FArchive& operator<<(FArchive& Ar, FImportTypeHierarchy& I)
 
 void operator<<(FStructuredArchive::FSlot Slot, FImportTypeHierarchy& I)
 {
+	if (!FImportTypeHierarchy::IsSerializationEnabled())
+	{
+		return;
+	}
+
 	FStructuredArchive::FRecord Record = Slot.EnterRecord();
 
 	int32 SuperTypeCount = I.SuperTypes.Num();
diff --git a/Engine/Source/Runtime/CoreUObject/Private/UObject/LinkerLoad.cpp b/Engine/Source/Runtime/CoreUObject/Private/UObject/LinkerLoad.cpp
--- a/Engine/Source/Runtime/CoreUObject/Private/UObject/LinkerLoad.cpp
+++ b/Engine/Source/Runtime/CoreUObject/Private/UObject/LinkerLoad.cpp
@@ -2723,22 +2723,22 @@ FLinkerLoad::ELinkerStatus FLinkerLoad::SerializeImportTypeHierarchies()
 {
 	DECLARE_SCOPE_CYCLE_COUNTER(TEXT("FLinkerLoad::SerializeImportTypeHierarchies"), STAT_LinkerLoad_SerializeImportTypeHierarchies, STATGROUP_LinkerLoad);
 
-	// Skip serializing import types if the package does not include any
-	if (bHasSerializedImportTypeHierarchies || (Summary.ImportTypeHierarchiesCount == 0) || (Summary.ImportTypeHierarchiesOffset == 0))
+	if (!UE::Serialization::Private::FImportTypeHierarchy::IsSerializationEnabled())
 	{
 		return LINKER_Loaded;
 	}
 
-	if (!IsTextFormat())
+	// Skip serializing import types if the package does not include any or when dealing with a text format.
+	if (bHasSerializedImportTypeHierarchies || (Summary.ImportTypeHierarchiesCount == 0) || (Summary.ImportTypeHierarchiesOffset == 0) || IsTextFormat())
 	{
-		Seek(Summary.ImportTypeHierarchiesOffset);
+		return LINKER_Loaded;
 	}
 
+	Seek(Summary.ImportTypeHierarchiesOffset);
+
 	FStructuredArchive::FStream Stream = StructuredArchiveRootRecord->EnterStream(TEXT("ImportTypeHierarchies"));
 	Stream << ImportTypeHierarchies;
 
-	bHasSerializedImportTypeHierarchies = true;
-
 	// Return whether we finished this step and it's safe to start with the next.
 	if (Stream.GetUnderlyingArchive().IsCriticalError())
 	{
@@ -2748,6 +2748,9 @@ FLinkerLoad::ELinkerStatus FLinkerLoad::SerializeImportTypeHierarchies()
 	{
 		return LINKER_TimedOut;
 	}
+
+	bHasSerializedImportTypeHierarchies = true;
+
 	return LINKER_Loaded;
 }
 
diff --git a/Engine/Source/Runtime/CoreUObject/Private/UObject/SavePackage2.cpp b/Engine/Source/Runtime/CoreUObject/Private/UObject/SavePackage2.cpp
--- a/Engine/Source/Runtime/CoreUObject/Private/UObject/SavePackage2.cpp
+++ b/Engine/Source/Runtime/CoreUObject/Private/UObject/SavePackage2.cpp
@@ -1670,7 +1670,7 @@ ESavePackageResult BuildLinker(FSaveContext& SaveContext)
 
 #if WITH_EDITORONLY_DATA
 		// Build ImportTypeHierarchies
-		if (!SaveContext.IsCooking())
+		if (!SaveContext.IsCooking() && UE::Serialization::Private::FImportTypeHierarchy::IsSerializationEnabled())
 		{
 			SCOPED_SAVETIMER(UPackage_Save_BuildImportTypeHierarchies);
 
@@ -2506,7 +2506,7 @@ ESavePackageResult WritePackageHeader(FStructuredArchive::FRecord& StructuredArc
 	Linker->Summary.ImportTypeHierarchiesCount = 0;
 	Linker->Summary.ImportTypeHierarchiesOffset = 0;
 
-	if (!SaveContext.IsCooking())
+	if (!SaveContext.IsCooking() && UE::Serialization::Private::FImportTypeHierarchy::IsSerializationEnabled())
 	{
 		SCOPED_SAVETIMER(UPackage_Save_SaveImportTypeHierarchies);
 		
diff --git a/Engine/Source/Runtime/CoreUObject/Public/Serialization/ImportTypeHierarchy.h b/Engine/Source/Runtime/CoreUObject/Public/Serialization/ImportTypeHierarchy.h
--- a/Engine/Source/Runtime/CoreUObject/Public/Serialization/ImportTypeHierarchy.h
+++ b/Engine/Source/Runtime/CoreUObject/Public/Serialization/ImportTypeHierarchy.h
@@ -45,6 +45,8 @@ public:
 	UE_INTERNAL void AddSuperType(const FTypeResource& SuperType) { SuperTypes.Add(SuperType); }
 	UE_INTERNAL void AddSuperType(FTypeResource&& SuperType) { SuperTypes.Add(SuperType); }
 
+	UE_INTERNAL static bool IsSerializationEnabled();
+
 	/** I/O functions */
 	friend COREUOBJECT_API FArchive& operator<<(FArchive& Ar, FImportTypeHierarchy& I);
 	friend COREUOBJECT_API void operator<<(FStructuredArchive::FSlot Slot, FImportTypeHierarchy& I);
