[UE] Disable import type hierarchy serialization to work around some cook time issues.

#rnx
#FailedValidations Preflight

diff --git a/Engine/Source/Runtime/CoreUObject/Private/Serialization/ImportTypeHierarchy.cpp b/Engine/Source/Runtime/CoreUObject/Private/Serialization/ImportTypeHierarchy.cpp
--- a/Engine/Source/Runtime/CoreUObject/Private/Serialization/ImportTypeHierarchy.cpp
+++ b/Engine/Source/Runtime/CoreUObject/Private/Serialization/ImportTypeHierarchy.cpp
@@ -6,7 +6,7 @@
 namespace UE::Serialization::Private
 {
 
-static int32 GImportTypeHierarchyEnabled = 1;
+static int32 GImportTypeHierarchyEnabled = 0;
 static FAutoConsoleVariableRef CVarImportTypeHierarchyEnabled(
 	TEXT("s.ImportTypeHierarchyEnabled"),
 	GImportTypeHierarchyEnabled,
