IoStore: Decrement the unused serialized CookedHeaderSize by 8 to eliminate the impact of CL44349541 on unversioned cooked game patch sizes.
Note: This field should ultimately be removed or be serialized as 0.

#rnx
#tests QA smoke of standard build PF 68dd5f77600ab37e9eae4424

diff --git a/Engine/Source/Developer/IoStoreUtilities/Private/PackageStoreOptimizer.cpp b/Engine/Source/Developer/IoStoreUtilities/Private/PackageStoreOptimizer.cpp
--- a/Engine/Source/Developer/IoStoreUtilities/Private/PackageStoreOptimizer.cpp
+++ b/Engine/Source/Developer/IoStoreUtilities/Private/PackageStoreOptimizer.cpp
@@ -1093,6 +1093,18 @@ void FPackageStoreOptimizer::FinalizePackageHeader(FPackageStorePackage* Package
 		+ GraphDataSize
 		+ ImportedPackagesSize;
 
+	// NOTE: The CookedHeaderSize has been unused since EUnrealEngineObjectUE5Version::DATA_RESOURCES that added the bulk/data resource table.
+	//       It should ultimately be removed, until then freeze the FPackageFileSummary baseline impact on serialized CookedHeaderSize at
+	//       EUnrealEngineObjectUE5Version::OS_SUB_OBJECT_SHADOW_SERIALIZATION.
+	// NOTE: EUnrealEngineObjectUE5Version::IMPORT_TYPE_HIERARCHIES adds two int32 fields to FPackageFileSummary. This increases the total header size by 8 bytes.
+#if WITH_EDITORONLY_DATA
+	constexpr size_t ExpectedPackageFileSummarySize = 360;
+#else
+	constexpr size_t ExpectedPackageFileSummarySize = 344;
+#endif
+	static_assert(sizeof(FPackageFileSummary) == ExpectedPackageFileSummarySize, "Any change to FPackageFileSummary serialization requires adjusting the offset below, transient changes can just update the size in this ssert.");
+	constexpr uint32 CookedHeaderSizeOffsetToMinimizePatchSize = 8;
+
 	Package->HeaderBuffer = FIoBuffer(HeaderSize);
 	uint8* HeaderData = Package->HeaderBuffer.Data();
 	FMemory::Memzero(HeaderData, HeaderSize);
@@ -1100,7 +1112,7 @@ void FPackageStoreOptimizer::FinalizePackageHeader(FPackageStorePackage* Package
 	PackageSummary->HeaderSize = HeaderSize;
 	PackageSummary->Name = MappedPackageName;
 	PackageSummary->PackageFlags = Package->PackageFlags;
-	PackageSummary->CookedHeaderSize = Package->CookedHeaderSize;
+	PackageSummary->_Unused = Package->CookedHeaderSize - CookedHeaderSizeOffsetToMinimizePatchSize; // Was CookedHeaderSize
 	FBufferWriter HeaderArchive(HeaderData, HeaderSize);
 	HeaderArchive.Seek(sizeof(FZenPackageSummary));
 
diff --git a/Engine/Source/Runtime/CoreUObject/Internal/Serialization/ZenPackageHeader.h b/Engine/Source/Runtime/CoreUObject/Internal/Serialization/ZenPackageHeader.h
--- a/Engine/Source/Runtime/CoreUObject/Internal/Serialization/ZenPackageHeader.h
+++ b/Engine/Source/Runtime/CoreUObject/Internal/Serialization/ZenPackageHeader.h
@@ -12,7 +12,6 @@
 
 struct FZenPackageHeader
 {
-	uint32 CookedHeaderSize = 0;
 	uint32 ExportCount = 0; // Need to keep this count around after ExportMap is cleared
 	TOptional<FZenPackageVersioningInfo> VersioningInfo;
 	FNameMap NameMap;
diff --git a/Engine/Source/Runtime/CoreUObject/Private/Serialization/AsyncLoading2.cpp b/Engine/Source/Runtime/CoreUObject/Private/Serialization/AsyncLoading2.cpp
--- a/Engine/Source/Runtime/CoreUObject/Private/Serialization/AsyncLoading2.cpp
+++ b/Engine/Source/Runtime/CoreUObject/Private/Serialization/AsyncLoading2.cpp
@@ -2533,47 +2533,17 @@ public:
 
 	virtual int64 TotalSize() override
 	{
-		if (bExportsCookedToSeparateArchive)
-		{
-			return ExportSerialSize;
-		}
-		else
-		{
-			int64 CookedFileSize = (ActiveFPLB->EndFastPathLoadBuffer - ActiveFPLB->OriginalFastPathLoadBuffer);
-			CookedFileSize -= HeaderData->PackageSummary->HeaderSize;
-			CookedFileSize += HeaderData->CookedHeaderSize;
-			return CookedFileSize;
-		}
+		return ExportSerialSize;
 	}
 
 	virtual int64 Tell() override
 	{
-		if (bExportsCookedToSeparateArchive)
-		{
-			return (ActiveFPLB->StartFastPathLoadBuffer - ActiveFPLB->OriginalFastPathLoadBuffer) - ExportSerialOffset;
-		}
-		else
-		{
-			int64 CookedFilePosition = (ActiveFPLB->StartFastPathLoadBuffer - ActiveFPLB->OriginalFastPathLoadBuffer);
-			CookedFilePosition -= HeaderData->PackageSummary->HeaderSize;
-			CookedFilePosition += HeaderData->CookedHeaderSize;
-			return CookedFilePosition;
-		}
+		return (ActiveFPLB->StartFastPathLoadBuffer - ActiveFPLB->OriginalFastPathLoadBuffer) - ExportSerialOffset;
 	}
 
 	virtual void Seek(int64 Position) override
 	{
-		if (bExportsCookedToSeparateArchive)
-		{
-			ActiveFPLB->StartFastPathLoadBuffer = ActiveFPLB->OriginalFastPathLoadBuffer + ExportSerialOffset + Position;
-		}
-		else
-		{
-			uint64 BufferPosition = (uint64)Position;
-			BufferPosition -= HeaderData->CookedHeaderSize;
-			BufferPosition += HeaderData->PackageSummary->HeaderSize;
-			ActiveFPLB->StartFastPathLoadBuffer = ActiveFPLB->OriginalFastPathLoadBuffer + BufferPosition;
-		}
+		ActiveFPLB->StartFastPathLoadBuffer = ActiveFPLB->OriginalFastPathLoadBuffer + ExportSerialOffset + Position;
 		CheckBufferPosition(TEXT("InvalidSeek"));
 	}
 
@@ -2928,7 +2898,6 @@ private:
 	uint64 ExportSerialOffset = 0;
 	uint64 ExportSerialSize = 0;
 	bool bIsOptionalSegment = false;
-	bool bExportsCookedToSeparateArchive = false;
 
 	void FixupSoftObjectPathForInstancedPackage(FSoftObjectPath& InOutSoftObjectPath)
 	{
@@ -7937,7 +7906,6 @@ void FAsyncPackage2::InitializeExportArchive(FExportArchive& Ar, bool bIsOptiona
 	Ar.ExternalReadDependencies = &ExternalReadDependencies;
 	Ar.InstanceContext = &Desc.InstancingContext;
 	Ar.bIsOptionalSegment = bIsOptionalSegment;
-	Ar.bExportsCookedToSeparateArchive = Ar.UEVer() >= EUnrealEngineObjectUE5Version::DATA_RESOURCES;
 }
 
 EEventLoadNodeExecutionResult FAsyncPackage2::Event_ProcessExportBundle(FAsyncLoadingThreadState2& ThreadState, FAsyncPackage2* Package, int32 InExportBundleIndex)
diff --git a/Engine/Source/Runtime/CoreUObject/Private/Serialization/ZenPackageHeader.cpp b/Engine/Source/Runtime/CoreUObject/Private/Serialization/ZenPackageHeader.cpp
--- a/Engine/Source/Runtime/CoreUObject/Private/Serialization/ZenPackageHeader.cpp
+++ b/Engine/Source/Runtime/CoreUObject/Private/Serialization/ZenPackageHeader.cpp
@@ -63,7 +63,6 @@ FZenPackageHeader FZenPackageHeader::MakeView(FMemoryView Memory, FString& OutEr
 		PackageHeader.BulkDataMap = MakeArrayView(reinterpret_cast<const FBulkDataMapEntry*>(BulkDataMapData), BulkDataMapSize / sizeof(FBulkDataMapEntry));
 	}
 
-	PackageHeader.CookedHeaderSize = PackageHeader.PackageSummary->CookedHeaderSize;
 	PackageHeader.ImportedPublicExportHashes = TArrayView<const uint64>(
 		reinterpret_cast<const uint64*>(PackageHeaderDataPtr + PackageSummary->ImportedPublicExportHashesOffset),
 		(PackageSummary->ImportMapOffset - PackageSummary->ImportedPublicExportHashesOffset) / sizeof(uint64));
diff --git a/Engine/Source/Runtime/CoreUObject/Public/Serialization/AsyncLoading2.h b/Engine/Source/Runtime/CoreUObject/Public/Serialization/AsyncLoading2.h
--- a/Engine/Source/Runtime/CoreUObject/Public/Serialization/AsyncLoading2.h
+++ b/Engine/Source/Runtime/CoreUObject/Public/Serialization/AsyncLoading2.h
@@ -306,7 +306,7 @@ struct FZenPackageSummary
 	uint32 HeaderSize;
 	FMappedName Name;
 	uint32 PackageFlags;
-	uint32 CookedHeaderSize;
+	uint32 _Unused; // Was CookedHeaderSize
 	int32 ImportedPublicExportHashesOffset;
 	int32 ImportMapOffset;
 	int32 ExportMapOffset;
