[BugFix] Reinstancing followed a different codepath and handled IDO loading differently. A lot of this code was out of date from old naieve implementations so I updated it to act similar to PostLoadInstanceDataObject instead. This fixed various bugs and crashes:
- [CrashFix] Crash in UE::CopyPropertyValueSerializedData because an array property was matched between the source and destination but the inner type of those arrays were differing structs. this update made CopyPropertyValueSerializedData uneccesary so it was removed
- [BugFix] when recompiling containers who's element type changed, the expected behavior is for the container size to stay the same even if the element type isn't serializable directly. This worked on load because of the PostLoad behavior but needed to be fixed on recompile otherwise the array would get cleared on save+load because the IDO array was empty.

diff --git a/Engine/Source/Runtime/CoreUObject/Internal/UObject/InstanceDataObjectUtils.h b/Engine/Source/Runtime/CoreUObject/Internal/UObject/InstanceDataObjectUtils.h
--- a/Engine/Source/Runtime/CoreUObject/Internal/UObject/InstanceDataObjectUtils.h
+++ b/Engine/Source/Runtime/CoreUObject/Internal/UObject/InstanceDataObjectUtils.h
@@ -59,9 +59,6 @@ UE_API UObject* GetTemplateForInstanceDataObject(const UObject* Instance, const
 /** Generate a UClass that contains the union of the properties of PropertyTree and OwnerClass. */
 UClass* CreateInstanceDataObjectClass(const FPropertyPathNameTree* PropertyTree, const FUnknownEnumNames* EnumNames, UClass* OwnerClass, UObject* Outer);
 
-/** Copy whether each property was set by serialization from one IDO to another. */
-void CopyPropertyValueSerializedData(const FFieldVariant& OldField, void* OldDataPtr, const FFieldVariant& NewField, void* NewDataPtr);
-
 /** Returns whether the object is an IDO. */
 UE_API bool IsInstanceDataObject(const UObject* Object);
 
diff --git a/Engine/Source/Runtime/CoreUObject/Private/UObject/InstanceDataObjectUtils.cpp b/Engine/Source/Runtime/CoreUObject/Private/UObject/InstanceDataObjectUtils.cpp
--- a/Engine/Source/Runtime/CoreUObject/Private/UObject/InstanceDataObjectUtils.cpp
+++ b/Engine/Source/Runtime/CoreUObject/Private/UObject/InstanceDataObjectUtils.cpp
@@ -1730,120 +1730,6 @@ namespace UE
 		return CastField<FByteProperty>(Struct->FindPropertyByName(NAME_SerializedValues));
 	}
 
-	void CopyPropertyValueSerializedData(const FFieldVariant& OldField, void* OldDataPtr, const FFieldVariant& NewField, void* NewDataPtr)
-	{
-		if (const FStructProperty* OldAsStructProperty = OldField.Get<FStructProperty>())
-		{
-			const FStructProperty* NewAsStructProperty = NewField.Get<FStructProperty>();
-			checkf(NewAsStructProperty, TEXT("Type mismatch between OldField and NewField. Expected FStructProperty"));
-			CopyPropertyValueSerializedData(OldAsStructProperty->Struct, OldDataPtr, NewAsStructProperty->Struct, NewDataPtr);
-		}
-		else if (const FArrayProperty* OldAsArrayProperty = OldField.Get<FArrayProperty>())
-		{
-			const FArrayProperty* NewAsArrayProperty = NewField.Get<FArrayProperty>();
-			checkf(NewAsArrayProperty, TEXT("Type mismatch between OldField and NewField. Expected FArrayProperty"));
-			
-			FScriptArrayHelper OldArrayHelper(OldAsArrayProperty, OldDataPtr);
-			FScriptArrayHelper NewArrayHelper(NewAsArrayProperty, NewDataPtr);
-			for (int32 ArrayIndex = 0; ArrayIndex < OldArrayHelper.Num(); ++ArrayIndex)
-			{
-				if (NewArrayHelper.IsValidIndex(ArrayIndex))
-				{
-					CopyPropertyValueSerializedData(
-						OldAsArrayProperty->Inner, OldArrayHelper.GetElementPtr(ArrayIndex),
-						NewAsArrayProperty->Inner, NewArrayHelper.GetElementPtr(ArrayIndex));
-				}
-			}
-		}
-		else if (const FSetProperty* OldAsSetProperty = OldField.Get<FSetProperty>())
-		{
-			const FSetProperty* NewAsSetProperty = NewField.Get<FSetProperty>();
-			checkf(NewAsSetProperty, TEXT("Type mismatch between OldField and NewField. Expected FSetProperty"));
-			
-			FScriptSetHelper OldSetHelper(OldAsSetProperty, OldDataPtr);
-			FScriptSetHelper NewSetHelper(NewAsSetProperty, NewDataPtr);
-			FScriptSetHelper::FIterator OldItr = OldSetHelper.CreateIterator();
-			FScriptSetHelper::FIterator NewItr = NewSetHelper.CreateIterator();
-			
-			for (; OldItr && NewItr; ++OldItr, ++NewItr)
-			{
-				CopyPropertyValueSerializedData(
-					OldAsSetProperty->ElementProp, OldSetHelper.GetElementPtr(OldItr),
-					NewAsSetProperty->ElementProp, NewSetHelper.GetElementPtr(NewItr));
-			}
-		}
-		else if (const FMapProperty* OldAsMapProperty = OldField.Get<FMapProperty>())
-		{
-			const FMapProperty* NewAsMapProperty = NewField.Get<FMapProperty>();
-			checkf(NewAsMapProperty, TEXT("Type mismatch between OldField and NewField. Expected FMapProperty"));
-			
-			FScriptMapHelper OldMapHelper(OldAsMapProperty, OldDataPtr);
-			FScriptMapHelper NewMapHelper(NewAsMapProperty, NewDataPtr);
-			FScriptMapHelper::FIterator OldItr = OldMapHelper.CreateIterator();
-			FScriptMapHelper::FIterator NewItr = NewMapHelper.CreateIterator();
-			
-			for (; OldItr && NewItr; ++OldItr, ++NewItr)
-			{
-				CopyPropertyValueSerializedData(
-					OldAsMapProperty->KeyProp, OldMapHelper.GetKeyPtr(OldItr),
-					NewAsMapProperty->KeyProp, NewMapHelper.GetKeyPtr(NewItr));
-				CopyPropertyValueSerializedData(
-					OldAsMapProperty->ValueProp, OldMapHelper.GetValuePtr(OldItr),
-					NewAsMapProperty->ValueProp, NewMapHelper.GetValuePtr(NewItr));
-			}
-		}
-		else if (UStruct* OldAsStruct = OldField.Get<UStruct>())
-		{
-			const UStruct* NewAsStruct = NewField.Get<UStruct>();
-			checkf(NewAsStruct, TEXT("Type mismatch between OldField and NewField. Expected UStruct"));
-
-			auto FindMatchingProperty = [](const UStruct* Struct, const FProperty* Property) -> const FProperty*
-			{
-				for (const FProperty* StructProperty : TFieldRange<FProperty>(Struct))
-				{
-					if (StructProperty->GetFName() == Property->GetFName() && StructProperty->GetID() == Property->GetID())
-					{
-						return StructProperty;
-					}
-				}
-				return nullptr;
-			};
-
-			// clear existing set-flags first
-			if (const FByteProperty* SerializedValuesProperty = FindSerializedValuesProperty(NewAsStruct))
-			{
-				SerializedValuesProperty->InitializeValue_InContainer(NewDataPtr);
-			}
-
-			const FSerializedPropertyValueState OldSerializedState(OldAsStruct, OldDataPtr);
-			FSerializedPropertyValueState NewSerializedState(NewAsStruct, NewDataPtr);
-			for (const FProperty* OldSubProperty : TFieldRange<FProperty>(OldAsStruct))
-			{
-				if (const FProperty* NewSubProperty = FindMatchingProperty(NewAsStruct, OldSubProperty))
-				{
-					for (int32 ArrayIndex = 0; ArrayIndex < FMath::Min(OldSubProperty->ArrayDim, NewSubProperty->ArrayDim); ++ArrayIndex)
-					{
-						// copy set flags to new struct instance
-						if (OldSerializedState.IsSet(OldSubProperty, ArrayIndex))
-						{
-							NewSerializedState.Set(NewSubProperty, ArrayIndex);
-						}
-						else if (NewSubProperty->GetBoolMetaData(NAME_IsLooseMetadata))
-						{
-							// loose properties should be marked as serialized regardless of whether the old struct marked them as such
-							NewSerializedState.Set(NewSubProperty, ArrayIndex);
-						}
-					
-						// recurse
-						CopyPropertyValueSerializedData(
-							OldSubProperty, OldSubProperty->ContainerPtrToValuePtr<void>(OldDataPtr, ArrayIndex),
-							NewSubProperty, NewSubProperty->ContainerPtrToValuePtr<void>(NewDataPtr, ArrayIndex));
-					}
-				}
-			}
-		}
-	}
-
 bool IsInstanceDataObject(const UObject* Object)
 {
 	return Object && Object->GetClass()->UObject::IsA(UInstanceDataObjectClass::StaticClass());
diff --git a/Engine/Source/Runtime/CoreUObject/Private/UObject/PropertyBagRepository.cpp b/Engine/Source/Runtime/CoreUObject/Private/UObject/PropertyBagRepository.cpp
--- a/Engine/Source/Runtime/CoreUObject/Private/UObject/PropertyBagRepository.cpp
+++ b/Engine/Source/Runtime/CoreUObject/Private/UObject/PropertyBagRepository.cpp
@@ -202,17 +202,21 @@ void FPropertyBagRepository::ReassociateObjects(const TMap<UObject*, UObject*>&
 		if (AssociatedData.RemoveAndCopyValue(Pair.Key, OldBagData))
 		{
 			InstanceDataObjectToOwner.Remove(OldBagData.InstanceDataObject);
-			if (Pair.Value != nullptr) // Pair.Value can be nullptr when an object was destroyed like for example a UClass when it's deleted
-			{
-				if (const FPropertyBagAssociationData* NewBagData = AssociatedData.Find(Pair.Value))
-				{
-					InstanceDataObjectToOwner.Add(NewBagData->InstanceDataObject, Pair.Value);
 
-					CopyPropertyValueSerializedData(
-						OldBagData.InstanceDataObject->GetClass(), OldBagData.InstanceDataObject,
-						NewBagData->InstanceDataObject->GetClass(), NewBagData->InstanceDataObject);
-				}
+			if (FPropertyBagAssociationData* NewBagData = AssociatedData.Find(Pair.Value))
+			{
+				// copy data from owner to IDO. This can be needed if for example a type mismatch was detected but the serialization code still
+				// inferred some info. Like for example if an array type mismatch occured and the new array was auto-grown in the owner to match the
+				// old one. This array growth needs to be copied over to the IDO to maintain it
+				CopyTaggedProperties(Pair.Value, NewBagData->InstanceDataObject);
+
+#if !UE_BUILD_SHIPPING_WITH_EDITOR
+				// Run IDT operations on this instance to recover data
+				TObjectPtr<UObject> InOutInstance = const_cast<UObject*>(Pair.Value);
+				UE::PatchInstanceDataObject(NewBagData->InstanceDataObject, InOutInstance);
+#endif
 			}
+
 			OldBagData.Destroy();
 		}
 		else if (UStruct* TypeObject = Cast<UStruct>(Pair.Key))
diff --git a/Engine/Source/Runtime/CoreUObject/Private/UObject/PropertyPathName.cpp b/Engine/Source/Runtime/CoreUObject/Private/UObject/PropertyPathName.cpp
--- a/Engine/Source/Runtime/CoreUObject/Private/UObject/PropertyPathName.cpp
+++ b/Engine/Source/Runtime/CoreUObject/Private/UObject/PropertyPathName.cpp
@@ -7,6 +7,7 @@
 #include "Internationalization/Regex.h"
 #include "Dom/JsonObject.h"
 #include "Dom/JsonValue.h"
+#include "UObject/PackedObjectRef.h"
 
 namespace UE
 {
