[CIS Fix] UE::PatchInstanceDataObject Is not defined in 38.00 and shouldn't be called

diff --git a/Engine/Source/Runtime/CoreUObject/Private/UObject/PropertyBagRepository.cpp b/Engine/Source/Runtime/CoreUObject/Private/UObject/PropertyBagRepository.cpp
--- a/Engine/Source/Runtime/CoreUObject/Private/UObject/PropertyBagRepository.cpp
+++ b/Engine/Source/Runtime/CoreUObject/Private/UObject/PropertyBagRepository.cpp
@@ -209,12 +209,6 @@ void FPropertyBagRepository::ReassociateObjects(const TMap<UObject*, UObject*>&
 				// inferred some info. Like for example if an array type mismatch occured and the new array was auto-grown in the owner to match the
 				// old one. This array growth needs to be copied over to the IDO to maintain it
 				CopyTaggedProperties(Pair.Value, NewBagData->InstanceDataObject);
-
-#if !UE_BUILD_SHIPPING_WITH_EDITOR
-				// Run IDT operations on this instance to recover data
-				TObjectPtr<UObject> InOutInstance = const_cast<UObject*>(Pair.Value);
-				UE::PatchInstanceDataObject(NewBagData->InstanceDataObject, InOutInstance);
-#endif
 			}
 
 			OldBagData.Destroy();
