[StructUtils] Fix dirtying and recompile property bags when a UserDefinedStruct is modified.

diff --git a/Engine/Source/Runtime/CoreUObject/Private/StructUtils/PropertyBag.cpp b/Engine/Source/Runtime/CoreUObject/Private/StructUtils/PropertyBag.cpp
--- a/Engine/Source/Runtime/CoreUObject/Private/StructUtils/PropertyBag.cpp
+++ b/Engine/Source/Runtime/CoreUObject/Private/StructUtils/PropertyBag.cpp
@@ -4093,7 +4093,9 @@ bool UPropertyBag::ContainsUserDefinedStruct(const UUserDefinedStruct* UserDefin
 				if (OwnedUserDefinedStruct == UserDefinedStruct
 					|| OwnedUserDefinedStruct->PrimaryStruct == UserDefinedStruct
 					|| OwnedUserDefinedStruct == UserDefinedStruct->PrimaryStruct
-					|| OwnedUserDefinedStruct->PrimaryStruct == UserDefinedStruct->PrimaryStruct)
+					|| (OwnedUserDefinedStruct->PrimaryStruct == UserDefinedStruct->PrimaryStruct
+						&& OwnedUserDefinedStruct->PrimaryStruct != nullptr)
+					)
 				{
 					return true;
 				}
