UE-328847 - Setting a struct bool property (bitfield) to itself writes the wrong value

execLetBool was incorrectly passing a 'bool' for Step's result parameter. If we're working with a FBoolProperty, the underlying storage is actually a uint8, since it can also represents bitfields. This can lead to erroneous codegen in compilers that are a bit more aggressive at optimizing undefined behavior (eg: Clang 19 and newer).

The safer approach here is to simply use a 'uint8' in place of the 'bool'.

diff --git a/Engine/Source/Runtime/CoreUObject/Private/UObject/ScriptCore.cpp b/Engine/Source/Runtime/CoreUObject/Private/UObject/ScriptCore.cpp
--- a/Engine/Source/Runtime/CoreUObject/Private/UObject/ScriptCore.cpp
+++ b/Engine/Source/Runtime/CoreUObject/Private/UObject/ScriptCore.cpp
@@ -3020,21 +3020,26 @@ DEFINE_FUNCTION(UObject::execLetBool)
 		}
 	}
 
-	bool NewValue = false;
+	// Evaluate the r-value for this expression into Value
+	// Note that a bool property can represent bitfields,
+	// which is why a uint8 is needed here.
 
-	// evaluate the r-value for this expression into Value
-	Stack.Step( Stack.Object, &NewValue );
-	if( BoolAddr )
+	uint8 TempValue = 0;
+	Stack.Step(Stack.Object, &TempValue);
+
+	if (BoolAddr)
 	{
+		const bool NewBoolValue = (TempValue != 0);
+
 		checkSlow(CastField<FBoolProperty>(BoolProperty));
 		if (BoolProperty->HasSetter())
 		{
 			check(PropertyContainer != nullptr);
-			BoolProperty->SetValue_InContainer(PropertyContainer, &NewValue);
+			BoolProperty->SetValue_InContainer(PropertyContainer, &NewBoolValue);
 		}
 		else
 		{
-			BoolProperty->SetPropertyValue(BoolAddr, NewValue);
+			BoolProperty->SetPropertyValue(BoolAddr, NewBoolValue);
 		}
 	}
 }
