[StructUtils] Fixed incorrect logic in reordering properties in a property bag

* Also converted std::move to MoveTemp to maintain compliance with code standards

#rnx

diff --git a/Engine/Source/Runtime/CoreUObject/Private/StructUtils/PropertyBag.cpp b/Engine/Source/Runtime/CoreUObject/Private/StructUtils/PropertyBag.cpp
--- a/Engine/Source/Runtime/CoreUObject/Private/StructUtils/PropertyBag.cpp
+++ b/Engine/Source/Runtime/CoreUObject/Private/StructUtils/PropertyBag.cpp
@@ -1567,36 +1567,21 @@ namespace UE::StructUtils::Private
 			return EPropertyBagAlterationResult::TargetPropertyNotFound;
 		}
 
-		TArray<FPropertyBagPropertyDesc> NewDescs;
-		NewDescs.Reserve(Descs.Num());
+		FPropertyBagPropertyDesc SourceDesc = MoveTemp(Descs[SourcePropertyIndex]);
+		Descs.RemoveAt(SourcePropertyIndex);
 
-		const int32 Min = std::min(SourcePropertyIndex, TargetPropertyIndex);
-		const int32 SourceSpacer = SourcePropertyIndex < TargetPropertyIndex;
-
-		// Elements before the source and target properties.
-		int32 ReadIndex;
-		for (ReadIndex = 0; ReadIndex < Min; ++ReadIndex)
+		int32 NewTargetPropertyIndex = TargetPropertyIndex;
+		if (SourcePropertyIndex < TargetPropertyIndex)
 		{
-			NewDescs.Emplace(std::move(Descs[ReadIndex]));
+			--NewTargetPropertyIndex;
 		}
 
-		// Between the source and target properties (skip source if needed)
-		for (ReadIndex += SourceSpacer; ReadIndex < TargetPropertyIndex; ++ReadIndex)
-		{
-			NewDescs.Emplace(std::move(Descs[ReadIndex]));
-		}
-
-		// Source and target property
-		NewDescs.Emplace(std::move(Descs[bInsertBefore ? SourcePropertyIndex : TargetPropertyIndex]));
-		NewDescs.Emplace(std::move(Descs[bInsertBefore ? TargetPropertyIndex : SourcePropertyIndex]));
-
-		// The remaining elements
-		for (++ReadIndex; ReadIndex < Descs.Num(); ++ReadIndex)
+		if (!bInsertBefore)
 		{
-			NewDescs.Emplace(std::move(Descs[ReadIndex]));
+			++NewTargetPropertyIndex;
 		}
 
-		Descs = std::move(NewDescs);
+		Descs.Insert(MoveTemp(SourceDesc), NewTargetPropertyIndex);
 
 		return EPropertyBagAlterationResult::Success;
 	}
