UE-349454 - Preemptive BPVM fixes for Clang codegen

Follow-up to UE-328847.

Other script functions were found to use a 'bool' as the result argument for FFrame::Step(). This is potentially hazardous since bitfield data from a uint8 might get stored in a bool.

While these functions haven't demonstrated the same problem as UE-328847 with Clang 20.1.8, there is still a risk that a future version might introduce the issue.

diff --git a/Engine/Source/Runtime/CoreUObject/Private/UObject/ScriptCore.cpp b/Engine/Source/Runtime/CoreUObject/Private/UObject/ScriptCore.cpp
--- a/Engine/Source/Runtime/CoreUObject/Private/UObject/ScriptCore.cpp
+++ b/Engine/Source/Runtime/CoreUObject/Private/UObject/ScriptCore.cpp
@@ -1122,17 +1122,18 @@ void UObject::CallFunction( FFrame& Stack, RESULT_DECL, UFunction* Function )
 			FMemory::Memzero( Buffer, Function->ParmsSize );
 
 			// Form the RPC parameters.
-			for( TFieldIterator<FProperty> It(Function); It && (It->PropertyFlags & (CPF_Parm|CPF_ReturnParm))==CPF_Parm; ++It )
+			for (TFieldIterator<FProperty> It(Function); It && (It->PropertyFlags & (CPF_Parm|CPF_ReturnParm))==CPF_Parm; ++It)
 			{
 				uint8* CurrentPropAddr = It->ContainerPtrToValuePtr<uint8>(Buffer);
-				if ( CastField<FBoolProperty>(*It) && It->ArrayDim == 1 )
+				if (CastField<FBoolProperty>(*It) && It->ArrayDim == 1)
 				{
-					// we're going to get '1' returned for bools that are set, so we need to manually mask it in to the proper place
-					bool bValue = false;
-					Stack.Step(Stack.Object, &bValue);
-					if (bValue)
+					uint8 TempValue = 0;
+					Stack.Step(Stack.Object, &TempValue);
+					const bool NewBoolValue = (TempValue != 0);
+
+					if (NewBoolValue)
 					{
-						((FBoolProperty*)*It)->SetPropertyValue( CurrentPropAddr, true );
+						((FBoolProperty*)*It)->SetPropertyValue(CurrentPropAddr, true);
 					}
 				}
 				else
@@ -2585,14 +2586,14 @@ DEFINE_FUNCTION(UObject::execJumpIfNot)
 	// Get code offset.
 	CodeSkipSizeType Offset = Stack.ReadCodeSkipCount();
 
-	// Get boolean test value.
-	bool Value=0;
-	Stack.Step( Stack.Object, &Value );
+	uint8 TempValue = 0;
+	Stack.Step(Stack.Object, &TempValue);
+	const bool NewBoolValue = (TempValue != 0);
 
 	// Jump if false.
-	if( !Value )
+	if (!NewBoolValue)
 	{
-		Stack.Code = &Stack.Node->Script[ Offset ];
+		Stack.Code = &Stack.Node->Script[Offset];
 	}
 }
 IMPLEMENT_VM_FUNCTION( EX_JumpIfNot, execJumpIfNot );
@@ -2657,11 +2658,11 @@ DEFINE_FUNCTION(UObject::execPopExecutionFlowIfNot)
 	// Since this is a branch function, check for runaway script execution
 	CheckRunaway();
 
-	// Get boolean test value.
-	bool Value=0;
-	Stack.Step( Stack.Object, &Value );
+	uint8 TempValue = 0;
+	Stack.Step(Stack.Object, &TempValue);
+	const bool NewBoolValue = (TempValue != 0);
 
-	if (!Value)
+	if (!NewBoolValue)
 	{
 		// Try to pop an entry off the stack and go there
 		if (Stack.FlowStack.Num())
@@ -4187,10 +4188,12 @@ IMPLEMENT_VM_FUNCTION( EX_AutoRtfmStopTransact, execAutoRtfmStopTransact );
 DEFINE_FUNCTION(UObject::execAutoRtfmAbortIfNot)
 {
 	// evaluate a boolean expression
-	bool Result = false;
-	Stack.Step(Stack.Object, &Result);
+	uint8 TempValue = 0;
+	Stack.Step(Stack.Object, &TempValue);
+	const bool NewBoolValue = (TempValue != 0);
+
 	// if the expression returned false, abort the current transaction
-	if (!Result)
+	if (!NewBoolValue)
 	{
 		AutoRTFM::AbortTransaction();
 	}
