Memory-mapped bulk data alignment fix

Alignment was set to 1 to minimize all the padding in the memory-mapped bulk data storage & runtime memory.

Animations are memory-mapped on mobile, which violates the ACL anim sequence codec 16-byte alignment requirement. This requirement is enforced at runtime when we attempt to decompress in Debug/Development builds with an assert.

Changed default alignment to be 16.

#rnx

diff --git a/Engine/Config/Android/BaseAndroidEngine.ini b/Engine/Config/Android/BaseAndroidEngine.ini
--- a/Engine/Config/Android/BaseAndroidEngine.ini
+++ b/Engine/Config/Android/BaseAndroidEngine.ini
@@ -38,4 +38,4 @@ ASTCVersion=501
 
 [MemoryMappedFiles]
 Enable=true
-Alignment=1
+Alignment=16
diff --git a/Engine/Config/IOS/BaseIOSEngine.ini b/Engine/Config/IOS/BaseIOSEngine.ini
--- a/Engine/Config/IOS/BaseIOSEngine.ini
+++ b/Engine/Config/IOS/BaseIOSEngine.ini
@@ -85,7 +85,7 @@ DefaultPlatformService=IOSLocalNotification
 
 [MemoryMappedFiles]
 Enable=true
-Alignment=1
+Alignment=16
 
 [OnlineSubsystem]
 DefaultPlatformService=IOS
diff --git a/Engine/Source/Runtime/CoreUObject/Private/UObject/SavePackage2.cpp b/Engine/Source/Runtime/CoreUObject/Private/UObject/SavePackage2.cpp
--- a/Engine/Source/Runtime/CoreUObject/Private/UObject/SavePackage2.cpp
+++ b/Engine/Source/Runtime/CoreUObject/Private/UObject/SavePackage2.cpp
@@ -81,12 +81,12 @@ static FAutoConsoleVariableRef CVar_FixupStandaloneFlags(
 	TEXT("If non-zero, when the UAsset of a package is missing RF_Standalone, the flag is added. If zero, the flags are not changed and the save fails.")
 );
 
-int32 GMemoryMappedBulkDataAlignment = 1;
+int32 GMemoryMappedBulkDataAlignment = 16;
 static FAutoConsoleVariableRef CVar_MemoryMappedBulkDataAlignment(
 	TEXT("cook.MemoryMappedBulkDataAlignment"),
 	GMemoryMappedBulkDataAlignment,
 	TEXT("0 - Target platform memory mapping alignment is respected, might result in big zero-padding")
-	TEXT("1 - No special alignment required, 1-byte alignment used (no padding, default)")															  
+	TEXT("1 - No special alignment required, 1-byte alignment used")															  
 	TEXT("Any other non-zero value - overrides the memory mapped data alignment, implies padding")
 );
 
