Add a list of prefixed redirectors in FRedirectionSummary. Those prefixed redirectors are not in the list of global redirectors anymore.
This should make the list of global core redirectors constant across dlc and prevent useless recooks when cooking on top of a multi product cook.

#rnx

diff --git a/Engine/Source/Runtime/CoreUObject/Internal/UObject/CoreRedirects/RedirectionSummary.h b/Engine/Source/Runtime/CoreUObject/Internal/UObject/CoreRedirects/RedirectionSummary.h
--- a/Engine/Source/Runtime/CoreUObject/Internal/UObject/CoreRedirects/RedirectionSummary.h
+++ b/Engine/Source/Runtime/CoreUObject/Internal/UObject/CoreRedirects/RedirectionSummary.h
@@ -62,6 +62,8 @@ private:
 private:
 	TMap<FName, FRedirectContainer> RedirectsForPackage;
 	FRedirectContainer GlobalRedirects;
+	TSortedMap<FCoreRedirect, bool, FDefaultAllocator, FCompareRedirect> PrefixedRedirects;
+
 	/**
 	 * CoreRedirects are written when the engine is single threaded, but we do not write the hashes for the
 	 * global and per-package containers until they are requested on demand. The requests on demand can
diff --git a/Engine/Source/Runtime/CoreUObject/Private/UObject/RedirectionSummary.cpp b/Engine/Source/Runtime/CoreUObject/Private/UObject/RedirectionSummary.cpp
--- a/Engine/Source/Runtime/CoreUObject/Private/UObject/RedirectionSummary.cpp
+++ b/Engine/Source/Runtime/CoreUObject/Private/UObject/RedirectionSummary.cpp
@@ -35,8 +35,11 @@ FRedirectionSummary& FRedirectionSummary::operator=(FRedirectionSummary&& Other)
 
 	RedirectsForPackage = MoveTemp(Other.RedirectsForPackage);
 	GlobalRedirects = MoveTemp(Other.GlobalRedirects);
+	PrefixedRedirects = MoveTemp(Other.PrefixedRedirects);
+
 	Other.RedirectsForPackage.Empty();
 	Other.GlobalRedirects.Empty();
+	Other.PrefixedRedirects.Empty();
 
 	return *this;
 }
@@ -45,7 +48,11 @@ void FRedirectionSummary::Add(const FCoreRedirect& CoreRedirect)
 {
 	FWriteScopeLock ScopeLock(Lock);
 
-	if (CoreRedirect.IsWildcardMatch() || CoreRedirect.OldName.PackageName.IsNone())
+	if (CoreRedirect.IsPrefixMatch() && !CoreRedirect.OldName.PackageName.IsNone())
+	{
+		PrefixedRedirects.Add(FCoreRedirect(CoreRedirect));
+	}
+	else if (CoreRedirect.IsWildcardMatch() || CoreRedirect.OldName.PackageName.IsNone())
 	{
 		GlobalRedirects.Add(FCoreRedirect(CoreRedirect));
 	}
@@ -62,7 +69,11 @@ void FRedirectionSummary::Remove(const FCoreRedirect& CoreRedirect)
 {
 	FWriteScopeLock ScopeLock(Lock);
 
-	if (CoreRedirect.IsWildcardMatch() || CoreRedirect.OldName.PackageName.IsNone())
+	if (CoreRedirect.IsPrefixMatch() && !CoreRedirect.OldName.PackageName.IsNone())
+	{
+		PrefixedRedirects.Remove(CoreRedirect);
+	}
+	else if (CoreRedirect.IsWildcardMatch() || CoreRedirect.OldName.PackageName.IsNone())
 	{
 		GlobalRedirects.Remove(CoreRedirect);
 	}
@@ -136,22 +147,48 @@ void FRedirectionSummary::GetHashAffectingPackages(const TConstArrayView<FName>&
 		}
 	}
 
-	if (!bNeedWriteLock)
+	if (bNeedWriteLock)
 	{
-		return;
+
+		FWriteScopeLock WriteLock(Lock);
+		while (Index < PackageNamesNum)
+		{
+			FRedirectContainer* Container = RedirectsForPackage.Find(PackageNamesData[Index]);
+			if (Container)
+			{
+				FBlake3 Hasher;
+				Container->AppendHashInWriteLock(Hasher);
+				Hashes[Index] = Hasher.Finalize();
+			}
+			++Index;
+		}
 	}
 
-	FWriteScopeLock WriteLock(Lock);
-	while (Index < PackageNamesNum)
+	//Now add prefixed redirects.
+	for (int32 PackageIndex = 0; PackageIndex < PackageNamesNum; ++PackageIndex)
 	{
-		FRedirectContainer* Container = RedirectsForPackage.Find(PackageNamesData[Index]);
-		if (Container)
+		const FName& PackageName = PackageNamesData[PackageIndex];
+
+		bool bPrefixedRedirectFound = false;
+		FBlake3 Hasher;
+		for (const TPair<FCoreRedirect, bool>& Pair : PrefixedRedirects)
 		{
-			FBlake3 Hasher;
-			Container->AppendHashInWriteLock(Hasher);
-			Hashes[Index] = Hasher.Finalize();
+			FCoreRedirectObjectName ObjectName(PackageName.ToString());
+
+			const FCoreRedirect& Redirect = Pair.Key;
+			bool bMatch = Redirect.Matches(ObjectName);
+			if (bMatch)
+			{
+				Redirect.AppendHash(Hasher);
+				bPrefixedRedirectFound = true;
+			}
+		}
+
+		if (bPrefixedRedirectFound)
+		{
+			Hasher.Update(&Hashes[PackageIndex], sizeof(FBlake3Hash));
+			Hashes[PackageIndex] = Hasher.Finalize();
 		}
-		++Index;
 	}
 }
 
