Improvements to stats delegates for remote object transfer
- Added new delegates to record loading and saving of objects to database
- Consolidated the two pairs of migration receive and send delegates to be one delegate each
- No effective change.  These delegates are not used by main engine functionality yet.

diff --git a/Engine/Source/Runtime/CoreUObject/Private/UObject/RemoteObjectTransfer.cpp b/Engine/Source/Runtime/CoreUObject/Private/UObject/RemoteObjectTransfer.cpp
--- a/Engine/Source/Runtime/CoreUObject/Private/UObject/RemoteObjectTransfer.cpp
+++ b/Engine/Source/Runtime/CoreUObject/Private/UObject/RemoteObjectTransfer.cpp
@@ -31,10 +31,10 @@ UObject* MigrateObjectFromObjectData(FRemoteObjectData& ObjectData, const FUObje
 	DEFINE_MIGRATION_DELEGATE(RemoteObjectTransferDelegate);
 	DEFINE_MIGRATION_DELEGATE(RemoteObjectDeniedTransferDelegate);
 	DEFINE_MIGRATION_DELEGATE(RequestRemoteObjectDelegate);
-	DEFINE_MIGRATION_DELEGATE(OnObjectDataReceivedDelegate);
 	DEFINE_MIGRATION_DELEGATE(OnObjectsReceivedDelegate);
-	DEFINE_MIGRATION_DELEGATE(OnObjectDataSentDelegate);
 	DEFINE_MIGRATION_DELEGATE(OnObjectsSentDelegate);
+	DEFINE_MIGRATION_DELEGATE(OnObjectsLoadedFromDiskDelegate);
+	DEFINE_MIGRATION_DELEGATE(OnObjectsSavedToDiskDelegate);
 	DEFINE_MIGRATION_DELEGATE(OnObjectTouchedDelegate);
 	DEFINE_MIGRATION_DELEGATE(StoreRemoteObjectDataDelegate);
 	DEFINE_MIGRATION_DELEGATE(RestoreRemoteObjectDataDelegate);
@@ -221,9 +221,6 @@ public:
 							*GetNameSafe(SO));
 					}
 
-					// Notify that object data has been migrated (actually will be but SendObjectData does a MoveTemp on the data)
-					OnObjectDataSentDelegate.Broadcast(ObjectData, MigrationContext);
-
 					FMigrateSendParams SendRemoteObjectParams = { .MigrationContext = MigrationContext, .ObjectData = ObjectData };
 					SendRemoteObject(SendRemoteObjectParams);
 
@@ -233,7 +230,7 @@ public:
 						FUnsafeToMigrateScope UnsafeToMigrateScope;
 
 						// Notify that objects have been migrated to a remote server
-						OnObjectsSentDelegate.Broadcast(SentObjects, MigrationContext);
+						OnObjectsSentDelegate.Broadcast(SentObjects, ObjectData, MigrationContext);
 
 						for (UObject* SentObject : SentObjects)
 						{
@@ -829,10 +826,7 @@ UObject* MigrateObjectFromObjectData(FRemoteObjectData& ObjectData, const FUObje
 
 	check(!AutoRTFM::IsClosed());
 
-	// Notify that object data has been migrated from a remote server
-	OnObjectDataReceivedDelegate.Broadcast(ObjectData, MigrationContext);
-
-	FRemoteServerId LocalServerId = FRemoteServerId::GetLocalServerId();
+ 	FRemoteServerId LocalServerId = FRemoteServerId::GetLocalServerId();
 	TArray<FRemoteObjectId> ReceivedObjectRemoteIds;
 	TArray<UObject*> ReceivedObjects;
 	int32 RequestedObjectIndex = DeserializeObjectData(ObjectData, &MigrationContext, ReceivedObjectRemoteIds, ReceivedObjects);
@@ -916,8 +910,16 @@ UObject* MigrateObjectFromObjectData(FRemoteObjectData& ObjectData, const FUObje
 	}
 #endif
 
-	// Notify that objects have been migrated from a remote server
-	OnObjectsReceivedDelegate.Broadcast(ReceivedObjects, MigrationContext);
+ 	if (MigrationContext.RemoteServerId == DatabaseId)
+ 	{
+ 		// Let other systems know that we've loaded an object from storage
+ 		OnObjectsLoadedFromDiskDelegate.Broadcast(ReceivedObjects, ObjectData, MigrationContext);
+ 	}
+ 	else
+	{
+		// Notify that objects have been migrated from a remote server
+		OnObjectsReceivedDelegate.Broadcast(ReceivedObjects, ObjectData, MigrationContext);
+	}
 
 	// Set the ownership to the correct server id
 	for (UObject* ReceivedObject : ReceivedObjects)
@@ -1111,9 +1113,6 @@ void TransferObjectOwnershipToRemoteServer(UObject* Object, FRemoteServerId Dest
 			*GetNameSafe(SO));
 	}
 
-	// Notify that object data has been migrated (actually will be but SendObjectData does a MoveTemp on the data)
-	OnObjectDataSentDelegate.Broadcast(ObjectData, MigrationContext);
-
 	FMigrateSendParams SendRemoteObjectParams = { .MigrationContext = MigrationContext, .ObjectData = ObjectData };
 	GTransferQueue.SendRemoteObject(SendRemoteObjectParams);
 
@@ -1125,7 +1124,7 @@ void TransferObjectOwnershipToRemoteServer(UObject* Object, FRemoteServerId Dest
 		FUnsafeToMigrateScope UnsafeToMigrateScope;
 
 		// Notify that objects have been migrated to a remote server
-		OnObjectsSentDelegate.Broadcast(SentObjects, MigrationContext);
+		OnObjectsSentDelegate.Broadcast(SentObjects, ObjectData, MigrationContext);
 
 		for (UObject* SentObject : SentObjects)
 		{
@@ -1287,6 +1286,9 @@ void StoreObjectToDatabase(UObject* Object)
 
 			const FMigrateSendParams Params { .MigrationContext = StoreToDatabaseContext, .ObjectData = ObjectData };
 			StoreRemoteObjectDataDelegate.Execute(Params);
+
+			// Let other systems know that we've saved out object(s)
+			UE::RemoteObject::Transfer::OnObjectsSavedToDiskDelegate.Broadcast(SentObjects, Params.ObjectData, StoreToDatabaseContext);
 		}
 
 		checkf(SentObjects.Num() > 0, TEXT("PeerId:%s unable to store object %s to DatabaseId:%s"), *FRemoteServerId::GetLocalServerId().ToString(), *FRemoteObjectId(Object).ToString(), *DatabaseId.ToString());
diff --git a/Engine/Source/Runtime/CoreUObject/Public/UObject/RemoteObjectTransfer.h b/Engine/Source/Runtime/CoreUObject/Public/UObject/RemoteObjectTransfer.h
--- a/Engine/Source/Runtime/CoreUObject/Public/UObject/RemoteObjectTransfer.h
+++ b/Engine/Source/Runtime/CoreUObject/Public/UObject/RemoteObjectTransfer.h
@@ -143,24 +143,27 @@ namespace UE::RemoteObject::Transfer
 	/** Delegate that requests remote object data from LastKnownResidentServerId to be transferred to DestinationServerId. Allows requests to be forwarded if ObjectId does not reside on LastKnownResidentServerId. */
 	extern COREUOBJECT_API TDelegate<void(FRemoteWorkPriority /*RequestPriority*/, FRemoteObjectId /*ObjectId*/, FRemoteServerId /*LastKnownResidentServerId*/, FRemoteServerId /*DestinationServerId*/)> RequestRemoteObjectDelegate;
 
-	/** Delegate executed when objects have been migrated from another server. */
-	extern COREUOBJECT_API TMulticastDelegate<void(const FRemoteObjectData& /*ObjectData*/, const FUObjectMigrationContext& /*MigrationContext*/)> OnObjectDataReceivedDelegate;
+	/***************************************************************************************
+	 * Notification delegates (not required to be bound to); useful for stats reporting
+	 */
 
 	/** Delegate executed when objects have been migrated from another server. */
-	extern COREUOBJECT_API TMulticastDelegate<void(const TArray<UObject*>& /*MigratedObjects*/, const FUObjectMigrationContext& /*MigrationContext*/)> OnObjectsReceivedDelegate;
+	extern COREUOBJECT_API TMulticastDelegate<void(const TArray<UObject*>& /*ReceivedObjects*/, const FRemoteObjectData& /*ObjectData*/, const FUObjectMigrationContext& /*MigrationContext*/)> OnObjectsReceivedDelegate;
 
 	/** Delegate executed when object data has been migrated to another server */
-	extern COREUOBJECT_API TMulticastDelegate<void(const FRemoteObjectData& /*ObjectData*/, const FUObjectMigrationContext& /*MigrationContext*/)> OnObjectDataSentDelegate;
+	extern COREUOBJECT_API TMulticastDelegate<void(const TSet<UObject*>& /*SentObjects*/,  const FRemoteObjectData& /*ObjectData*/, const FUObjectMigrationContext& /*MigrationContext*/)> OnObjectsSentDelegate;
 
-	/** 
-	 * Delegate executed when objects have been migrated to another server. Migrated objects can still be resolved with weak pointers but at this point 
-	 * any changes to their internal state won't be migrated across to the remote server 
-	 */
-	extern COREUOBJECT_API TMulticastDelegate<void(const TSet<UObject*>& /*MigratedObjects*/, const FUObjectMigrationContext& /*MigrationContext*/)> OnObjectsSentDelegate;
+	/** Delegate executed when object data has been loaded from disk */
+	extern COREUOBJECT_API TMulticastDelegate<void(const TArray<UObject*>& /*LoadedObjects*/, const FRemoteObjectData& /*ObjectData*/, const FUObjectMigrationContext& /*MigrationContext*/)> OnObjectsLoadedFromDiskDelegate;
+
+	/** Delegate executed when object data has been saved to disk */
+	extern COREUOBJECT_API TMulticastDelegate<void(const TSet<UObject*>& /*SavedObjects*/, const FRemoteObjectData& /*ObjectData*/, const FUObjectMigrationContext& /*MigrationContext*/)> OnObjectsSavedToDiskDelegate;
 
 	/** Delegate executed when an object has been accessed by a transaction */
 	extern COREUOBJECT_API TMulticastDelegate<void(FRemoteTransactionId /*RequestId*/, FRemoteObjectId /*ObjectId*/)> OnObjectTouchedDelegate;
 
+	/***************************************************************************************/
+
 	/** Delegate that stores locally unrachable object data into a database. */
 	extern COREUOBJECT_API TDelegate<void(const FMigrateSendParams&)> StoreRemoteObjectDataDelegate;
 
