Include ObjectMacros.h before trying to use UE_WITH_CONSTINIT_UOBJECT to avoid warnings for undefined preprocessor macros when it is disabled.

diff --git a/Engine/Source/Runtime/CoreUObject/Public/UObject/RegisterCompiledInObjects.h b/Engine/Source/Runtime/CoreUObject/Public/UObject/RegisterCompiledInObjects.h
--- a/Engine/Source/Runtime/CoreUObject/Public/UObject/RegisterCompiledInObjects.h
+++ b/Engine/Source/Runtime/CoreUObject/Public/UObject/RegisterCompiledInObjects.h
@@ -2,6 +2,8 @@
 
 #pragma once
 
+#include "UObject/ObjectMacros.h"
+
 #if UE_WITH_CONSTINIT_UOBJECT 
 
 #include "Containers/ArrayView.h"
