Fix FObjectHandle builds to have constructors to handle non-trivial union member in TSAN builds.

diff --git a/Engine/Source/Runtime/CoreUObject/Public/UObject/ObjectHandle.h b/Engine/Source/Runtime/CoreUObject/Public/UObject/ObjectHandle.h
--- a/Engine/Source/Runtime/CoreUObject/Public/UObject/ObjectHandle.h
+++ b/Engine/Source/Runtime/CoreUObject/Public/UObject/ObjectHandle.h
@@ -76,6 +76,50 @@ namespace UE::CoreUObject::Private
 		{
 			return PointerOrRef != 0;
 		}
+
+	// When using TSAN, PointerOrRef has non-trivial copy construction so we need to explicitly implement constructors to handle the union.
+	// This means this type can't use aggregate initialization even when TSAN is disabled
+	FObjectHandlePrivate()
+	: PointerOrRef(0)
+	{
+	}
+	~FObjectHandlePrivate() = default;
+	constexpr FObjectHandlePrivate(const void* InPointer)
+		: Pointer(InPointer)
+	{
+	}
+	FObjectHandlePrivate(UPTRINT Packed)
+		: PointerOrRef(Packed)
+	{
+	}
+#if USING_THREAD_SANITISER
+	constexpr FObjectHandlePrivate(const FObjectHandlePrivate& Other)
+	{
+		UE_IF_CONSTEVAL
+		{
+			Pointer = Other.Pointer;
+		}
+		else
+		{
+			PointerOrRef = Other.PointerOrRef;
+		}
+	}
+	constexpr FObjectHandlePrivate& operator=(const FObjectHandlePrivate& Other)
+	{
+		UE_IF_CONSTEVAL
+		{
+			Pointer = Other.Pointer;
+		}
+		else
+		{
+			PointerOrRef = Other.PointerOrRef;
+		}
+		return *this;
+	}
+#else 
+	constexpr FObjectHandlePrivate(const FObjectHandlePrivate&) = default;
+	constexpr FObjectHandlePrivate& operator=(const FObjectHandlePrivate&) = default;
+#endif
 	};
 
 	/* Returns the packed object ref for this object IF one exists otherwise returns a null PackedObjectRef */
@@ -84,6 +128,7 @@ namespace UE::CoreUObject::Private
 	/* Creates and ObjectRef from a packed object ref*/
 	COREUOBJECT_API FObjectRef MakeObjectRef(FPackedObjectRef Handle);
 
+
 #elif UE_WITH_REMOTE_OBJECT_HANDLE
 
 	struct FRemoteObjectHandlePrivate
@@ -367,11 +412,11 @@ namespace UE::CoreUObject::Private
 #if UE_WITH_OBJECT_HANDLE_LATE_RESOLVE
 		UE_IF_CONSTEVAL
 		{
-			return { .Pointer = Object };
+			return FObjectHandle(Object);
 		}
 		else
 		{ 
-			return { .PointerOrRef = UPTRINT(Object) };
+			return FObjectHandle(UPTRINT(Object));
 		}
 #elif UE_WITH_REMOTE_OBJECT_HANDLE
 		return FObjectHandle(Object);
