BugHawk fixes and including ObjectMacros.h to avoid undefined symbol warnings

diff --git a/Engine/Source/Runtime/CoreUObject/Private/UObject/CompiledInObjectRegistry.h b/Engine/Source/Runtime/CoreUObject/Private/UObject/CompiledInObjectRegistry.h
--- a/Engine/Source/Runtime/CoreUObject/Private/UObject/CompiledInObjectRegistry.h
+++ b/Engine/Source/Runtime/CoreUObject/Private/UObject/CompiledInObjectRegistry.h
@@ -2,6 +2,8 @@
 
 #pragma once
 
+#include "UObject/ObjectMacros.h"
+
 #if UE_WITH_CONSTINIT_UOBJECT 
 
 #include "Misc/TVariant.h"
@@ -17,9 +19,9 @@ class FCompiledInObjectRegistry
     // Linked list of objects to construct. New items are added here, before items that may have
     // already been partially processed which are indicated by the following pointers.
     FRegisterCompiledInObjects* ListHead = nullptr;
-    // First block of objects that have already been processed by AddAndHashPendingObjects
+    // First block of objects that have already been processed by AddAndHashObjects
     FRegisterCompiledInObjects* AlreadyAdded = nullptr;
-    // First block of objects that have already been processed by FinishConstructingPendingObjects
+    // First block of objects that have already been processed by FinishConstructingObjects
     FRegisterCompiledInObjects* AlreadyConstructed = nullptr;
 
 public:
diff --git a/Engine/Source/Runtime/CoreUObject/Private/UObject/CompiledInUObjectInit.cpp b/Engine/Source/Runtime/CoreUObject/Private/UObject/CompiledInUObjectInit.cpp
--- a/Engine/Source/Runtime/CoreUObject/Private/UObject/CompiledInUObjectInit.cpp
+++ b/Engine/Source/Runtime/CoreUObject/Private/UObject/CompiledInUObjectInit.cpp
@@ -1,5 +1,7 @@
 // Copyright Epic Games, Inc. All Rights Reserved.
 
+#include "UObject/ObjectMacros.h"
+
 #if UE_WITH_CONSTINIT_UOBJECT
 
 #include "Logging/StructuredLog.h"
diff --git a/Engine/Source/Runtime/CoreUObject/Private/UObject/DeferredRegistry.h b/Engine/Source/Runtime/CoreUObject/Private/UObject/DeferredRegistry.h
--- a/Engine/Source/Runtime/CoreUObject/Private/UObject/DeferredRegistry.h
+++ b/Engine/Source/Runtime/CoreUObject/Private/UObject/DeferredRegistry.h
@@ -2,6 +2,8 @@
 
 #pragma once
 
+#include "UObject/ObjectMacros.h"
+
 #if !UE_WITH_CONSTINIT_UOBJECT
 
 #include "Modules/ModuleManager.h"
diff --git a/Engine/Source/Runtime/CoreUObject/Private/UObject/GarbageCollection.cpp b/Engine/Source/Runtime/CoreUObject/Private/UObject/GarbageCollection.cpp
--- a/Engine/Source/Runtime/CoreUObject/Private/UObject/GarbageCollection.cpp
+++ b/Engine/Source/Runtime/CoreUObject/Private/UObject/GarbageCollection.cpp
@@ -6595,7 +6595,7 @@ bool FStructProperty::ContainsObjectReference(TArray<const FStructProperty*>& En
 	// Version that can iterate properties before full linking
 	for (TFieldIterator<FProperty> It(Struct); It; ++It)
 	{
-		if (It->ContainsObjectReference(EncounteredStructProps))
+		if (It->ContainsObjectReference(EncounteredStructProps, InReferenceType))
 		{
 			bValue = true;
 			break;
diff --git a/Engine/Source/Runtime/CoreUObject/Public/UObject/Class.h b/Engine/Source/Runtime/CoreUObject/Public/UObject/Class.h
--- a/Engine/Source/Runtime/CoreUObject/Public/UObject/Class.h
+++ b/Engine/Source/Runtime/CoreUObject/Public/UObject/Class.h
@@ -1679,7 +1679,6 @@ namespace UE::CoreUObject::Private
 	//   ...:                         false
 	//   void(*HandlingFunctions[2]): { &Destruct<S>, &Identical<S> }
 	//
-	// We can't use constexpr because it uses reinterpret casts.
 	template <typename CppStruct, auto... Funcs>
 	inline constexpr const TStructOpsFakeVTable<Funcs...> TStructOpsFakeVTableInstance = 
 	{
diff --git a/Engine/Source/Runtime/CoreUObject/Public/UObject/IntrinsicClassBaseChain.h b/Engine/Source/Runtime/CoreUObject/Public/UObject/IntrinsicClassBaseChain.h
--- a/Engine/Source/Runtime/CoreUObject/Public/UObject/IntrinsicClassBaseChain.h
+++ b/Engine/Source/Runtime/CoreUObject/Public/UObject/IntrinsicClassBaseChain.h
@@ -23,7 +23,7 @@ namespace UE::Private
 	{
 		using BaseChainPtr = const FStructBaseChain*;
 		static constexpr SIZE_T Num = sizeof...(Classes);
-		static inline constinit BaseChainPtr Bases[] = { ((BaseChainPtr)GetConstInitClass<Classes>())... };
+		static inline constinit BaseChainPtr Bases[] = { (static_cast<BaseChainPtr>GetConstInitClass<Classes>())... };
 
 		static consteval TArrayView<const FStructBaseChain*> ToArrayView()
 		{
diff --git a/Engine/Source/Runtime/CoreUObject/Public/UObject/IntrinsicClassFwd.h b/Engine/Source/Runtime/CoreUObject/Public/UObject/IntrinsicClassFwd.h
--- a/Engine/Source/Runtime/CoreUObject/Public/UObject/IntrinsicClassFwd.h
+++ b/Engine/Source/Runtime/CoreUObject/Public/UObject/IntrinsicClassFwd.h
@@ -2,6 +2,8 @@
 
 #pragma once
 
+#include "UObject/ObjectMacros.h"
+
 #if UE_WITH_CONSTINIT_UOBJECT
 
 #include "Templates/NoDestroy.h"
@@ -33,4 +35,4 @@ UE_DEFINE_INTRINSIC_CLASS_CONSTINIT_ACCESSOR(UScriptStruct, COREUOBJECT_API)
 
 #define UE_DEFINE_INTRINSIC_CLASS_CONSTINIT_ACCESSOR(ClassName, RequiredAPI)
 
-#endif // !UE_WITH_CONSTINIT_UOBJECT
\ No newline at end of file
+#endif // UE_WITH_CONSTINIT_UOBJECT
\ No newline at end of file
diff --git a/Engine/Source/Runtime/CoreUObject/Public/UObject/UObjectBase.h b/Engine/Source/Runtime/CoreUObject/Public/UObject/UObjectBase.h
--- a/Engine/Source/Runtime/CoreUObject/Public/UObject/UObjectBase.h
+++ b/Engine/Source/Runtime/CoreUObject/Public/UObject/UObjectBase.h
@@ -221,7 +221,7 @@ public:
 	FORCEINLINE FName GetFName() const
 	{
 #if UE_WITH_CONSTINIT_UOBJECT
-		checkfSlow(!EnumHasAllFlags((EObjectFlags)GetFlagsInternal(), RF_NeedInitialization|RF_MarkAsNative), TEXT("UObjectBase::GetName called on a native object which is not yet initialized."));
+		checkfSlow(!EnumHasAllFlags((EObjectFlags)GetFlagsInternal(), RF_NeedInitialization|RF_MarkAsNative), TEXT("UObjectBase::GetFName called on a native object which is not yet initialized."));
 #endif // UE_WITH_CONSTINIT_UOBJECT
 		return NamePrivate.Name;
 	}
diff --git a/Engine/Source/Runtime/CoreUObject/Public/UObject/UObjectGlobals.h b/Engine/Source/Runtime/CoreUObject/Public/UObject/UObjectGlobals.h
--- a/Engine/Source/Runtime/CoreUObject/Public/UObject/UObjectGlobals.h
+++ b/Engine/Source/Runtime/CoreUObject/Public/UObject/UObjectGlobals.h
@@ -4192,7 +4192,7 @@ namespace UECodeGen_Private
 	#define METADATA_PARAMS(x, y)
 #endif
 
-// UE_METADATA expands to its contents if WITH_METADATA is set, otherwise expands to nothing
+// IF_WITH_METADATA expands to its contents if WITH_METADATA is set, otherwise expands to nothing
 #if WITH_METADATA
 	#define IF_WITH_METADATA(...) __VA_ARGS__
 #else
