Change preprocessor checks to use USING_THREAD_SANITISER || USING_INSTRUMENTATION to match definition of TSAN_ATOMIC and support -enableinstrumentation

diff --git a/Engine/Source/Runtime/CoreUObject/Public/UObject/ObjectHandle.h b/Engine/Source/Runtime/CoreUObject/Public/UObject/ObjectHandle.h
--- a/Engine/Source/Runtime/CoreUObject/Public/UObject/ObjectHandle.h
+++ b/Engine/Source/Runtime/CoreUObject/Public/UObject/ObjectHandle.h
@@ -92,7 +92,7 @@ namespace UE::CoreUObject::Private
 		: PointerOrRef(Packed)
 	{
 	}
-#if USING_THREAD_SANITISER
+#if USING_THREAD_SANITISER || USING_INSTRUMENTATION
 	constexpr FObjectHandlePrivate(const FObjectHandlePrivate& Other)
 	{
 		UE_IF_CONSTEVAL
@@ -116,10 +116,10 @@ namespace UE::CoreUObject::Private
 		}
 		return *this;
 	}
-#else 
+#else // USING_THREAD_SANITISER || USING_INSTRUMENTATION
 	constexpr FObjectHandlePrivate(const FObjectHandlePrivate&) = default;
 	constexpr FObjectHandlePrivate& operator=(const FObjectHandlePrivate&) = default;
-#endif
+#endif // USING_THREAD_SANITISER || USING_INSTRUMENTATION
 	};
 
 	/* Returns the packed object ref for this object IF one exists otherwise returns a null PackedObjectRef */
