Fix errors about FRemoteObjectHandlePrivate not being constexpr-constructible.

diff --git a/Engine/Source/Runtime/CoreUObject/Public/UObject/ObjectHandle.h b/Engine/Source/Runtime/CoreUObject/Public/UObject/ObjectHandle.h
--- a/Engine/Source/Runtime/CoreUObject/Public/UObject/ObjectHandle.h
+++ b/Engine/Source/Runtime/CoreUObject/Public/UObject/ObjectHandle.h
@@ -161,7 +161,12 @@ namespace UE::CoreUObject::Private
 	///these functions are always defined regardless of UE_WITH_OBJECT_HANDLE_LATE_RESOLVE value
 
 	/* Makes a resolved FObjectHandle from an UObject. */
+#if UE_WITH_REMOTE_OBJECT_HANDLE
+	// Not constexpr in this case yet
+	inline FObjectHandle MakeObjectHandle(UObject* Object);
+#else
 	inline constexpr FObjectHandle MakeObjectHandle(UObject* Object);
+#endif 
 
 	/* Returns the UObject from Handle and the handle is updated cache the resolved UObject */
 	inline UObject* ResolveObjectHandle(FObjectHandle& Handle);
@@ -407,9 +412,9 @@ inline uint32 GetTypeHash(UE::CoreUObject::Private::FRemoteObjectHandlePrivate H
 
 namespace UE::CoreUObject::Private
 {
+#if UE_WITH_OBJECT_HANDLE_LATE_RESOLVE
 	inline constexpr FObjectHandle MakeObjectHandle(UObject* Object)
 	{
-#if UE_WITH_OBJECT_HANDLE_LATE_RESOLVE
 		UE_IF_CONSTEVAL
 		{
 			return FObjectHandle(Object);
@@ -418,12 +423,18 @@ namespace UE::CoreUObject::Private
 		{ 
 			return FObjectHandle(UPTRINT(Object));
 		}
+	}
 #elif UE_WITH_REMOTE_OBJECT_HANDLE
+	inline FObjectHandle MakeObjectHandle(UObject* Object)
+	{
 		return FObjectHandle(Object);
+	}
 #else
+	inline constexpr FObjectHandle MakeObjectHandle(UObject* Object)
+	{
 		return Object;
-#endif
 	}
+#endif
 
 	inline UObject* ResolveObjectHandle(FObjectHandle& Handle)
 	{
