[UE] Add subobject instancing LLTs to include ground truths for non-nullable and unset TOptional reference types.

diff --git a/Engine/Source/Runtime/CoreUObject/Tests/UObject/InstancedPropertiesTest.cpp b/Engine/Source/Runtime/CoreUObject/Tests/UObject/InstancedPropertiesTest.cpp
--- a/Engine/Source/Runtime/CoreUObject/Tests/UObject/InstancedPropertiesTest.cpp
+++ b/Engine/Source/Runtime/CoreUObject/Tests/UObject/InstancedPropertiesTest.cpp
@@ -457,6 +457,38 @@ TEST_CASE_NAMED(FInstancedPropertiesTest, "CoreUObject::InstancedProperties", "[
 		CHECK_FALSE(OptionalInstancedObjectArrayValueProperty->Inner->GetBoolMetaData(TEXT("EditInline")));
 #endif
 	}
+
+	SECTION("Optional instanced object property - unset semantics (TObjectPtr)")
+	{
+		const FName OptionalNullableInstancedObjectPropertyName = GET_MEMBER_NAME_CHECKED(USubobjectInstancingTestOuterObject, OptionalNullableInnerObjectUnset);
+		const FOptionalProperty* OptionalNullableInstancedObjectProperty = CastField<FOptionalProperty>(USubobjectInstancingTestOuterObject::StaticClass()->FindPropertyByName(OptionalNullableInstancedObjectPropertyName));
+
+		REQUIRE(OptionalNullableInstancedObjectProperty != nullptr);
+
+		const FProperty* ValueProperty = OptionalNullableInstancedObjectProperty->GetValueProperty();
+		REQUIRE(ValueProperty != nullptr);
+
+		CHECK(ValueProperty->IsA<FObjectProperty>());
+		CHECK_FALSE(ValueProperty->HasIntrusiveUnsetOptionalState());
+		CHECK_FALSE(ValueProperty->HasAnyPropertyFlags(CPF_NonNullable));
+		CHECK(ValueProperty->HasAnyPropertyFlags(CPF_TObjectPtrWrapper));
+	}
+
+	SECTION("Optional instanced object property - unset semantics (TNonNullPtr)")
+	{
+		const FName OptionalNonNullableInstancedObjectPropertyName = GET_MEMBER_NAME_CHECKED(USubobjectInstancingTestOuterObject, OptionalNonNullableInnerObjectUnset);
+		const FOptionalProperty* OptionalNonNullableInstancedObjectProperty = CastField<FOptionalProperty>(USubobjectInstancingTestOuterObject::StaticClass()->FindPropertyByName(OptionalNonNullableInstancedObjectPropertyName));
+
+		REQUIRE(OptionalNonNullableInstancedObjectProperty != nullptr);
+
+		const FProperty* ValueProperty = OptionalNonNullableInstancedObjectProperty->GetValueProperty();
+		REQUIRE(ValueProperty != nullptr);
+
+		CHECK(ValueProperty->IsA<FObjectProperty>());
+		CHECK(ValueProperty->HasIntrusiveUnsetOptionalState());
+		CHECK(ValueProperty->HasAnyPropertyFlags(CPF_NonNullable));
+		CHECK_FALSE(ValueProperty->HasAnyPropertyFlags(CPF_TObjectPtrWrapper));
+	}
 }
 
 }
diff --git a/Engine/Source/Runtime/CoreUObject/Tests/UObject/SubobjectInstancingTest.cpp b/Engine/Source/Runtime/CoreUObject/Tests/UObject/SubobjectInstancingTest.cpp
--- a/Engine/Source/Runtime/CoreUObject/Tests/UObject/SubobjectInstancingTest.cpp
+++ b/Engine/Source/Runtime/CoreUObject/Tests/UObject/SubobjectInstancingTest.cpp
@@ -1193,6 +1193,47 @@ namespace SubobjectInstancingTest::Private
 				CHECK_FALSE(InnerObject->IsTemplateForSubobjects());
 			}
 		}
+
+		SECTION("Unset optional default subobject initialized from CDO")
+		{
+			USubobjectInstancingTestOuterObject* OuterObject = NewObject<USubobjectInstancingTestOuterObject>(GetTransientPackage(), InstancingTestClass);
+
+			CHECK_FALSE(OuterObject->OptionalNullableInnerObjectUnset.IsSet());
+			CHECK(OuterObject->OptionalNullableInnerObjectUnset.GetPtrOrNull() == nullptr);
+
+			CHECK_FALSE(OuterObject->OptionalNonNullableInnerObjectUnset.IsSet());
+			CHECK(OuterObject->OptionalNonNullableInnerObjectUnset.GetPtrOrNull() == nullptr);
+		}
+
+		SECTION("Unset optional default subobject initialized from template")
+		{
+			USubobjectInstancingTestOuterObject* Template = NewObject<USubobjectInstancingTestOuterObject>(GetTransientPackage(), InstancingTestClass, NAME_None, RF_ArchetypeObject);
+
+			CHECK_FALSE(Template->OptionalNullableInnerObjectUnset.IsSet());
+			CHECK_FALSE(Template->OptionalNonNullableInnerObjectUnset.IsSet());
+
+			USubobjectInstancingTestOuterObject* OuterObject = NewObject<USubobjectInstancingTestOuterObject>(GetTransientPackage(), InstancingTestClass, NAME_None, RF_NoFlags, Template);
+
+			CHECK_FALSE(OuterObject->OptionalNullableInnerObjectUnset.IsSet());
+			CHECK_FALSE(OuterObject->OptionalNonNullableInnerObjectUnset.IsSet());
+
+			Template->OptionalNullableInnerObjectUnset = nullptr;
+			Template->OptionalNonNullableInnerObjectUnset = NewObject<USubobjectInstancingTestObject>(Template, GET_MEMBER_NAME_CHECKED(USubobjectInstancingTestOuterObject, OptionalNonNullableInnerObjectUnset), Template->GetMaskedFlags(RF_PropagateToSubObjects));
+
+			CHECK(Template->OptionalNullableInnerObjectUnset.IsSet());
+			CHECK(Template->OptionalNonNullableInnerObjectUnset.IsSet());
+
+			CHECK(Template->OptionalNullableInnerObjectUnset.GetValue().Get() == nullptr);
+			REQUIRE(Template->OptionalNonNullableInnerObjectUnset.GetValue().Get() != nullptr);
+
+			OuterObject = NewObject<USubobjectInstancingTestOuterObject>(GetTransientPackage(), InstancingTestClass, NAME_None, RF_NoFlags, Template);
+
+			CHECK(OuterObject->OptionalNullableInnerObjectUnset.IsSet());
+			CHECK(OuterObject->OptionalNonNullableInnerObjectUnset.IsSet());
+
+			CHECK(OuterObject->OptionalNullableInnerObjectUnset.GetValue().Get() == nullptr);
+			REQUIRE(OuterObject->OptionalNonNullableInnerObjectUnset.GetValue().Get() != nullptr);
+		}
 	}
 
 	void RunSubobjectInstancingTests_Nested(UClass* InstancingTestClass)
diff --git a/Engine/Source/Runtime/CoreUObject/Tests/UObject/SubobjectInstancingTest.h b/Engine/Source/Runtime/CoreUObject/Tests/UObject/SubobjectInstancingTest.h
--- a/Engine/Source/Runtime/CoreUObject/Tests/UObject/SubobjectInstancingTest.h
+++ b/Engine/Source/Runtime/CoreUObject/Tests/UObject/SubobjectInstancingTest.h
@@ -167,6 +167,12 @@ public:
 	UPROPERTY(Instanced)
 	TOptional<TArray<TObjectPtr<USubobjectInstancingTestObject>>>  OptionalInnerObjectArray;
 
+	UPROPERTY(Instanced)
+	TOptional<TObjectPtr<USubobjectInstancingTestObject>> OptionalNullableInnerObjectUnset;
+
+	UPROPERTY(Instanced)
+	TOptional<TNonNullPtr<USubobjectInstancingTestObject>> OptionalNonNullableInnerObjectUnset;
+
 	UPROPERTY(Instanced)
 	TObjectPtr<USubobjectInstancingTestDirectlyNestedObject> InnerObjectWithDirectlyNestedObject;
 
