[UBT][CppCompileWarnings] Engine compliance wundef

diff --git a/Engine/Source/Runtime/CoreUObject/Private/UObject/UObjectConstructInternal.h b/Engine/Source/Runtime/CoreUObject/Private/UObject/UObjectConstructInternal.h
--- a/Engine/Source/Runtime/CoreUObject/Private/UObject/UObjectConstructInternal.h
+++ b/Engine/Source/Runtime/CoreUObject/Private/UObject/UObjectConstructInternal.h
@@ -2,7 +2,7 @@
 
 #pragma once
 
-#if !UE_WITH_CONSTINIT_UOBJECT
+#if !defined(UE_WITH_CONSTINIT_UOBJECT) || !UE_WITH_CONSTINIT_UOBJECT
 
 #include "CoreMinimal.h"
 #include "UObject/Class.h"
diff --git a/Engine/Source/Runtime/CoreUObject/Public/UObject/IntrinsicClassBaseChain.h b/Engine/Source/Runtime/CoreUObject/Public/UObject/IntrinsicClassBaseChain.h
--- a/Engine/Source/Runtime/CoreUObject/Public/UObject/IntrinsicClassBaseChain.h
+++ b/Engine/Source/Runtime/CoreUObject/Public/UObject/IntrinsicClassBaseChain.h
@@ -2,7 +2,7 @@
 
 #pragma once
 
-#if UE_WITH_CONSTINIT_UOBJECT
+#if defined(UE_WITH_CONSTINIT_UOBJECT) && UE_WITH_CONSTINIT_UOBJECT
 
 #include "UObject/Class.h"
 
