Fixed FRemoteObjectId::Generate(Object, Name) not generating the id based on a full path name for compiled-in types constructed in static init which could lead to generating the same id for two different objects
Refactored FRemoteObjectId::Generate() into multiple functions, added FRemoteObjectId::Generate(FRemoteServerIdConstants).

diff --git a/Engine/Source/Runtime/CoreUObject/Private/UObject/RemoteObject.cpp b/Engine/Source/Runtime/CoreUObject/Private/UObject/RemoteObject.cpp
--- a/Engine/Source/Runtime/CoreUObject/Private/UObject/RemoteObject.cpp
+++ b/Engine/Source/Runtime/CoreUObject/Private/UObject/RemoteObject.cpp
@@ -115,9 +115,9 @@ public:
 class FRemoteObjectMaps
 {
 	mutable FTransactionallySafeCriticalSection ObjectMapCritical;
-	// Maps remote object id to to a stub. Note that at the moment stubs are never destroyed (this is required by FRemoteObjectHandlePrivate)
+	// Maps remote object id to to a stub. Note that at the moment stubs are only destroyed on exit (this is required by FRemoteObjectHandlePrivate)
 	FRemoteObjectStubMap RemoteObjects;
-	// Maps remote object id to to a pathname. Note that at the moment these pathnames are never destroyed (this is required by FRemoteObjectClass)
+	// Maps remote object id to to a pathname. Note that at the moment these pathnames are only destroyed on exit (this is required by FRemoteObjectClass)
 	TMap<FRemoteObjectId, FRemoteObjectPathName*> AssetPaths;
 
 public:
@@ -876,16 +876,42 @@ namespace UE::CoreUObject::Private
 #define UE_WITH_REMOTE_ASSET_ID 1 // set this to 0 to disable remote asset IDs
 #endif
 
-FRemoteObjectId FRemoteObjectId::Generate(UObjectBase* InObject, const TCHAR* InName, EInternalObjectFlags InInitialFlags /*= EInternalObjectFlags::None*/)
+bool FRemoteObjectId::IsMigratingAsset(UObjectBase* InObject)
 {
-	using namespace UE::RemoteObject::Private;
 	using namespace UE::CoreUObject::Private;
 
+	// An object is a migrating asset if:
+	// its class or any of its outers' class is marked as MigratingAsset
+	// AND it's NOT an archetype or subobject of an archetype 
+	// AND it's NOT a subobject of a UStruct (class).
+	// In other words we don't want CDOs or Blueprint Classes to be migrating even when they are marked as migrating assets
+	bool bIsMarkedAsMigratingAsset = false;
+	bool bIsClassOrArchetype = false;
+
+	for (UObjectBase* OuterIt = InObject; OuterIt; OuterIt = FObjectHandleUtils::GetNonAccessTrackedOuterNoResolve(OuterIt))
+	{
+		UClass* Class = OuterIt->GetClass();
+		if (!!(OuterIt->GetFlags() & RF_ArchetypeObject) || Class->IsChildOf(UStruct::StaticClass()))
+		{
+			bIsClassOrArchetype = true;
+			break;
+		}
+		if (!!(OuterIt->GetFlags() & RF_MigratingAsset))
+		{
+			bIsMarkedAsMigratingAsset = true;
+		}
+	}
+
+	return bIsMarkedAsMigratingAsset && !bIsClassOrArchetype;
+}
+
+bool FRemoteObjectId::IsAssetInternal(UObjectBase* InObject, EInternalObjectFlags InInitialFlags)
+{
 	bool bIsAsset = false;
-#if UE_WITH_REMOTE_ASSET_ID
+
 	if (GIsInitialLoad || !!(InInitialFlags & EInternalObjectFlags::Native) || !!(InObject->GetFlags() & RF_ArchetypeObject))
 	{
-		// Native objects (classes, CDOs etc) or objects loaded during initial load are always in memory and are considered assets any server can find locally
+		// Native objects (classes, CDOs etc) are always in memory and are considered assets any server can find locally
 		// Note that this first condition can not touch too much of UObject API because we might literarlly be constructing the first StaticClass() etc.
 		// Hopefull GIsInitialLoad and the Native flag will filter most of the initially created objects and most of the native objects will be constructed before we ever hit the 'else' block
 		bIsAsset = true;
@@ -893,70 +919,121 @@ FRemoteObjectId FRemoteObjectId::Generate(UObjectBase* InObject, const TCHAR* In
 	else
 	{
 		FUObjectThreadContext& ThreadContext = FUObjectThreadContext::Get();
+		// Check if we're currently inside of a callstack where something is being loaded or the object has the RF_WasLoaded flag
 		if (ThreadContext.AsyncPackageLoader || (ThreadContext.GetSerializeContext() && ThreadContext.GetSerializeContext()->GetBeginLoadCount() > 0) || !!(InObject->GetFlags() & RF_WasLoaded))
 		{
-			// If we're constructing this object when loading content then this object is an asset if:
-			// its class or any of its outers' class is NOT marked as MigratingAsset
-			// OR it's an archetype or subobject of an archetype 
-			// OR it's a subobject of a UStruct (class)
-			bool bIsMigratingAsset = false;
-			bool bIsClassOrArchetypeSubobject = false;
-
-			for (UObjectBase* OuterIt = InObject; OuterIt; OuterIt = FObjectHandleUtils::GetNonAccessTrackedOuterNoResolve(OuterIt))
+			bIsAsset = !IsMigratingAsset(InObject);
+		}
+	}
+
+	return bIsAsset;
+}
+
+namespace
+{
+	/**
+	* Utility function that generates a similar output to GetObjectPathNameSafe(InObject) but for UObjectBase which may not even have a name or outer set
+	* @param InObject Object to generate the name for
+	* @param InName Name of the object (identical to InObject->GetFName().ToString() but the object may not even have the name assigned yet)
+	* @param InOuterPathName Path name of the outer of the specified object (optional, used only when the Outer is not set yet)
+	* @return PathName of the object or "None" if the specified object is null
+	*/
+	FString GetUObjectBasePathNameSafe(UObjectBase* InObject, const TCHAR* InName, const TCHAR* InOuterPathName)
+	{
+		FString Result;		
+		if (InObject)
+		{
+			FStringBuilderBase ObjectPathBuilder;
+			if (InObject->GetOuter())
 			{
-				UClass* Class = OuterIt->GetClass();
-				if (!!(OuterIt->GetFlags() & RF_ArchetypeObject) || Class->IsChildOf(UStruct::StaticClass()))
-				{
-					bIsClassOrArchetypeSubobject = true;
-					break;
-				}
-				if (!!(OuterIt->GetFlags() & RF_MigratingAsset))
-				{
-					bIsMigratingAsset = true;
-				}
+				ObjectPathBuilder << InObject->GetOuter()->GetPathName();
+			}
+			else if (InOuterPathName)
+			{
+				ObjectPathBuilder << InOuterPathName;
 			}
 
-			bIsAsset = bIsClassOrArchetypeSubobject || !bIsMigratingAsset;
+			if (ObjectPathBuilder.Len() > 0)
+			{
+				ObjectPathBuilder << SUBOBJECT_DELIMITER_CHAR;
+			}
+			ObjectPathBuilder << InName;
+			Result += FStringView(ObjectPathBuilder);
 		}
+		else
+		{
+			// Similar to what GetObjectNameSafe(nullptr) would return
+			Result = TEXT("None");
+		}
+		
+		return Result;
+	}
+
+	uint64 GenerateSerialNumberFromPathName(UObjectBase* InObject, const TCHAR* InName, const TCHAR* InOuterPathName)
+	{
+		// An asset must have the same remote object id on each game server. In order to achieve this the object's full path is converted into a 53-bit hash (the lower
+		// 53 bits of a 64-bit hash) and used as the remote id.
+		FString ObjectPathName = GetUObjectBasePathNameSafe(InObject, InName, InOuterPathName);
+
+		const uint64 ObjectPathHash64 = CityHash64(
+			reinterpret_cast<const char*>(*ObjectPathName),
+			ObjectPathName.Len() * sizeof(FString::ElementType)
+		);
+
+		const uint64 MaskLowBits = (1ULL << REMOTE_OBJECT_SERIAL_NUMBER_BIT_SIZE) - 1;
+		const uint64 ObjectPathHash53 = ObjectPathHash64 & MaskLowBits;
+
+		return ObjectPathHash53;
 	}
+}
+
+FRemoteObjectId FRemoteObjectId::Generate(UObjectBase* InObject, const TCHAR* InName, const TCHAR* InOuterPathName, EInternalObjectFlags InInitialFlags /*= EInternalObjectFlags::None*/)
+{
+	using namespace UE::RemoteObject::Private;
+	using namespace UE::CoreUObject::Private;
+
+	bool bIsAsset = false;
+#if UE_WITH_REMOTE_ASSET_ID
+	bIsAsset = IsAssetInternal(InObject, InInitialFlags);
 #endif // UE_WITH_REMOTE_ASSET_ID
 
 	FRemoteObjectId Result;
 
 	if (!bIsAsset)
 	{
-		Result = FRemoteObjectId(FRemoteServerId(ERemoteServerIdConstants::Local), RemoteObjectSerialNumber.fetch_add(1));
-		UE_LOG(LogRemoteObject, Verbose, TEXT("FRemoteObjectId::Generate: Type=Dynamic Object=%s RemoteId=%s"), InObject ? *InObject->GetFName().ToString() : TEXT("None"), *Result.ToString());
+		Result = FRemoteObjectId::Generate(ERemoteServerIdConstants::Local);
 	}
 	else
 	{
-		// An asset must have the same remote object id on each game server. In order to achieve this the object's full path is converted into a 53-bit hash (the lower
-		// 53 bits of a 64-bit hash) and used as the remote id.
+		Result = FRemoteObjectId(FRemoteServerId(ERemoteServerIdConstants::Asset), GenerateSerialNumberFromPathName(InObject, InName, InOuterPathName));
+	}
 
-		FString ObjectOuterPath = (InObject && InObject->GetOuter()) ? InObject->GetOuter()->GetPathName() : TEXT("");
-		FString ObjectName(InName);
+	UE_LOG(LogRemoteObject, Verbose, TEXT("FRemoteObjectId::Generate: Object=%s RemoteId=%s"), *GetUObjectBasePathNameSafe(InObject, InName, InOuterPathName), *Result.ToString(ERemoteIdToStringVerbosity::Id));
 
-		FStringBuilderBase ObjectPathBuilder;
-		if (ObjectOuterPath.Len() > 0)
-		{
-			ObjectPathBuilder << ObjectOuterPath;
-			ObjectPathBuilder << SUBOBJECT_DELIMITER_CHAR;
-		}
-		ObjectPathBuilder << ObjectName;
+	return Result;
+}
 
-		const uint64 ObjectPathHash64 = CityHash64(
-			reinterpret_cast<const char*>(ObjectPathBuilder.GetData()), 
-			ObjectPathBuilder.Len() * sizeof(FStringBuilderBase::ElementType)
-		);
+FRemoteObjectId FRemoteObjectId::Generate(UObjectBase* InObject, const TCHAR* InName, EInternalObjectFlags InInitialFlags /*= EInternalObjectFlags::None*/)
+{
+	return Generate(InObject, InName, nullptr, InInitialFlags);
+}
 
-		const uint64 MaskLowBits = (1ULL << REMOTE_OBJECT_SERIAL_NUMBER_BIT_SIZE) - 1;
-		const uint64 ObjectPathHash53 = ObjectPathHash64 & MaskLowBits;
+FRemoteObjectId FRemoteObjectId::Generate(ERemoteServerIdConstants ServerId)
+{
+	using namespace UE::RemoteObject::Private;
 
-		Result = FRemoteObjectId(FRemoteServerId(ERemoteServerIdConstants::Asset), ObjectPathHash53);
-		UE_LOG(LogRemoteObject, Verbose, TEXT("FRemoteObjectId::Generate: Type=Asset Object=%s RemoteId=%s Base64BitHash=%llu"), ObjectPathBuilder.ToString(), *Result.ToString(), ObjectPathHash64);
-	}
+	switch (ServerId)
+	{
+	case ERemoteServerIdConstants::Asset:
+		return FRemoteObjectId(FRemoteServerId(ERemoteServerIdConstants::Asset), AssetObjectSerialNumber.fetch_add(1));
 
-	return Result;
+	case ERemoteServerIdConstants::Local:
+		return FRemoteObjectId(FRemoteServerId(ERemoteServerIdConstants::Local), RemoteObjectSerialNumber.fetch_add(1));
+
+	default:
+		UE_LOG(LogRemoteObject, Fatal, TEXT("Unable to generate an object id for server id: %u"), (uint32)ServerId);
+	};
+	return FRemoteObjectId();
 }
 
 FString FRemoteObjectId::ToString(ERemoteIdToStringVerbosity InVerbosityOverride /*= ERemoteIdToStringVerbosity::Default*/) const
diff --git a/Engine/Source/Runtime/CoreUObject/Private/UObject/UObjectArray.cpp b/Engine/Source/Runtime/CoreUObject/Private/UObject/UObjectArray.cpp
--- a/Engine/Source/Runtime/CoreUObject/Private/UObject/UObjectArray.cpp
+++ b/Engine/Source/Runtime/CoreUObject/Private/UObject/UObjectArray.cpp
@@ -244,7 +244,7 @@ void FUObjectArray::AllocateUObjectIndex(UObjectBase* Object, EInternalObjectFla
 #if UE_WITH_REMOTE_OBJECT_HANDLE
 	if (!RemoteId.IsValid())
 	{
-		RemoteId = FRemoteObjectId::Generate(Object, *Object->GetFName().ToString(), InitialFlags);
+		RemoteId = FRemoteObjectId::Generate(Object, *Object->GetFName().ToString(), nullptr, InitialFlags);
 	}
 	else
 	{
diff --git a/Engine/Source/Runtime/CoreUObject/Private/UObject/UObjectBase.cpp b/Engine/Source/Runtime/CoreUObject/Private/UObject/UObjectBase.cpp
--- a/Engine/Source/Runtime/CoreUObject/Private/UObject/UObjectBase.cpp
+++ b/Engine/Source/Runtime/CoreUObject/Private/UObject/UObjectBase.cpp
@@ -550,7 +550,7 @@ void UObjectBase::Register(class UClass* (*StaticClassFn)(), const TCHAR* Packag
 	FPendingRegistrant* PendingRegistration = new FPendingRegistrant(this);
 	PendingRegistrants.Add(this, FPendingRegistrantInfo(StaticClassFn, InName, PackageName
 #if UE_WITH_REMOTE_OBJECT_HANDLE
-		, FRemoteObjectId::Generate(this, InName)
+		, FRemoteObjectId::Generate(this, InName, PackageName)
 #endif
 	));
 
diff --git a/Engine/Source/Runtime/CoreUObject/Public/UObject/RemoteObjectTypes.h b/Engine/Source/Runtime/CoreUObject/Public/UObject/RemoteObjectTypes.h
--- a/Engine/Source/Runtime/CoreUObject/Public/UObject/RemoteObjectTypes.h
+++ b/Engine/Source/Runtime/CoreUObject/Public/UObject/RemoteObjectTypes.h
@@ -329,9 +329,23 @@ public:
 
 	COREUOBJECT_API FString ToString(ERemoteIdToStringVerbosity InVerbosityOverride = ERemoteIdToStringVerbosity::Default) const;
 	static COREUOBJECT_API FRemoteObjectId FromString(const FStringView& InText);
-
+	
+	UE_DEPRECATED(5.8, "Use FRemoteServerId::Generate(InObject, InName, InPackageName)")
 	COREUOBJECT_API static FRemoteObjectId Generate(UObjectBase* InObject, const TCHAR* InName, EInternalObjectFlags InInitialFlags = EInternalObjectFlags::None);
 
+	/**
+	* Generates remote id for the specified object (assigns the server id of the remote object id based on the fact if the object was loaded or not (Asset / Local).
+	* Note that this happens BEFORE the object is added the to global GUObjectArray
+	* @param InObject Object to generated the id for
+	* @param InName Name of the object (because it might not have one set yet)
+	* @param InOuterPathName Pathname of the Outer object (when the Outer is not set yet, can be null if the Outer is set)
+	* @param InInitialFlags Initial internal flags the specified object is being constructed with
+	* @return Remote id
+	*/
+	COREUOBJECT_API static FRemoteObjectId Generate(UObjectBase* InObject, const TCHAR* InName, const TCHAR* InOuterPathName, EInternalObjectFlags InInitialFlags = EInternalObjectFlags::None);
+	COREUOBJECT_API static FRemoteObjectId Generate(ERemoteServerIdConstants ServerId);
+	COREUOBJECT_API static bool IsMigratingAsset(UObjectBase* InObject);
+
 	COREUOBJECT_API bool Serialize(FArchive& Ar);
 	COREUOBJECT_API bool NetSerialize(FArchive& Ar, class UPackageMap* Map, bool& bOutSuccess);
 
@@ -341,6 +355,9 @@ public:
 	friend UE::RemoteObject::Private::FRemoteIdLocalizationHelper;
 	friend struct FUObjectItem;
 	friend class FNetGUIDCache;	
+private:
+	
+	static bool IsAssetInternal(UObjectBase* InObject, EInternalObjectFlags InInitialFlags);
 };
 
 UE_FORCEINLINE_HINT uint32 GetTypeHash(const FRemoteObjectId& ObjectId)
