Fixing deprecation comment

diff --git a/Engine/Source/Runtime/CoreUObject/Public/UObject/RemoteObjectTypes.h b/Engine/Source/Runtime/CoreUObject/Public/UObject/RemoteObjectTypes.h
--- a/Engine/Source/Runtime/CoreUObject/Public/UObject/RemoteObjectTypes.h
+++ b/Engine/Source/Runtime/CoreUObject/Public/UObject/RemoteObjectTypes.h
@@ -330,7 +330,7 @@ public:
 	COREUOBJECT_API FString ToString(ERemoteIdToStringVerbosity InVerbosityOverride = ERemoteIdToStringVerbosity::Default) const;
 	static COREUOBJECT_API FRemoteObjectId FromString(const FStringView& InText);
 	
-	UE_DEPRECATED(5.8, "Use FRemoteServerId::Generate(InObject, InName, InPackageName)")
+	UE_DEPRECATED(5.8, "Use FRemoteObjectId::Generate(InObject, InName, InOuterPathName, InInitialFlags)")
 	COREUOBJECT_API static FRemoteObjectId Generate(UObjectBase* InObject, const TCHAR* InName, EInternalObjectFlags InInitialFlags = EInternalObjectFlags::None);
 
 	/**
