Fixed compile error in FLinkerLoad::operator<<.

diff --git a/Engine/Source/Runtime/CoreUObject/Private/UObject/LinkerLoad.cpp b/Engine/Source/Runtime/CoreUObject/Private/UObject/LinkerLoad.cpp
--- a/Engine/Source/Runtime/CoreUObject/Private/UObject/LinkerLoad.cpp
+++ b/Engine/Source/Runtime/CoreUObject/Private/UObject/LinkerLoad.cpp
@@ -6335,7 +6335,7 @@ FArchive& FLinkerLoad::operator<<(FObjectPtr& ObjectPtr)
 					// If type safety features are enabled, create a packed reference mapping for the placeholder-typed object.
 					// This resolves to NULL on access since the reference can't be cast to a pointer bound to its original type.
 					// However, the underlying value (when serialized) will always resolve to the placeholder object (e.g. for GC).
-					return FObjectPtr({ UE::CoreUObject::Private::MakePackedObjectRef(ResolvedObject).EncodedRef });
+					return FObjectPtr(FObjectHandle{ UE::CoreUObject::Private::MakePackedObjectRef(ResolvedObject).EncodedRef });
 #else
 					// If type safety features are disabled, serialize it as an unsafe reference to a placeholder-typed object.
 					// Note: Similar to hard references above, this means we won't find it for replacement at reinstancing time.
