Fix for implicitly constructed FObjectHandle and FPackedObjectRef.

diff --git a/Engine/Source/Runtime/CoreUObject/Private/UObject/LinkerLoadImportBehavior.cpp b/Engine/Source/Runtime/CoreUObject/Private/UObject/LinkerLoadImportBehavior.cpp
--- a/Engine/Source/Runtime/CoreUObject/Private/UObject/LinkerLoadImportBehavior.cpp
+++ b/Engine/Source/Runtime/CoreUObject/Private/UObject/LinkerLoadImportBehavior.cpp
@@ -218,7 +218,7 @@ bool TryLazyImport(const IAssetRegistryInterface& AssetRegistry, const FObjectIm
 		FObjectPathId ObjectPath;
 		FObjectRef ImportRef(Import.ObjectName, Import.ClassPackage, Import.ClassName, ObjectPath);
 		FPackedObjectRef PackedObjectRef = MakePackedObjectRef(ImportRef);
-		ObjectPtr = FObjectPtr({ PackedObjectRef.EncodedRef });
+		ObjectPtr = FObjectPtr(FObjectHandle{ PackedObjectRef.EncodedRef });
 		return true;
 	}
 
@@ -240,7 +240,7 @@ bool TryLazyImport(const IAssetRegistryInterface& AssetRegistry, const FObjectIm
 	AssetData.AssetName.AppendString(NameBuilder);
 	FObjectRef Ref = { AssetData.PackageName, AssetData.AssetClassPath.GetPackageName(), AssetData.AssetClassPath.GetAssetName(), FObjectPathId(NameBuilder) };
 	FPackedObjectRef PackedObjectRef = MakePackedObjectRef(Ref);
-	FObjectPtr Ptr({ PackedObjectRef.EncodedRef });
+	FObjectPtr Ptr(FObjectHandle{ PackedObjectRef.EncodedRef });
 	ObjectPtr = Ptr;
 	return true;
 }
@@ -275,7 +275,7 @@ bool TryLazyLoad(const UClass& Class, const FSoftObjectPath& ObjectPath, TObject
 	AssetData.AssetName.AppendString(NameBuilder);
 	FObjectRef ImportRef = { AssetData.PackageName, AssetData.AssetClassPath.GetPackageName(), AssetData.AssetClassPath.GetAssetName(), FObjectPathId(NameBuilder) };
 	FPackedObjectRef PackedObjectRef = MakePackedObjectRef(ImportRef);
-	FObjectPtr Ptr({ PackedObjectRef.EncodedRef });
+	FObjectPtr Ptr(FObjectHandle{ PackedObjectRef.EncodedRef });
 	ObjectPtr = Ptr;
 	return true;
 }
diff --git a/Engine/Source/Runtime/CoreUObject/Private/UObject/ObjectHandleIndex.cpp b/Engine/Source/Runtime/CoreUObject/Private/UObject/ObjectHandleIndex.cpp
--- a/Engine/Source/Runtime/CoreUObject/Private/UObject/ObjectHandleIndex.cpp
+++ b/Engine/Source/Runtime/CoreUObject/Private/UObject/ObjectHandleIndex.cpp
@@ -180,7 +180,7 @@ namespace UE::CoreUObject::Private
 	static inline FPackedObjectRef Pack(FPackageId PackageId, FObjectId ObjectId, bool bHasPlaceholderType = false)
 	{
 		checkf(PackageId.ToIndex() <= PackageIdMask, TEXT("Package count exceeded the space permitted within packed object references.  This implies over %d packages are in use."), PackageIdMask);
-		return { static_cast<UPTRINT>(PackageId.ToIndex()) << PackageIdShift |
+		return FPackedObjectRef{ static_cast<UPTRINT>(PackageId.ToIndex()) << PackageIdShift |
 				static_cast<UPTRINT>(ObjectId.ToPackedIndex() << ObjectIdShift) |
 				static_cast<UPTRINT>(bHasPlaceholderType << TypeIdShift) | 1 };
 	}
@@ -360,7 +360,7 @@ namespace UE::CoreUObject::Private
 
 		if (PathId.IsNone())
 		{
-			return { 0 }; //return null for packages as they have no outer
+			return FPackedObjectRef{ 0 }; //return null for packages as they have no outer
 		}
 		else if (PathId.IsSimple())
 		{
@@ -523,7 +523,7 @@ namespace UE::CoreUObject::Private
 	{
 		if (ObjectRef.IsNull())
 		{
-			return { 0 };
+			return FPackedObjectRef{ 0 };
 		}
 
 		return MakePackedObjectRef(ObjectRef.PackageName, ObjectRef.ClassPackageName, ObjectRef.ClassName, ObjectRef.GetObjectPath());
@@ -561,13 +561,13 @@ namespace UE::CoreUObject::Private
 	{
 		if (!Object)
 		{
-			return { 0 };
+			return FPackedObjectRef{ 0 };
 		}
 
 		const bool bHasPlaceholderType = UE::FPropertyBagRepository::IsPropertyBagPlaceholderObject(Object);
 		if (!bHasPlaceholderType && UE::LinkerLoad::FindLoadBehavior(*Object->GetClass()) == UE::LinkerLoad::EImportBehavior::Eager)
 		{
-			return { 0 };
+			return FPackedObjectRef{ 0 };
 		}
 
 		int32 ObjectIndex = GUObjectArray.ObjectToIndex(Object);
diff --git a/Engine/Source/Runtime/CoreUObject/Private/UObject/ObjectPtr.cpp b/Engine/Source/Runtime/CoreUObject/Private/UObject/ObjectPtr.cpp
--- a/Engine/Source/Runtime/CoreUObject/Private/UObject/ObjectPtr.cpp
+++ b/Engine/Source/Runtime/CoreUObject/Private/UObject/ObjectPtr.cpp
@@ -63,7 +63,7 @@ FObjectPtr FObjectPtr::GetOuter() const
 		return FObjectPtr(ResolveObjectHandleNoRead(LocalHandle)->GetOuter());
 	}
 	UE::CoreUObject::Private::FPackedObjectRef PackedRef = UE::CoreUObject::Private::GetOuter(ReadObjectHandlePackedObjectRefNoCheck(LocalHandle));
-	FObjectPtr Ptr({ PackedRef.EncodedRef });
+	FObjectPtr Ptr(FObjectHandle{ PackedRef.EncodedRef });
 	return Ptr;
 }
 
@@ -75,7 +75,7 @@ FObjectPtr FObjectPtr::GetPackage() const
 		return FObjectPtr(ResolveObjectHandleNoRead(LocalHandle)->GetPackage());
 	}
 	UE::CoreUObject::Private::FPackedObjectRef PackedRef = UE::CoreUObject::Private::GetPackage(ReadObjectHandlePackedObjectRefNoCheck(LocalHandle));
-	FObjectPtr Ptr({ PackedRef.EncodedRef });
+	FObjectPtr Ptr(FObjectHandle{ PackedRef.EncodedRef });
 	return Ptr;
 }
 
diff --git a/Engine/Source/Runtime/CoreUObject/Public/UObject/ObjectHandle.h b/Engine/Source/Runtime/CoreUObject/Public/UObject/ObjectHandle.h
--- a/Engine/Source/Runtime/CoreUObject/Public/UObject/ObjectHandle.h
+++ b/Engine/Source/Runtime/CoreUObject/Public/UObject/ObjectHandle.h
@@ -84,16 +84,16 @@ namespace UE::CoreUObject::Private
 	{
 	}
 	~FObjectHandlePrivate() = default;
-	constexpr FObjectHandlePrivate(const void* InPointer)
+	[[nodiscard]] constexpr explicit FObjectHandlePrivate(const void* InPointer)
 		: Pointer(InPointer)
 	{
 	}
-	FObjectHandlePrivate(UPTRINT Packed)
+	[[nodiscard]] explicit FObjectHandlePrivate(UPTRINT Packed)
 		: PointerOrRef(Packed)
 	{
 	}
 #if USING_THREAD_SANITISER || USING_INSTRUMENTATION
-	constexpr FObjectHandlePrivate(const FObjectHandlePrivate& Other)
+	[[nodiscard]] constexpr FObjectHandlePrivate(const FObjectHandlePrivate& Other)
 	{
 		UE_IF_CONSTEVAL
 		{
@@ -453,7 +453,7 @@ namespace UE::CoreUObject::Private
 #if UE_WITH_OBJECT_HANDLE_LATE_RESOLVE
 	inline FPackedObjectRef ReadObjectHandlePackedObjectRefNoCheck(FObjectHandle Handle)
 	{
-		return { Handle.PointerOrRef };
+		return FPackedObjectRef{ Handle.PointerOrRef };
 	}
 #endif
 
diff --git a/Engine/Source/Runtime/CoreUObject/Public/UObject/PackedObjectRef.h b/Engine/Source/Runtime/CoreUObject/Public/UObject/PackedObjectRef.h
--- a/Engine/Source/Runtime/CoreUObject/Public/UObject/PackedObjectRef.h
+++ b/Engine/Source/Runtime/CoreUObject/Public/UObject/PackedObjectRef.h
@@ -15,6 +15,13 @@ namespace UE::CoreUObject::Private
 	 */
 	struct FPackedObjectRef
 	{
+		[[nodiscard]] FPackedObjectRef() = default;
+
+		[[nodiscard]] explicit FPackedObjectRef(UPTRINT InEncodedRef)
+			: EncodedRef(InEncodedRef)
+		{
+		}
+
 		// Must be 0 for a reference to null.
 		// The least significant bit must always be 1 in a non-null reference.
 		UPTRINT EncodedRef;
diff --git a/Engine/Source/Runtime/CoreUObject/Tests/ObjectHandleTest.cpp b/Engine/Source/Runtime/CoreUObject/Tests/ObjectHandleTest.cpp
--- a/Engine/Source/Runtime/CoreUObject/Tests/ObjectHandleTest.cpp
+++ b/Engine/Source/Runtime/CoreUObject/Tests/ObjectHandleTest.cpp
@@ -27,7 +27,7 @@ protected:
 	void TestResolveFailure(UE::CoreUObject::Private::FPackedObjectRef PackedRef)
 	{
 		FSnapshotObjectRefMetrics ObjectRefMetrics(*this);
-		FObjectHandle TargetHandle = { PackedRef.EncodedRef };
+		FObjectHandle TargetHandle = FObjectHandle{ PackedRef.EncodedRef };
 		UObject* ResolvedObject = FObjectPtr(TargetHandle).Get();
 		ObjectRefMetrics.TestNumResolves(TEXT("NumResolves should be incremented by one after a resolve attempt"), 1);
 		ObjectRefMetrics.TestNumReads(TEXT("NumReads should be incremented by one after a resolve attempt"), 1);
@@ -254,7 +254,7 @@ TEST_CASE_METHOD(FObjectHandleTestBase, "CoreUObject::TObjectPtr::Null Behavior"
 
 	FObjectRef TargetRef(FName("SomePackage"), FName("ClassPackageName"), FName("ClassName"), UE::CoreUObject::Private::FObjectPathId("ObjectName"));
 	UE::CoreUObject::Private::FPackedObjectRef PackedObjectRef = UE::CoreUObject::Private::MakePackedObjectRef(TargetRef);
-	FObjectPtr ObjectPtr({ PackedObjectRef.EncodedRef });
+	FObjectPtr ObjectPtr(FObjectHandle{ PackedObjectRef.EncodedRef });
 	REQUIRE(!ObjectPtr.IsResolved()); //make sure not resolved
 
 	//an unresolved pointers compared against nullptr should still not resolve
@@ -355,7 +355,7 @@ TEST_CASE_METHOD(FObjectHandleTestBase, "CoreUObject::FObjectHandle::Resolve Mal
 
 	uint32 ObjectId = ~0u;
 	UPTRINT PackedId = ObjectId << 1 | 1;
-	UE::CoreUObject::Private::FPackedObjectRef PackedObjectRef = { PackedId };
+	UE::CoreUObject::Private::FPackedObjectRef PackedObjectRef{ PackedId };
 	TestResolveFailure(PackedObjectRef); // packed ref has a valid package id but invalid object id
 
 	TestResolveFailure(UE::CoreUObject::Private::FPackedObjectRef { 0xFFFF'FFFF'FFFF'FFFFull });
@@ -398,10 +398,10 @@ TEST_CASE_METHOD(FObjectHandleTestBase, "CoreUObject::FObjectHandle::Type Safety
 	FObjectHandle TestSafeObjectHandle = UE::CoreUObject::Private::MakeObjectHandle(TestSafeObject);
 	FObjectHandle TestSafeInvalidAddressHandle = UE::CoreUObject::Private::MakeObjectHandle(TestInvalidAddress);
 #if UE_WITH_OBJECT_HANDLE_LATE_RESOLVE
-	FObjectHandle TestLateResolveSafeObjectHandle = { UE::CoreUObject::Private::MakePackedObjectRef(TestSafeObject).EncodedRef };
-	FObjectHandle TestLateResolveUnsafeObjectHandle = { UE::CoreUObject::Private::MakePackedObjectRef(TestUnsafeObject).EncodedRef };
-	FObjectHandle TestLateResolveSafeInvalidAddressHandle = { (UE::CoreUObject::Private::FPackedObjectRef { reinterpret_cast<UPTRINT>(TestInvalidAddress) | 1 }).EncodedRef };
-	FObjectHandle TestLateResolveUnsafeInvalidAddressHandle = { (UE::CoreUObject::Private::FPackedObjectRef { reinterpret_cast<UPTRINT>(TestInvalidAddress) | (1 << UE::CoreUObject::Private::TypeIdShift) | 1 }).EncodedRef };
+	FObjectHandle TestLateResolveSafeObjectHandle{ UE::CoreUObject::Private::MakePackedObjectRef(TestSafeObject).EncodedRef };
+	FObjectHandle TestLateResolveUnsafeObjectHandle{ UE::CoreUObject::Private::MakePackedObjectRef(TestUnsafeObject).EncodedRef };
+	FObjectHandle TestLateResolveSafeInvalidAddressHandle{ (UE::CoreUObject::Private::FPackedObjectRef { reinterpret_cast<UPTRINT>(TestInvalidAddress) | 1 }).EncodedRef };
+	FObjectHandle TestLateResolveUnsafeInvalidAddressHandle{ (UE::CoreUObject::Private::FPackedObjectRef { reinterpret_cast<UPTRINT>(TestInvalidAddress) | (1 << UE::CoreUObject::Private::TypeIdShift) | 1 }).EncodedRef };
 #endif
 
 	// NULL/type-safe object handles should report as being safe
diff --git a/Engine/Source/Runtime/CoreUObject/Tests/ObjectRefTrackingTestBase.h b/Engine/Source/Runtime/CoreUObject/Tests/ObjectRefTrackingTestBase.h
--- a/Engine/Source/Runtime/CoreUObject/Tests/ObjectRefTrackingTestBase.h
+++ b/Engine/Source/Runtime/CoreUObject/Tests/ObjectRefTrackingTestBase.h
@@ -14,7 +14,7 @@
 inline FObjectHandle MakeUnresolvedHandle(const UObject* Obj)
 {
 	UE::CoreUObject::Private::FPackedObjectRef PackedObjectRef = UE::CoreUObject::Private::MakePackedObjectRef(Obj);
-	return { PackedObjectRef.EncodedRef };
+	return FObjectHandle{ PackedObjectRef.EncodedRef };
 }
 
 #endif
