Making remote object handle constructor constexpr to fix internal demo compile errors

diff --git a/Engine/Source/Runtime/CoreUObject/Public/UObject/ObjectHandle.h b/Engine/Source/Runtime/CoreUObject/Public/UObject/ObjectHandle.h
--- a/Engine/Source/Runtime/CoreUObject/Public/UObject/ObjectHandle.h
+++ b/Engine/Source/Runtime/CoreUObject/Public/UObject/ObjectHandle.h
@@ -133,12 +133,18 @@ namespace UE::CoreUObject::Private
 
 	struct FRemoteObjectHandlePrivate
 	{
-		UPTRINT PointerOrHandle;
+		// Stores either FRemoteObjectStub* or a UObject*
+		// Pointer union member is for constinit initialization where we want the linker to be able to do relocations for us within a binary.
+		union
+		{
+			TSAN_ATOMIC(UPTRINT) PointerOrHandle;
+			const void* Pointer;
+		};
 
 		FRemoteObjectHandlePrivate() = default;
-		explicit FRemoteObjectHandlePrivate(UObject* Object)
+		explicit constexpr FRemoteObjectHandlePrivate(UObject* Object)
+			: Pointer(Object)
 		{
-			PointerOrHandle = UPTRINT(Object);
 		}
 		explicit FRemoteObjectHandlePrivate(UE::RemoteObject::Handle::FRemoteObjectStub* RemoteInfo)
 		{
@@ -161,12 +167,7 @@ namespace UE::CoreUObject::Private
 	///these functions are always defined regardless of UE_WITH_OBJECT_HANDLE_LATE_RESOLVE value
 
 	/* Makes a resolved FObjectHandle from an UObject. */
-#if UE_WITH_REMOTE_OBJECT_HANDLE
-	// Not constexpr in this case yet
-	inline FObjectHandle MakeObjectHandle(UObject* Object);
-#else
 	inline constexpr FObjectHandle MakeObjectHandle(UObject* Object);
-#endif 
 
 	/* Returns the UObject from Handle and the handle is updated cache the resolved UObject */
 	inline UObject* ResolveObjectHandle(FObjectHandle& Handle);
@@ -425,7 +426,7 @@ namespace UE::CoreUObject::Private
 		}
 	}
 #elif UE_WITH_REMOTE_OBJECT_HANDLE
-	inline FObjectHandle MakeObjectHandle(UObject* Object)
+	inline constexpr FObjectHandle MakeObjectHandle(UObject* Object)
 	{
 		return FObjectHandle(Object);
 	}
