Fixup API macros

#rnx

diff --git a/Engine/Plugins/PCG/Source/PCG/Private/Metadata/PCGMetadataCommon.cpp b/Engine/Plugins/PCG/Source/PCG/Private/Metadata/PCGMetadataCommon.cpp
--- a/Engine/Plugins/PCG/Source/PCG/Private/Metadata/PCGMetadataCommon.cpp
+++ b/Engine/Plugins/PCG/Source/PCG/Private/Metadata/PCGMetadataCommon.cpp
@@ -2,6 +2,8 @@
 
 #include "Metadata/PCGMetadataCommon.h"
 
+#include UE_INLINE_GENERATED_CPP_BY_NAME(PCGMetadataCommon)
+
 namespace PCGMetadataDomainID
 {
 	const FPCGMetadataDomainID Default{EPCGMetadataDomainFlag::Default, -1, TEXT("Default")};
diff --git a/Engine/Plugins/PCGInterops/PCGGeometryScriptInterop/Source/PCGGeometryScriptInterop/Public/Elements/PCGGeometryBlueprintElement.h b/Engine/Plugins/PCGInterops/PCGGeometryScriptInterop/Source/PCGGeometryScriptInterop/Public/Elements/PCGGeometryBlueprintElement.h
--- a/Engine/Plugins/PCGInterops/PCGGeometryScriptInterop/Source/PCGGeometryScriptInterop/Public/Elements/PCGGeometryBlueprintElement.h
+++ b/Engine/Plugins/PCGInterops/PCGGeometryScriptInterop/Source/PCGGeometryScriptInterop/Public/Elements/PCGGeometryBlueprintElement.h
@@ -24,7 +24,7 @@ public:
 	PCGGEOMETRYSCRIPTINTEROP_API UPCGGeometryBlueprintElement();
 
 	UE_DEPRECATED(5.7, "Use Execute_Implementation instead")
-	PCGGEOMETRYSCRIPTINTEROP_API virtual void ExecuteWithContext_Implementation(FPCGContext& InContext, const FPCGDataCollection& Input, FPCGDataCollection& Output) {}
+	virtual void ExecuteWithContext_Implementation(FPCGContext& InContext, const FPCGDataCollection& Input, FPCGDataCollection& Output) {}
 
 	/** Virtual implementation to streamline the creation of a Geometry Script node processing. Do not override this function or Execute if you want to use the streamlined version. */
 	PCGGEOMETRYSCRIPTINTEROP_API virtual void Execute_Implementation(const FPCGDataCollection& Input, FPCGDataCollection& Output) override;
diff --git a/Engine/Plugins/Runtime/Metasound/Source/MetasoundFrontend/Public/MetasoundFrontendRegistryKey.h b/Engine/Plugins/Runtime/Metasound/Source/MetasoundFrontend/Public/MetasoundFrontendRegistryKey.h
--- a/Engine/Plugins/Runtime/Metasound/Source/MetasoundFrontend/Public/MetasoundFrontendRegistryKey.h
+++ b/Engine/Plugins/Runtime/Metasound/Source/MetasoundFrontend/Public/MetasoundFrontendRegistryKey.h
@@ -63,7 +63,7 @@ namespace Metasound::Frontend
 		UE_API FNodeClassInfo& operator=(FNodeClassInfo&& Other);
 
 		// Added to aid in deprecation of this class in favor of FNodeClassRegistryKey
-		UE_API static const FNodeClassInfo& GetInvalid()
+		static const FNodeClassInfo& GetInvalid()
 		{
 			static const FNodeClassInfo Invalid;
 			return Invalid;
diff --git a/Engine/Plugins/Runtime/StateTree/Source/StateTreeModule/Private/PropertyFunctions/StateTreeTransformPropertyFunctions.cpp b/Engine/Plugins/Runtime/StateTree/Source/StateTreeModule/Private/PropertyFunctions/StateTreeTransformPropertyFunctions.cpp
--- a/Engine/Plugins/Runtime/StateTree/Source/StateTreeModule/Private/PropertyFunctions/StateTreeTransformPropertyFunctions.cpp
+++ b/Engine/Plugins/Runtime/StateTree/Source/StateTreeModule/Private/PropertyFunctions/StateTreeTransformPropertyFunctions.cpp
@@ -3,6 +3,8 @@
 #include "StateTreeTransformPropertyFunctions.h"
 #include "StateTreeExecutionContext.h"
 
+#include UE_INLINE_GENERATED_CPP_BY_NAME(StateTreeTransformPropertyFunctions)
+
 void FStateTreeMakeTransformPropertyFunction::Execute(FStateTreeExecutionContext& Context) const
 {
 	FInstanceDataType& InstanceData = Context.GetInstanceData(*this);
@@ -14,4 +16,4 @@ void FStateTreeBreakTransformPropertyFunction::Execute(FStateTreeExecutionContex
 	FInstanceDataType& InstanceData = Context.GetInstanceData(*this);
 	InstanceData.OutTranslation = InstanceData.InTransform.GetTranslation();
 	InstanceData.OutRotation = InstanceData.InTransform.GetRotation();
-}
\ No newline at end of file
+}
diff --git a/Engine/Source/Runtime/CoreUObject/Public/UObject/Class.h b/Engine/Source/Runtime/CoreUObject/Public/UObject/Class.h
--- a/Engine/Source/Runtime/CoreUObject/Public/UObject/Class.h
+++ b/Engine/Source/Runtime/CoreUObject/Public/UObject/Class.h
@@ -412,13 +412,13 @@ class FStructBaseChain
 {
 #if USTRUCT_FAST_ISCHILDOF_IMPL == USTRUCT_ISCHILDOF_STRUCTARRAY
 protected:
-	COREUOBJECT_API FStructBaseChain() = default;
+	FStructBaseChain() = default;
 public:
 	COREUOBJECT_API ~FStructBaseChain();
 	UE_NONCOPYABLE(FStructBaseChain);
 	
 #if UE_WITH_CONSTINIT_UOBJECT
-	consteval COREUOBJECT_API FStructBaseChain(TArrayView<const FStructBaseChain*> InBases)
+	consteval FStructBaseChain(TArrayView<const FStructBaseChain*> InBases)
 		: StructBaseChainArray(InBases.GetData())
 		, NumStructBasesInChainMinusOne(InBases.Num()-1)
 	{
@@ -582,7 +582,7 @@ public:
 	COREUOBJECT_API explicit UStruct(const FObjectInitializer& ObjectInitializer, UStruct* InSuperStruct, SIZE_T ParamsSize = 0, SIZE_T Alignment = 0 );
 	COREUOBJECT_API virtual ~UStruct();
 #if UE_WITH_CONSTINIT_UOBJECT
-	consteval COREUOBJECT_API UStruct(
+	consteval UStruct(
 		UE::CodeGen::ConstInit::FObjectParams ObjectParams,
 		UE::CodeGen::ConstInit::FUFieldParams FieldParams,
 		UE::CodeGen::ConstInit::FStructParams StructParams
@@ -2682,7 +2682,7 @@ class UDelegateFunction : public UFunction
 	GENERATED_BODY()
 
 public:
-	COREUOBJECT_API explicit UDelegateFunction(const FObjectInitializer& ObjectInitializer)
+	explicit UDelegateFunction(const FObjectInitializer& ObjectInitializer)
 		: Super(ObjectInitializer)
 	{
 	}
@@ -2711,7 +2711,7 @@ class USparseDelegateFunction : public UDelegateFunction
 	GENERATED_BODY()
 
 public:
-	COREUOBJECT_API explicit USparseDelegateFunction(const FObjectInitializer& ObjectInitializer)
+	explicit USparseDelegateFunction(const FObjectInitializer& ObjectInitializer)
 		: Super(ObjectInitializer)
 	{
 	}
diff --git a/Engine/Source/Runtime/CoreUObject/Public/UObject/Package.h b/Engine/Source/Runtime/CoreUObject/Public/UObject/Package.h
--- a/Engine/Source/Runtime/CoreUObject/Public/UObject/Package.h
+++ b/Engine/Source/Runtime/CoreUObject/Public/UObject/Package.h
@@ -231,10 +231,10 @@ public:
 	{
 	}
 
-	COREUOBJECT_API ~UPackage() = default;
+	~UPackage() = default;
 
 #if UE_WITH_CONSTINIT_UOBJECT
-	explicit consteval COREUOBJECT_API UPackage(UE::CodeGen::ConstInit::FPackageParams InParams)
+	explicit consteval UPackage(UE::CodeGen::ConstInit::FPackageParams InParams)
 		: UObject(InParams.Object)
 		, bDirty(false)
 #if WITH_EDITORONLY_DATA
diff --git a/Engine/Source/Runtime/Engine/Private/Interfaces/Interface_AssetUserData.cpp b/Engine/Source/Runtime/Engine/Private/Interfaces/Interface_AssetUserData.cpp
--- a/Engine/Source/Runtime/Engine/Private/Interfaces/Interface_AssetUserData.cpp
+++ b/Engine/Source/Runtime/Engine/Private/Interfaces/Interface_AssetUserData.cpp
@@ -5,6 +5,8 @@
 #include "Templates/SubclassOf.h"
 #include "Engine/AssetUserData.h"
 
+#include UE_INLINE_GENERATED_CPP_BY_NAME(Interface_AssetUserData)
+
 void IInterface_AssetUserData::AddAssetUserData(UAssetUserData* InUserData)
 {
 }
@@ -33,4 +35,4 @@ const TArray<UAssetUserData*>* IInterface_AssetUserData::GetAssetUserDataArray()
 
 void IInterface_AssetUserData::RemoveUserDataOfClass(TSubclassOf<UAssetUserData> InUserDataClass)
 {
-}
\ No newline at end of file
+}
