[AsyncLoading2]

-Add loader stats to allow collecting stats, like partial flush records, so we may dump loader information once upon shutdown or when the loader is hung rather than each time we see them since it can be quite a lot of log information to have spread across a log file.
- By default we log inline like we do today, but when cooking we only report loader stats (in this case partial flushes) during shutdown

Stats are excluded in shipping configs.

#rnx

diff --git a/Engine/Source/Runtime/CoreUObject/Private/Serialization/AsyncLoading2.cpp b/Engine/Source/Runtime/CoreUObject/Private/Serialization/AsyncLoading2.cpp
--- a/Engine/Source/Runtime/CoreUObject/Private/Serialization/AsyncLoading2.cpp
+++ b/Engine/Source/Runtime/CoreUObject/Private/Serialization/AsyncLoading2.cpp
@@ -5065,6 +5065,56 @@ private:
 
 	/** Number of times we re-entered the async loading tick, mostly used by singlethreaded ticking. Debug purposes only. */
 	int32 AsyncLoadingTickCounter;
+
+	class FLoaderStats
+	{
+		struct FPartialFlushInfo
+		{
+			static constexpr int32 MaxCallstackDepth = 32;
+			TArray<uint64, TFixedAllocator<MaxCallstackDepth>> Callstack;
+			FName RequestedPackage;
+			FName RequestingPackage;
+			EAsyncPackageLoadingState2 RequestedPackageState;
+			EAsyncPackageLoadingState2 RequestingPackageState;
+
+			inline bool operator==(const FPartialFlushInfo& Other) const
+			{
+				return RequestedPackage == Other.RequestedPackage
+					&& RequestingPackage == Other.RequestingPackage;
+			}
+
+			inline friend uint32 GetTypeHash(const FPartialFlushInfo& Info)
+			{
+				return HashCombine(GetTypeHash(Info.RequestedPackageState), GetTypeHash(Info.RequestingPackageState));
+			}
+		};
+
+	public:
+		enum ELogMode : uint8
+		{
+			Inline,				// Stats are recorded and logged immediately
+			Dump,				// Stats are recorded and logged during ShutdownLoading() and DumpState()
+			DumpWithCallstack	// Same as Dump but callstacks are also recorded
+		};
+
+		void Clear();
+		void DumpToLog() const;
+		void RecordPartialFlush(const FAsyncPackage2* RequestedPackage, const FAsyncPackage2* RequestingPackage);
+		void Log(FStringBuilderBase& Builder, const FPartialFlushInfo& PartialFlush) const;
+
+		inline void SetLogMode(ELogMode Mode)
+		{
+#if !UE_BUILD_SHIPPING
+			LogMode = Mode;
+#endif
+		}
+
+#if !UE_BUILD_SHIPPING
+	private:
+		TSet<FPartialFlushInfo> PartialFlushes;
+		ELogMode LogMode;
+#endif
+	} LoaderStats;
 };
 
 #if WITH_VERSE_VM || defined(__INTELLISENSE__)
@@ -5237,6 +5287,19 @@ void FAsyncLoadingThread2::InitializeLoading()
 	}
 
 	FileOpenLogWrapper = (FPlatformFileOpenLog*)(FPlatformFileManager::Get().FindPlatformFile(FPlatformFileOpenLog::GetTypeName()));
+
+	if (IsRunningCookCommandlet())
+	{
+		// Don't overwhelm cook logs with stat logging during cooking and instead 
+		// accumulate records to flush when we shutdown the loader
+		LoaderStats.SetLogMode(FLoaderStats::ELogMode::DumpWithCallstack);
+	}
+	else
+	{
+		// When not cooking, log immediately so we don't bear any memory overhead from recording stats
+		// and can ensure we have logged records should the application crash before we can call DumpToLog
+		LoaderStats.SetLogMode(FLoaderStats::ELogMode::Inline);
+	}
 #endif
 
 	PackageStore.OnPendingEntriesAdded().AddLambda([this]()
@@ -9751,16 +9814,7 @@ void FAsyncLoadingThread2::UpdateSyncLoadContext(FAsyncLoadingThreadState2& Thre
 					if (RequestingPackage->AsyncPackageLoadingState < EAsyncPackageLoadingState2::DeferredPostLoad &&
 						RequestedPackage->AsyncPackageLoadingState < EAsyncPackageLoadingState2::Complete)
 					{
-						//
-						// Note: Update the FLoadingTests_RecursiveLoads_FullFlushFrom_Serialize test if you edit this error message
-						//
-						UE_LOG(LogStreaming, Display, TEXT("Flushing package %s (state: %s) recursively from another package %s (state: %s) will result in a partially loaded package to avoid a deadlock."),
-							*RequestedPackage->Desc.UPackageName.ToString(),
-							LexToString(RequestedPackage->AsyncPackageLoadingState),
-							*RequestingPackage->Desc.UPackageName.ToString(),
-							LexToString(RequestingPackage->AsyncPackageLoadingState),
-							RequestID
-						);
+						LoaderStats.RecordPartialFlush(RequestedPackage, RequestingPackage);
 
 						// Check if partial loading rules allow to release the package right now.
 						RequestedPackage->ConditionalReleasePartialRequests(ThreadState);
@@ -10290,6 +10344,8 @@ FAsyncLoadingThread2::~FAsyncLoadingThread2()
 	{
 		ShutdownLoading();
 	}
+
+	LoaderStats.DumpToLog();
 }
 
 void FAsyncLoadingThread2::ShutdownLoading()
@@ -12097,6 +12153,100 @@ void FGlobalImportStore::FlushDeferredDeletePackagesQueue()
 	AsyncLoadingThread.ProcessDeferredDeletePackagesQueue();
 }
 
+void FAsyncLoadingThread2::FLoaderStats::Clear()
+{
+#if !UE_BUILD_SHIPPING
+	PartialFlushes.Empty();
+#endif
+}
+
+void FAsyncLoadingThread2::FLoaderStats::DumpToLog() const
+{
+
+#if !UE_BUILD_SHIPPING
+	if ((LogMode != ELogMode::Dump) && (LogMode != ELogMode::DumpWithCallstack))
+	{
+		return;
+	}
+
+	UE_LOG(LogStreaming, Display, TEXT("============ Loader Stats ============"));
+
+	UE_LOG(LogStreaming, Display, TEXT("Unique Partial Flushes: %d"), PartialFlushes.Num());
+	TStringBuilder<1024> Builder;
+	for (const FPartialFlushInfo& PartialFlush : PartialFlushes)
+	{
+		Log(Builder, PartialFlush);
+	}
+
+	UE_LOG(LogStreaming, Display, TEXT("============ Loader Stats End ============"));
+#endif
+
+}
+
+void FAsyncLoadingThread2::FLoaderStats::RecordPartialFlush(const FAsyncPackage2* RequestedPackage, const FAsyncPackage2* RequestingPackage)
+{
+#if !UE_BUILD_SHIPPING
+	FPartialFlushInfo PartialFlush
+	{
+		.RequestedPackage = RequestedPackage->Desc.UPackageName,
+		.RequestingPackage = RequestingPackage->Desc.UPackageName,
+		.RequestedPackageState = RequestedPackage->AsyncPackageLoadingState,
+		.RequestingPackageState = RequestingPackage->AsyncPackageLoadingState,
+	};
+
+	if (LogMode == ELogMode::DumpWithCallstack)
+	{
+		PartialFlush.Callstack.SetNumUninitialized(FPartialFlushInfo::MaxCallstackDepth);
+
+		const int32 NumStackFramesToIgnore = 3; // Ignore this callstack and 2 frames from CaptureStackBackTrace
+		uint32 NumFrames = FPlatformStackWalk::CaptureStackBackTrace(PartialFlush.Callstack.GetData(), PartialFlush.Callstack.Num());
+		PartialFlush.Callstack.SetNum(NumFrames);
+	}
+
+	if (LogMode == ELogMode::Inline)
+	{
+		TStringBuilder<1024> Builder;
+		Log(Builder, PartialFlush);
+	}
+	else
+	{
+		PartialFlushes.Add(MoveTemp(PartialFlush));
+	}
+#endif
+}
+
+void FAsyncLoadingThread2::FLoaderStats::Log(FStringBuilderBase& Builder, const FPartialFlushInfo& PartialFlush) const
+{
+#if !UE_BUILD_SHIPPING
+	//
+	// Note: Update the FLoadingTests_RecursiveLoads_FullFlushFrom_Serialize test if you edit this error message
+	//
+	Builder.Reset();
+	Builder.Appendf(
+		TEXT("Partially loaded package %s (state: %s) to avoid deadlock when it was recursively flushed by package %s (state: %s). Callstack:\n"),
+		*PartialFlush.RequestedPackage.ToString(),
+		LexToString(PartialFlush.RequestedPackageState),
+		*PartialFlush.RequestingPackage.ToString(),
+		LexToString(PartialFlush.RequestingPackageState));
+
+	FString HumanReadableString;
+	for (const uint64 ProgramCounter : PartialFlush.Callstack)
+	{
+		FProgramCounterSymbolInfoEx SymbolInfo;
+		FPlatformStackWalk::ProgramCounterToSymbolInfoEx(ProgramCounter, SymbolInfo);
+		FPlatformStackWalk::SymbolInfoToHumanReadableStringEx(SymbolInfo, HumanReadableString);
+
+		Builder.Append(TEXT("\t"));
+		Builder.Append(HumanReadableString);
+		Builder.Append(TEXT("\n"));
+
+		HumanReadableString.Reset();
+	}
+
+	UE_LOG(LogStreaming, Display, TEXT("%s"), Builder.ToString());
+#endif
+}
+
 #if ALT2_DUMP_STATE_ON_HANG
 
 void FAsyncPackage2::DumpState()
@@ -12107,6 +12257,8 @@ void FAsyncPackage2::DumpState()
 	DumpStateImpl(Set);
 
 	UE_LOG(LogStreaming, Warning, TEXT("============"));
+
+	AsyncLoadingThread.LoaderStats.DumpToLog();
 }
 
 void FAsyncPackage2::DumpStateImpl(TSet<FAsyncPackage2*>& Set, int32 Indent, TMultiMap<FEventLoadNode2*, FEventLoadNode2*>* MappedNodes)
diff --git a/Engine/Source/Runtime/Engine/Private/Tests/Loading/AsyncLoadingTests_RecursiveLoads.cpp b/Engine/Source/Runtime/Engine/Private/Tests/Loading/AsyncLoadingTests_RecursiveLoads.cpp
--- a/Engine/Source/Runtime/Engine/Private/Tests/Loading/AsyncLoadingTests_RecursiveLoads.cpp
+++ b/Engine/Source/Runtime/Engine/Private/Tests/Loading/AsyncLoadingTests_RecursiveLoads.cpp
@@ -213,7 +213,7 @@ bool FLoadingTests_RecursiveLoads_FullFlushFrom_Serialize::RunTest(const FString
 {
 	FLoadingTestsScope LoadingTestScope(this);
 
-	AddExpectedMessage(TEXT("will result in a partially loaded package to avoid a deadlock."), EAutomationExpectedErrorFlags::Contains);
+	AddExpectedMessage(TEXT("to avoid deadlock when it was recursively flushed by package"), EAutomationExpectedErrorFlags::Contains);
 
 	UE::FManualResetEvent PostLoadEvent;
 	std::atomic<bool> bHasBeenPostloaded = false;
