Added comments to all CoreUObject enum flags to clarify that composite flags should not be made a enumerators within the enums, like already existed for other enums.
Fixed EPropertyObjectReferenceType enumerators which didn't do this.

diff --git a/Engine/Source/Runtime/CoreUObject/Private/UObject/Class.cpp b/Engine/Source/Runtime/CoreUObject/Private/UObject/Class.cpp
--- a/Engine/Source/Runtime/CoreUObject/Private/UObject/Class.cpp
+++ b/Engine/Source/Runtime/CoreUObject/Private/UObject/Class.cpp
@@ -8,6 +8,7 @@
 #include "Misc/ConfigCacheIni.h"
 #include "Misc/ConfigContext.h"
 #include "Misc/OutputDeviceHelper.h"
+#include "Misc/EnumRange.h"
 #include "Misc/FeedbackContext.h"
 #include "Misc/OutputDeviceConsole.h"
 #include "Misc/PackageAccessTrackingOps.h"
@@ -1107,7 +1108,7 @@ void UStruct::Link(FArchive& Ar, bool bRelinkExistingProperties)
 
 		// Ref link contains any properties which contain object references including types with user-defined serializers which don't explicitly specify whether they 
 		// contain object references
-		if (Property->ContainsObjectReference(EncounteredStructProps, EPropertyObjectReferenceType::Any))
+		if (Property->ContainsObjectReference(EncounteredStructProps, EPropertyObjectReferenceType_Any))
 		{
 			RefLinkBuilder.AppendNoTerminate(*Property);
 		}
@@ -4806,21 +4807,17 @@ static bool FindUndeclaredObjectReferencesInStructSerializers(FAutomationTestBas
 		if(bExactCheck && EnumHasAllFlags(DeclaredReferences, EPropertyObjectReferenceType::Conservative))
 		{
 			FString SerializedTypes;
-			for (uint32 Flag = 1, Max = (uint32)EPropertyObjectReferenceType::MAX; Flag != Max; Flag <<= 1)
+			for (EPropertyObjectReferenceType TypedFlag : MakeFlagsRange(Ar.References))
 			{
-				EPropertyObjectReferenceType TypedFlag = (EPropertyObjectReferenceType)Flag;
-				if (EnumHasAllFlags(Ar.References, TypedFlag))
+				if (SerializedTypes.IsEmpty())
 				{
-					if (SerializedTypes.IsEmpty())
-					{
-						SerializedTypes = LexToString(TypedFlag);
-					}
-					else 
-					{
-						SerializedTypes.Appendf(TEXT(" | %s"),LexToString(TypedFlag));
-					}
+					SerializedTypes = LexToString(TypedFlag);
 				}
-			}	
+				else 
+				{
+					SerializedTypes.Appendf(TEXT(" | %s"),LexToString(TypedFlag));
+				}
+			}
 			if (SerializedTypes.IsEmpty())
 			{
 				SerializedTypes = LexToString(EPropertyObjectReferenceType::None);
diff --git a/Engine/Source/Runtime/CoreUObject/Public/UObject/ObjectMacros.h b/Engine/Source/Runtime/CoreUObject/Public/UObject/ObjectMacros.h
--- a/Engine/Source/Runtime/CoreUObject/Public/UObject/ObjectMacros.h
+++ b/Engine/Source/Runtime/CoreUObject/Public/UObject/ObjectMacros.h
@@ -159,6 +159,10 @@ enum EPackageFlags
 	PKG_RuntimeGenerated			= 0x20000000,	///< This package contains elements that are runtime generated, and may not follow standard loading order rules
 	PKG_ReloadingForCooker			= 0x40000000,   ///< This package is reloading in the cooker, try to avoid getting data we will never need. We won't save this package.
 	PKG_FilterEditorOnly			= 0x80000000,	///< Package has editor-only data filtered out
+
+	// !!!!!!!!!!!!!!!!!!!!!!!!!!!!!!!!!!!!!!!!!!!!!!!!!!!!!!!!!!!!!!!!!!!!!!!!!!!!!!!!!!!!!!!!!!!!!!!!!!!!!!!!!!!!!!!!!!!!!!!!!!!!
+	// DO NOT ADD composite flags to EPackageFlags. Debugger visualisations have trouble displaying composite flag values as text.
+	// !!!!!!!!!!!!!!!!!!!!!!!!!!!!!!!!!!!!!!!!!!!!!!!!!!!!!!!!!!!!!!!!!!!!!!!!!!!!!!!!!!!!!!!!!!!!!!!!!!!!!!!!!!!!!!!!!!!!!!!!!!!!
 };
 
 // Transient Flags are cleared when serializing to or from PackageFileSummary
@@ -272,6 +276,10 @@ enum EClassFlags
 	CLASS_ConfigDoNotCheckDefaults = 0x40000000u,
 	/** Class has been consigned to oblivion as part of a blueprint recompile, and a newer version currently exists. */
 	CLASS_NewerVersionExists  = 0x80000000u,
+
+	// !!!!!!!!!!!!!!!!!!!!!!!!!!!!!!!!!!!!!!!!!!!!!!!!!!!!!!!!!!!!!!!!!!!!!!!!!!!!!!!!!!!!!!!!!!!!!!!!!!!!!!!!!!!!!!!!!!!!!!!!!!
+	// DO NOT ADD composite flags to EClassFlags. Debugger visualisations have trouble displaying composite flag values as text.
+	// !!!!!!!!!!!!!!!!!!!!!!!!!!!!!!!!!!!!!!!!!!!!!!!!!!!!!!!!!!!!!!!!!!!!!!!!!!!!!!!!!!!!!!!!!!!!!!!!!!!!!!!!!!!!!!!!!!!!!!!!!!
 };
 
 // Declare bitwise operators to allow EClassFlags to be combined but still retain type safety
@@ -397,6 +405,10 @@ enum EClassCastFlags : uint64
 	CASTCLASS_FAnsiStrProperty				= 0x2000000000000000,
 	CASTCLASS_FVCellProperty				= 0x4000000000000000,
 	//										= 0x8000000000000000,
+
+	// !!!!!!!!!!!!!!!!!!!!!!!!!!!!!!!!!!!!!!!!!!!!!!!!!!!!!!!!!!!!!!!!!!!!!!!!!!!!!!!!!!!!!!!!!!!!!!!!!!!!!!!!!!!!!!!!!!!!!!!!!!!!!!
+	// DO NOT ADD composite flags to EClassCastFlags. Debugger visualisations have trouble displaying composite flag values as text.
+	// !!!!!!!!!!!!!!!!!!!!!!!!!!!!!!!!!!!!!!!!!!!!!!!!!!!!!!!!!!!!!!!!!!!!!!!!!!!!!!!!!!!!!!!!!!!!!!!!!!!!!!!!!!!!!!!!!!!!!!!!!!!!!!
 };
 
 #define CASTCLASS_AllFlags ((EClassCastFlags)0xFFFFFFFFFFFFFFFF)
@@ -477,6 +489,10 @@ enum EPropertyFlags : uint64
 	CPF_ExperimentalAlwaysOverriden		= 0x0400000000000000,	///< ****Experimental*** Property should never inherit from the parent when using overridable serialization
 	CPF_ExperimentalNeverOverriden		= 0x0800000000000000,	///< ****Experimental*** Property should never be overridden when using overridable serialization
 	CPF_AllowSelfReference				= 0x1000000000000000,	///< Enables the instancing graph self referencing logic, delegates and verse function are already using this
+
+	// !!!!!!!!!!!!!!!!!!!!!!!!!!!!!!!!!!!!!!!!!!!!!!!!!!!!!!!!!!!!!!!!!!!!!!!!!!!!!!!!!!!!!!!!!!!!!!!!!!!!!!!!!!!!!!!!!!!!!!!!!!!!!
+	// DO NOT ADD composite flags to EPropertyFlags. Debugger visualisations have trouble displaying composite flag values as text.
+	// !!!!!!!!!!!!!!!!!!!!!!!!!!!!!!!!!!!!!!!!!!!!!!!!!!!!!!!!!!!!!!!!!!!!!!!!!!!!!!!!!!!!!!!!!!!!!!!!!!!!!!!!!!!!!!!!!!!!!!!!!!!!!
 };
 
 /** All Native Access Specifier flags */
@@ -514,6 +530,10 @@ enum class EArrayPropertyFlags : uint8
 {
 	None,
 	UsesMemoryImageAllocator
+
+	// !!!!!!!!!!!!!!!!!!!!!!!!!!!!!!!!!!!!!!!!!!!!!!!!!!!!!!!!!!!!!!!!!!!!!!!!!!!!!!!!!!!!!!!!!!!!!!!!!!!!!!!!!!!!!!!!!!!!!!!!!!!!!!!!!!
+	// DO NOT ADD composite flags to EArrayPropertyFlags. Debugger visualisations have trouble displaying composite flag values as text.
+	// !!!!!!!!!!!!!!!!!!!!!!!!!!!!!!!!!!!!!!!!!!!!!!!!!!!!!!!!!!!!!!!!!!!!!!!!!!!!!!!!!!!!!!!!!!!!!!!!!!!!!!!!!!!!!!!!!!!!!!!!!!!!!!!!!!
 };
 
 ENUM_CLASS_FLAGS(EArrayPropertyFlags)
@@ -525,6 +545,10 @@ enum class EMapPropertyFlags : uint8
 {
 	None,
 	UsesMemoryImageAllocator
+
+	// !!!!!!!!!!!!!!!!!!!!!!!!!!!!!!!!!!!!!!!!!!!!!!!!!!!!!!!!!!!!!!!!!!!!!!!!!!!!!!!!!!!!!!!!!!!!!!!!!!!!!!!!!!!!!!!!!!!!!!!!!!!!!!!!
+	// DO NOT ADD composite flags to EMapPropertyFlags. Debugger visualisations have trouble displaying composite flag values as text.
+	// !!!!!!!!!!!!!!!!!!!!!!!!!!!!!!!!!!!!!!!!!!!!!!!!!!!!!!!!!!!!!!!!!!!!!!!!!!!!!!!!!!!!!!!!!!!!!!!!!!!!!!!!!!!!!!!!!!!!!!!!!!!!!!!!
 };
 
 ENUM_CLASS_FLAGS(EMapPropertyFlags)
@@ -536,12 +560,15 @@ enum class EPropertyObjectReferenceType : uint32
 	Weak = 1 << 1,   // Weak reference to a UObject, does not keep the object from being garbage collected, does not become valid again after object is unloaded and reloaded.
 	Soft = 1 << 2,	 // Soft path/identity-based reference to a UObject, does not keep the object from being garbage collected, does become valid again after object is unloaded and reloaded.
 	Conservative = 1 << 3,	 // Not a real reference type, used to mark native struct serializers which may serialize unknown reference types and to conservatively populate RefLink with such struct properties.
-	
-	MAX = Conservative << 1, // Marker for iterating over all flags
-	Any = Strong | Weak | Soft | Conservative,
+
+	// !!!!!!!!!!!!!!!!!!!!!!!!!!!!!!!!!!!!!!!!!!!!!!!!!!!!!!!!!!!!!!!!!!!!!!!!!!!!!!!!!!!!!!!!!!!!!!!!!!!!!!!!!!!!!!!!!!!!!!!!!!!!!!!!!!!!!!!!!!!
+	// DO NOT ADD composite flags to EPropertyObjectReferenceType. Debugger visualisations have trouble displaying composite flag values as text.
+	// !!!!!!!!!!!!!!!!!!!!!!!!!!!!!!!!!!!!!!!!!!!!!!!!!!!!!!!!!!!!!!!!!!!!!!!!!!!!!!!!!!!!!!!!!!!!!!!!!!!!!!!!!!!!!!!!!!!!!!!!!!!!!!!!!!!!!!!!!!!
 };
 ENUM_CLASS_FLAGS(EPropertyObjectReferenceType);
 
+inline constexpr EPropertyObjectReferenceType EPropertyObjectReferenceType_Any = EPropertyObjectReferenceType::Strong | EPropertyObjectReferenceType::Weak | EPropertyObjectReferenceType::Soft | EPropertyObjectReferenceType::Conservative;
+
 COREUOBJECT_API const TCHAR* LexToString(EPropertyObjectReferenceType Type);
 
 /**
@@ -597,6 +624,10 @@ enum EObjectFlags
 	// while checking the object flags is much faster outside of it where the Object pointer is already available and most likely cached.
 	RF_MirroredGarbage			=0x40000000,	///< Garbage from logical point of view and should not be referenced. This flag is mirrored in EInternalObjectFlags as Garbage for performance
 	RF_AllocatedInSharedPage	=0x80000000,	///< Allocated from a ref-counted page shared with other UObjects
+
+	// !!!!!!!!!!!!!!!!!!!!!!!!!!!!!!!!!!!!!!!!!!!!!!!!!!!!!!!!!!!!!!!!!!!!!!!!!!!!!!!!!!!!!!!!!!!!!!!!!!!!!!!!!!!!!!!!!!!!!!!!!!!
+	// DO NOT ADD composite flags to EObjectFlags. Debugger visualisations have trouble displaying composite flag values as text.
+	// !!!!!!!!!!!!!!!!!!!!!!!!!!!!!!!!!!!!!!!!!!!!!!!!!!!!!!!!!!!!!!!!!!!!!!!!!!!!!!!!!!!!!!!!!!!!!!!!!!!!!!!!!!!!!!!!!!!!!!!!!!!
 };
 
 UE_DEPRECATED(5.4, "RF_InternalPendingKill should no longer be used. Use IsValid(Object) instead.")
@@ -692,6 +723,10 @@ enum class EEnumFlags : uint8
 
 	Flags = 0x00000001, // Whether the UEnum represents a set of flags
 	NewerVersionExists = 0x00000002, // If set, this UEnum has been replaced by a newer version
+
+	// !!!!!!!!!!!!!!!!!!!!!!!!!!!!!!!!!!!!!!!!!!!!!!!!!!!!!!!!!!!!!!!!!!!!!!!!!!!!!!!!!!!!!!!!!!!!!!!!!!!!!!!!!!!!!!!!!!!!!!!!!
+	// DO NOT ADD composite flags to EEnumFlags. Debugger visualisations have trouble displaying composite flag values as text.
+	// !!!!!!!!!!!!!!!!!!!!!!!!!!!!!!!!!!!!!!!!!!!!!!!!!!!!!!!!!!!!!!!!!!!!!!!!!!!!!!!!!!!!!!!!!!!!!!!!!!!!!!!!!!!!!!!!!!!!!!!!!
 };
 
 ENUM_CLASS_FLAGS(EEnumFlags)
