Added UEOps to FObjectHandle.
Added a workaround for a bug on MSVC which caused a failure in ObjectHandleTest.cpp.
Added a new macro UE_OPS_NAMESPACE_VISIBLE to allow UEOps-using types in namespaces to work when their operators are called from other namespaces.

diff --git a/Engine/Source/Runtime/Core/Public/Misc/UEOps.h b/Engine/Source/Runtime/Core/Public/Misc/UEOps.h
--- a/Engine/Source/Runtime/Core/Public/Misc/UEOps.h
+++ b/Engine/Source/Runtime/Core/Public/Misc/UEOps.h
@@ -5,10 +5,47 @@
 #include "Concepts/SameAs.h"
 #include "HAL/Platform.h"
 
+/**
+ * This file defines global comparison operators for all types containing certain named member functions.
+ * These ensure that the operators and their variants are defined correctly based on a single customization
+ * point and should result in faster overload resolution.
+ *
+ * bool FMyType::UEOpEquals(const FMyType&) const;
+ * bool FMyType::UEOpEquals(const OtherType&) const;
+ *   Should be defined to opt into operator== and operator!=.
+ *
+ * bool FMyType::UEOpLessThan(const FMyType&) const;
+ * bool FMyType::UEOpLessThan(const OtherType&) const;
+ * bool FMyType::UEOpGreaterThan(const OtherType&) const;
+ *   Should be defined to opt into operator<, operator<=, operator> and operator>=.
+ *
+ * bool FMyType::UEOpGreaterThan(const FMyType&) const;
+ *   Will never be called if FMyType::UEOpLessThan is defined, so is redundant.
+ */
+
+
+/**
+ * On MSVC, operator!= isn't synthesized from operator== in some contexts.  This was found in a case of FObjectHandle:
+ *
+ * FObjectHandle NullObjectHandle = ...;
+ * FObjectHandle TestLateResolveUnsafeObjectHandle = ...;
+ * ...
+ * CHECK(NullObjectHandle != TestLateResolveUnsafeObjectHandle); // catch2\internal\catch_decomposer.hpp(219,9): error C2678: binary '!=': no operator found which takes a left-hand operand of type 'FObjectHandle' (or there is no acceptable conversion)
+ *
+ * Additional overloads are added to work around this issue - they'll just make overload resolution errors slightly worse.  When the following ticket is fixed, these overloads can be removed.
+ *
+ * https://developercommunity.visualstudio.com/t/Synthesised-operator-not-considered-in/10945932
+ */
+#if defined(_MSC_VER) && !defined(__clang__)
+	#define UE_OPS_WORKAROUND_PRIVATE 1
+#else
+	#define UE_OPS_WORKAROUND_PRIVATE 0
+#endif
+
 namespace UE::Core::Private
 {
 	// This only exists to bind to NULL, which is not perfectly-forwardable but which is commonly used in comparison operators.
-	// The type is unique so that it can't be confused with a type that will actually be used by real user code, and it used to
+	// The type is unique so that it can't be confused with a type that will actually be used by real user code, and is used to
 	// create a pointer-to-member type which is highly unlikely to be deduced.
 	struct FIncomplete;
 
@@ -31,24 +68,6 @@ namespace UE::Core::Private
 	};
 }
 
-/**
- * This file defines global comparison operators for all types containing certain named member functions.
- * These ensure that the operators and their variants are defined correctly based on a single customization
- * point and should result in faster overload resolution.
- *
- * bool FMyType::UEOpEquals(const FMyType&) const;
- * bool FMyType::UEOpEquals(const OtherType&) const;
- *   Should be defined to opt into operator== and operator!=.
- *
- * bool FMyType::UEOpLessThan(const FMyType&) const;
- * bool FMyType::UEOpLessThan(const OtherType&) const;
- * bool FMyType::UEOpGreaterThan(const OtherType&) const;
- *   Should be defined to opt into operator<, operator<=, operator> and operator>=.
- *
- * bool FMyType::UEOpGreaterThan(const FMyType&) const;
- *   Will never be called if FMyType::UEOpLessThan is defined, so is redundant.
- */
-
 template <typename LhsType, typename RhsType>
 	requires UE::Core::Private::CWithUEOpEquals<LhsType, RhsType> || UE::Core::Private::CWithUEOpEquals<RhsType, LhsType>
 [[nodiscard]] UE_REWRITE constexpr bool operator==(const LhsType& Lhs, const RhsType& Rhs)
@@ -71,6 +90,36 @@ template <typename LhsType>
 	return Lhs.UEOpEquals(nullptr);
 }
 
+#if UE_OPS_WORKAROUND_PRIVATE
+	template <typename LhsType, typename RhsType>
+		requires UE::Core::Private::CWithUEOpEquals<LhsType, RhsType> || UE::Core::Private::CWithUEOpEquals<RhsType, LhsType>
+	[[nodiscard]] UE_REWRITE constexpr bool operator!=(const LhsType& Lhs, const RhsType& Rhs)
+	{
+		if constexpr (UE::Core::Private::CWithUEOpEquals<LhsType, RhsType>)
+		{
+			return !Lhs.UEOpEquals(Rhs);
+		}
+		else
+		{
+			return !Rhs.UEOpEquals(Lhs);
+		}
+	}
+
+	template <typename LhsType>
+		requires UE::Core::Private::CWithUEOpEquals<LhsType, decltype(nullptr)>
+	[[nodiscard]] UE_REWRITE constexpr bool operator!=(const LhsType& Lhs, int UE::Core::Private::FIncomplete::* Rhs)
+	{
+		return !Lhs.UEOpEquals(nullptr);
+	}
+
+	template <typename RhsType>
+		requires UE::Core::Private::CWithUEOpEquals<RhsType, decltype(nullptr)>
+	[[nodiscard]] UE_REWRITE constexpr bool operator!=(int UE::Core::Private::FIncomplete::* Lhs, const RhsType& Rhs)
+	{
+		return !Rhs.UEOpEquals(nullptr);
+	}
+#endif
+
 template <typename LhsType, typename RhsType>
 	requires UE::Core::Private::CWithUEOpLessThan<LhsType, RhsType> || UE::Core::Private::CWithUEOpGreaterThan<RhsType, LhsType>
 [[nodiscard]] UE_REWRITE constexpr bool operator<(const LhsType& Lhs, const RhsType& Rhs)
@@ -127,6 +176,52 @@ template <typename LhsType, typename RhsType>
 	}
 }
 
+// When a type in a namespace uses UEOps functions, comparisons can fail when calling the operators from an
+// unrelated namespace, where the global operators are hidden and argument-dependent lookup cannot find the
+// operators because they are not in the same namespace as the type.
+//
+// UE_OPS_NAMESPACE_VISIBLE makes the operators visible in that type's namespace, allowing ADL to work as
+// intended. Example:
+//
+// namespace N1
+// {
+//     struct FMyType
+//     {
+//         bool UEOpEquals(const FMyType& Other) const;
+//     };
+//
+//     UE_OPS_NAMESPACE_VISIBLE
+// }
+//
+// namespace N2
+// {
+//     struct FOtherType;
+//     bool operator==(const FOtherType& Lhs, const FOtherType& Rhs);
+//
+//     bool AreDifferent(const N1::FMyType& Lhs, const N1::FMyType& Rhs)
+//     {
+//         return Lhs != Rhs; // without UE_OPS_NAMESPACE_VISIBLE above, this will not compile.
+//     }
+// }
+#if UE_OPS_WORKAROUND_PRIVATE
+	#define UE_OPS_NAMESPACE_VISIBLE \
+		using ::operator==; \
+		using ::operator!=; \
+		using ::operator<; \
+		using ::operator>; \
+		using ::operator<=; \
+		using ::operator>=;
+#else
+	#define UE_OPS_NAMESPACE_VISIBLE \
+		using ::operator==; \
+		using ::operator<; \
+		using ::operator>; \
+		using ::operator<=; \
+		using ::operator>=;
+#endif
+
+#undef UE_OPS_WORKAROUND_PRIVATE
+
 #if UE_ENABLE_INCLUDE_ORDER_DEPRECATED_IN_5_7
 #include "Templates/Requires.h"
 #endif
diff --git a/Engine/Source/Runtime/CoreUObject/Public/UObject/ObjectHandle.h b/Engine/Source/Runtime/CoreUObject/Public/UObject/ObjectHandle.h
--- a/Engine/Source/Runtime/CoreUObject/Public/UObject/ObjectHandle.h
+++ b/Engine/Source/Runtime/CoreUObject/Public/UObject/ObjectHandle.h
@@ -8,6 +8,7 @@
 #include "HAL/CriticalSection.h"
 #include "HAL/Platform.h"
 #include "Misc/AssertionMacros.h"
+#include "Misc/UEOps.h"
 #include "Templates/TypeHash.h"
 #include "UObject/NameTypes.h"
 #include "UObject/ObjectHandleTracking.h"
@@ -77,51 +78,56 @@ namespace UE::CoreUObject::Private
 			return PointerOrRef != 0;
 		}
 
-	// When using TSAN, PointerOrRef has non-trivial copy construction so we need to explicitly implement constructors to handle the union.
-	// This means this type can't use aggregate initialization even when TSAN is disabled
-	FObjectHandlePrivate()
-	: PointerOrRef(0)
-	{
-	}
-	~FObjectHandlePrivate() = default;
-	[[nodiscard]] constexpr explicit FObjectHandlePrivate(const void* InPointer)
-		: Pointer(InPointer)
-	{
-	}
-	[[nodiscard]] explicit FObjectHandlePrivate(UPTRINT Packed)
-		: PointerOrRef(Packed)
-	{
-	}
-#if USING_THREAD_SANITISER || USING_INSTRUMENTATION
-	[[nodiscard]] constexpr FObjectHandlePrivate(const FObjectHandlePrivate& Other)
-	{
-		UE_IF_CONSTEVAL
+		// When using TSAN, PointerOrRef has non-trivial copy construction so we need to explicitly implement constructors to handle the union.
+		// This means this type can't use aggregate initialization even when TSAN is disabled
+		FObjectHandlePrivate()
+		: PointerOrRef(0)
 		{
-			Pointer = Other.Pointer;
 		}
-		else
+		~FObjectHandlePrivate() = default;
+		[[nodiscard]] constexpr explicit FObjectHandlePrivate(const void* InPointer)
+			: Pointer(InPointer)
 		{
-			PointerOrRef = Other.PointerOrRef;
 		}
-	}
-	constexpr FObjectHandlePrivate& operator=(const FObjectHandlePrivate& Other)
-	{
-		UE_IF_CONSTEVAL
+		[[nodiscard]] explicit FObjectHandlePrivate(UPTRINT Packed)
+			: PointerOrRef(Packed)
 		{
-			Pointer = Other.Pointer;
 		}
-		else
+#if USING_THREAD_SANITISER || USING_INSTRUMENTATION
+		[[nodiscard]] constexpr FObjectHandlePrivate(const FObjectHandlePrivate& Other)
+		{
+			UE_IF_CONSTEVAL
+			{
+				Pointer = Other.Pointer;
+			}
+			else
+			{
+				PointerOrRef = Other.PointerOrRef;
+			}
+		}
+		constexpr FObjectHandlePrivate& operator=(const FObjectHandlePrivate& Other)
 		{
-			PointerOrRef = Other.PointerOrRef;
+			UE_IF_CONSTEVAL
+			{
+				Pointer = Other.Pointer;
+			}
+			else
+			{
+				PointerOrRef = Other.PointerOrRef;
+			}
+			return *this;
 		}
-		return *this;
-	}
 #else // USING_THREAD_SANITISER || USING_INSTRUMENTATION
-	constexpr FObjectHandlePrivate(const FObjectHandlePrivate&) = default;
-	constexpr FObjectHandlePrivate& operator=(const FObjectHandlePrivate&) = default;
+		constexpr FObjectHandlePrivate(const FObjectHandlePrivate&) = default;
+		constexpr FObjectHandlePrivate& operator=(const FObjectHandlePrivate&) = default;
 #endif // USING_THREAD_SANITISER || USING_INSTRUMENTATION
+
+		bool UEOpEquals(FObjectHandlePrivate Rhs) const;
 	};
 
+	// Allow generated operators to be accessible from comparisons in other namespaces.
+	UE_OPS_NAMESPACE_VISIBLE
+
 	/* Returns the packed object ref for this object IF one exists otherwise returns a null PackedObjectRef */
 	COREUOBJECT_API FPackedObjectRef FindExistingPackedObjectRef(const UObject* Object);
 
@@ -160,8 +166,13 @@ namespace UE::CoreUObject::Private
 
 		COREUOBJECT_API static FRemoteObjectHandlePrivate ConvertToRemoteHandle(UObject* Object);
 		COREUOBJECT_API static FRemoteObjectHandlePrivate FromIdNoResolve(FRemoteObjectId ObjectId);
+
+		bool UEOpEquals(FRemoteObjectHandlePrivate Rhs) const;
 	};
 
+	// Allow generated operators to be accessible from comparisons in other namespaces.
+	UE_OPS_NAMESPACE_VISIBLE
+
 #endif
 
 	///these functions are always defined regardless of UE_WITH_OBJECT_HANDLE_LATE_RESOLVE value
@@ -254,24 +265,24 @@ inline bool IsObjectHandleTypeSafe(FObjectHandle Handle)
 
 /////////////////////////////////////////////////////////////////////////////////////////////////////////////////////////
 #if UE_WITH_OBJECT_HANDLE_LATE_RESOLVE
-inline bool operator==(UE::CoreUObject::Private::FObjectHandlePrivate LHS, UE::CoreUObject::Private::FObjectHandlePrivate RHS)
+inline bool UE::CoreUObject::Private::FObjectHandlePrivate::UEOpEquals(FObjectHandlePrivate RHS) const
 {
 	using namespace UE::CoreUObject::Private;
 
-	bool LhsResolved = IsObjectHandleResolved(LHS);
+	bool LhsResolved = IsObjectHandleResolved(*this);
 	bool RhsResolved = IsObjectHandleResolved(RHS);
 
 	//if both resolved or both unresolved compare the uintptr
 	if (LhsResolved == RhsResolved)
 	{
-		return LHS.PointerOrRef == RHS.PointerOrRef;
+		return PointerOrRef == RHS.PointerOrRef;
 	}
 
 	//only one side can be resolved
 	if (LhsResolved)
 	{
 		//both sides can't be null as resolved status would have be true for both
-		const UObject* Obj = ReadObjectHandlePointerNoCheck(LHS);
+		const UObject* Obj = ReadObjectHandlePointerNoCheck(*this);
 		if (!Obj)
 		{
 			return false;
@@ -301,16 +312,11 @@ inline bool operator==(UE::CoreUObject::Private::FObjectHandlePrivate LHS, UE::C
 		{
 			return false;
 		}
-		return PackedRhs.EncodedRef == LHS.PointerOrRef;
+		return PackedRhs.EncodedRef == PointerOrRef;
 	}
 
 }
 
-inline bool operator!=(UE::CoreUObject::Private::FObjectHandlePrivate LHS, UE::CoreUObject::Private::FObjectHandlePrivate RHS)
-{
-	return !(LHS == RHS);
-}
-
 inline uint32 GetTypeHash(UE::CoreUObject::Private::FObjectHandlePrivate Handle)
 {
 	using namespace UE::CoreUObject::Private;
@@ -334,25 +340,25 @@ inline uint32 GetTypeHash(UE::CoreUObject::Private::FObjectHandlePrivate Handle)
 	return GetTypeHash(Handle.GetPointerOrRef());
 }
 #elif UE_WITH_REMOTE_OBJECT_HANDLE
-inline bool operator==(UE::CoreUObject::Private::FRemoteObjectHandlePrivate LHS, UE::CoreUObject::Private::FRemoteObjectHandlePrivate RHS)
+inline bool UE::CoreUObject::Private::FRemoteObjectHandlePrivate::UEOpEquals(FRemoteObjectHandlePrivate RHS) const
 {
 	using namespace UE::CoreUObject::Private;
 	using namespace UE::RemoteObject::Handle;
 
-	const bool bLhsResolved = IsObjectHandleResolved_ForGC(LHS);
+	const bool bLhsResolved = IsObjectHandleResolved_ForGC(*this);
 	const bool bRhsResolved = IsObjectHandleResolved_ForGC(RHS);
 
 	//if both resolved or both unresolved compare the uintptr
 	if (bLhsResolved == bRhsResolved)
 	{
-		return LHS.PointerOrHandle == RHS.PointerOrHandle;
+		return PointerOrHandle == RHS.PointerOrHandle;
 	}
 
 	//only one side can be resolved
 	if (bLhsResolved)
 	{
 		//both sides can't be null as resolved status would have be true for both
-		const UObject* Obj = ReadObjectHandlePointerNoCheck(LHS);
+		const UObject* Obj = ReadObjectHandlePointerNoCheck(*this);
 		if (!Obj)
 		{
 			return false;
@@ -370,17 +376,12 @@ inline bool operator==(UE::CoreUObject::Private::FRemoteObjectHandlePrivate LHS,
 			return false;
 		}
 
-		const FRemoteObjectStub* LhsStub = LHS.ToStub();
+		const FRemoteObjectStub* LhsStub = ToStub();
 		return LhsStub->Id == FRemoteObjectId((const UObjectBase*)Obj);
 	}
 
 }
 
-inline bool operator!=(UE::CoreUObject::Private::FRemoteObjectHandlePrivate LHS, UE::CoreUObject::Private::FRemoteObjectHandlePrivate RHS)
-{
-	return !(LHS == RHS);
-}
-
 inline FRemoteObjectId GetRemoteObjectId(UE::CoreUObject::Private::FRemoteObjectHandlePrivate Handle)
 {
 	using namespace UE::CoreUObject::Private;
