[UE] Move non-nullable unset optional LLTs to non-native subobject instancing tests.

Notes:
- TOptional<TNonNullPtr<T>> is not currently supported as a native reference property and will assert if serialized in a debug build config.

#rnx

diff --git a/Engine/Source/Runtime/CoreUObject/Tests/UObject/InstancedPropertiesTest.cpp b/Engine/Source/Runtime/CoreUObject/Tests/UObject/InstancedPropertiesTest.cpp
--- a/Engine/Source/Runtime/CoreUObject/Tests/UObject/InstancedPropertiesTest.cpp
+++ b/Engine/Source/Runtime/CoreUObject/Tests/UObject/InstancedPropertiesTest.cpp
@@ -458,7 +458,7 @@ TEST_CASE_NAMED(FInstancedPropertiesTest, "CoreUObject::InstancedProperties", "[
 #endif
 	}
 
-	SECTION("Optional instanced object property - unset semantics (TObjectPtr)")
+	SECTION("Optional instanced object property - unset semantics")
 	{
 		const FName OptionalNullableInstancedObjectPropertyName = GET_MEMBER_NAME_CHECKED(USubobjectInstancingTestOuterObject, OptionalNullableInnerObjectUnset);
 		const FOptionalProperty* OptionalNullableInstancedObjectProperty = CastField<FOptionalProperty>(USubobjectInstancingTestOuterObject::StaticClass()->FindPropertyByName(OptionalNullableInstancedObjectPropertyName));
@@ -473,22 +473,6 @@ TEST_CASE_NAMED(FInstancedPropertiesTest, "CoreUObject::InstancedProperties", "[
 		CHECK_FALSE(ValueProperty->HasAnyPropertyFlags(CPF_NonNullable));
 		CHECK(ValueProperty->HasAnyPropertyFlags(CPF_TObjectPtrWrapper));
 	}
-
-	SECTION("Optional instanced object property - unset semantics (TNonNullPtr)")
-	{
-		const FName OptionalNonNullableInstancedObjectPropertyName = GET_MEMBER_NAME_CHECKED(USubobjectInstancingTestOuterObject, OptionalNonNullableInnerObjectUnset);
-		const FOptionalProperty* OptionalNonNullableInstancedObjectProperty = CastField<FOptionalProperty>(USubobjectInstancingTestOuterObject::StaticClass()->FindPropertyByName(OptionalNonNullableInstancedObjectPropertyName));
-
-		REQUIRE(OptionalNonNullableInstancedObjectProperty != nullptr);
-
-		const FProperty* ValueProperty = OptionalNonNullableInstancedObjectProperty->GetValueProperty();
-		REQUIRE(ValueProperty != nullptr);
-
-		CHECK(ValueProperty->IsA<FObjectProperty>());
-		CHECK(ValueProperty->HasIntrusiveUnsetOptionalState());
-		CHECK(ValueProperty->HasAnyPropertyFlags(CPF_NonNullable));
-		CHECK_FALSE(ValueProperty->HasAnyPropertyFlags(CPF_TObjectPtrWrapper));
-	}
 }
 
 }
diff --git a/Engine/Source/Runtime/CoreUObject/Tests/UObject/SubobjectInstancingTest.cpp b/Engine/Source/Runtime/CoreUObject/Tests/UObject/SubobjectInstancingTest.cpp
--- a/Engine/Source/Runtime/CoreUObject/Tests/UObject/SubobjectInstancingTest.cpp
+++ b/Engine/Source/Runtime/CoreUObject/Tests/UObject/SubobjectInstancingTest.cpp
@@ -150,6 +150,8 @@ namespace SubobjectInstancingTest::Private
 	static const FName NonNativeInnerObjectPropertyName(TEXT("NonNativeInnerObject"));
 	static const FName NonNativeOuterObjectPropertyName(TEXT("NonNativeOwnerObject"));
 	static const FName NonNativeEditTimeInnerObjectPropertyName(TEXT("NonNativeEditTimeInnerObject"));
+	static const FName NonNullableInnerObjectPropertyName(TEXT("NonNullableInnerObject"));
+	static const FName NonNullableOptionalInnerObjectPropertyName(TEXT("OptionalNonNullableInnerObject"));
 
 	template<typename ClassType>
 	static UClass* NewTestClass(UClass* SuperClass)
@@ -207,6 +209,23 @@ namespace SubobjectInstancingTest::Private
 		NonNativeEditTimeInnerObjectProperty->SetPropertyClass(USubobjectInstancingTestObject::StaticClass());
 		PropertyListBuilder.AppendNoTerminate(*NonNativeEditTimeInnerObjectProperty);
 
+		// Non-nullable object reference property to store a non-native instantiation by ctor - TNonNullPtr<USubobjectInstancingTestObject>
+		EPropertyFlags NonNullableObjectPropertyFlags = (NonNativeObjectPropertyFlags & ~CPF_TObjectPtrWrapper) | CPF_NonNullable;
+		FObjectProperty* NonNullableInnerObjectProperty = new FObjectProperty(TestClass, NonNullableInnerObjectPropertyName, RF_Public | RF_Transient);
+		NonNullableInnerObjectProperty->SetPropertyFlags(NonNullableObjectPropertyFlags);
+		NonNullableInnerObjectProperty->SetPropertyClass(USubobjectInstancingTestObject::StaticClass());
+		PropertyListBuilder.AppendNoTerminate(*NonNullableInnerObjectProperty);
+
+		// Non-nullable optional object reference property to store an optional non-null instantiation by ctor - TOptional<TNonNullPtr<USubobjectinstancingTestObject>>
+		// Note: TOptional<TNonNullPtr<T>> is not currently supported as a native property, so this is being used to validate intrusive state for non-nullable references.
+		FOptionalProperty* NonNullableOptionalInnerObjectProperty = new FOptionalProperty(TestClass, NonNullableOptionalInnerObjectPropertyName, RF_Public | RF_Transient);
+		FObjectProperty* NonNullableOptionalInnerObjectValueProperty = new FObjectProperty(NonNullableOptionalInnerObjectProperty, NonNullableOptionalInnerObjectPropertyName, RF_Public | RF_Transient);
+		NonNullableOptionalInnerObjectProperty->SetPropertyFlags(CPF_ContainsInstancedReference);
+		NonNullableOptionalInnerObjectValueProperty->SetPropertyFlags(NonNullableObjectPropertyFlags);
+		NonNullableOptionalInnerObjectValueProperty->SetPropertyClass(USubobjectInstancingTestObject::StaticClass());
+		NonNullableOptionalInnerObjectProperty->SetValueProperty(NonNullableOptionalInnerObjectValueProperty);
+		PropertyListBuilder.AppendNoTerminate(*NonNullableOptionalInnerObjectProperty);
+
 		TestClass->SetSuperStruct(SuperClass);
 		TestClass->Bind();
 		TestClass->StaticLink(/*bRelinkExistingProperties =*/ true);
@@ -230,6 +249,11 @@ namespace SubobjectInstancingTest::Private
 		// Simulate a non-native class constructor initialization of the edit-time reference property's default value
 		NonNativeEditTimeInnerObjectProperty->SetObjectPropertyValue_InContainer(TestCDO, nullptr);
 
+		// Simulate a non-native clsas constructor initialization of a non-nullable subobject reference, also "captured" by the CDO and used to initialize all new objects of this type
+		static const FName NonNullableInnerObjectName(TEXT("NonNullableInnerObject"));
+		UObject* NonNullableInnerDefaultSubobject = NewObject<USubobjectInstancingTestObject>(TestCDO, NonNullableInnerObjectName, RF_ArchetypeObject | TestCDO->GetMaskedFlags(RF_PropagateToSubObjects));
+		NonNullableInnerObjectProperty->SetObjectPropertyValue_InContainer(TestCDO, NonNullableInnerDefaultSubobject);
+
 		return TestClass;
 	}
 
@@ -273,6 +297,10 @@ namespace SubobjectInstancingTest::Private
 		SECTION("Default subobject initialized from template")
 		{
 			USubobjectInstancingTestOuterObject* Template = NewObject<USubobjectInstancingTestOuterObject>(GetTransientPackage(), InstancingTestClass, NAME_None, RF_ArchetypeObject);
+
+			REQUIRE(Template->InnerObject != nullptr);
+			CHECK(Template->InnerObject->IsInOuter(Template));
+
 			Template->InnerObject->TestValue = 200;
 
 			USubobjectInstancingTestOuterObject* OuterObject = NewObject<USubobjectInstancingTestOuterObject>(GetTransientPackage(), InstancingTestClass, NAME_None, RF_NoFlags, Template);
@@ -1200,9 +1228,6 @@ namespace SubobjectInstancingTest::Private
 
 			CHECK_FALSE(OuterObject->OptionalNullableInnerObjectUnset.IsSet());
 			CHECK(OuterObject->OptionalNullableInnerObjectUnset.GetPtrOrNull() == nullptr);
-
-			CHECK_FALSE(OuterObject->OptionalNonNullableInnerObjectUnset.IsSet());
-			CHECK(OuterObject->OptionalNonNullableInnerObjectUnset.GetPtrOrNull() == nullptr);
 		}
 
 		SECTION("Unset optional default subobject initialized from template")
@@ -1210,29 +1235,20 @@ namespace SubobjectInstancingTest::Private
 			USubobjectInstancingTestOuterObject* Template = NewObject<USubobjectInstancingTestOuterObject>(GetTransientPackage(), InstancingTestClass, NAME_None, RF_ArchetypeObject);
 
 			CHECK_FALSE(Template->OptionalNullableInnerObjectUnset.IsSet());
-			CHECK_FALSE(Template->OptionalNonNullableInnerObjectUnset.IsSet());
 
 			USubobjectInstancingTestOuterObject* OuterObject = NewObject<USubobjectInstancingTestOuterObject>(GetTransientPackage(), InstancingTestClass, NAME_None, RF_NoFlags, Template);
 
 			CHECK_FALSE(OuterObject->OptionalNullableInnerObjectUnset.IsSet());
-			CHECK_FALSE(OuterObject->OptionalNonNullableInnerObjectUnset.IsSet());
-
+	
 			Template->OptionalNullableInnerObjectUnset = nullptr;
-			Template->OptionalNonNullableInnerObjectUnset = NewObject<USubobjectInstancingTestObject>(Template, GET_MEMBER_NAME_CHECKED(USubobjectInstancingTestOuterObject, OptionalNonNullableInnerObjectUnset), Template->GetMaskedFlags(RF_PropagateToSubObjects));
 
 			CHECK(Template->OptionalNullableInnerObjectUnset.IsSet());
-			CHECK(Template->OptionalNonNullableInnerObjectUnset.IsSet());
-
 			CHECK(Template->OptionalNullableInnerObjectUnset.GetValue().Get() == nullptr);
-			REQUIRE(Template->OptionalNonNullableInnerObjectUnset.GetValue().Get() != nullptr);
 
 			OuterObject = NewObject<USubobjectInstancingTestOuterObject>(GetTransientPackage(), InstancingTestClass, NAME_None, RF_NoFlags, Template);
 
 			CHECK(OuterObject->OptionalNullableInnerObjectUnset.IsSet());
-			CHECK(OuterObject->OptionalNonNullableInnerObjectUnset.IsSet());
-
 			CHECK(OuterObject->OptionalNullableInnerObjectUnset.GetValue().Get() == nullptr);
-			REQUIRE(OuterObject->OptionalNonNullableInnerObjectUnset.GetValue().Get() != nullptr);
 		}
 	}
 
@@ -1711,6 +1727,101 @@ namespace SubobjectInstancingTest::Private
 		}
 	}
 
+	void RunSubobjectInstancingTests_NonNullable(UClass* InstancingTestClass)
+	{
+		// Non-nullable tests are only applicable to non-native properties right now.
+		if (InstancingTestClass->IsNative())
+		{
+			return;
+		}
+
+		const FObjectProperty* ObjectProperty = CastFieldChecked<FObjectProperty>(InstancingTestClass->FindPropertyByName(NonNullableInnerObjectPropertyName));
+		const FOptionalProperty* OptionalProperty = CastFieldChecked<FOptionalProperty>(InstancingTestClass->FindPropertyByName(NonNullableOptionalInnerObjectPropertyName));
+
+		SECTION("Non-nullable subobject reference property semantics")
+		{
+			CHECK(ObjectProperty->HasIntrusiveUnsetOptionalState());
+		}
+
+		SECTION("Non-nullable subobject reference initialized from CDO")
+		{
+			USubobjectInstancingTestOuterObject* OuterObject = NewObject<USubobjectInstancingTestOuterObject>(GetTransientPackage(), InstancingTestClass);
+			USubobjectInstancingTestObject* ObjectReference = Cast<USubobjectInstancingTestObject>(ObjectProperty->GetObjectPropertyValue_InContainer(OuterObject));
+
+			REQUIRE(ObjectReference != nullptr);
+			CHECK(ObjectReference->IsIn(OuterObject));
+			CHECK(ObjectReference->TestValue == 100);
+
+			CHECK_FALSE(ObjectReference->IsTemplate());
+			CHECK(ObjectReference->IsDefaultSubobject());
+			CHECK_FALSE(ObjectReference->IsTemplateForSubobjects());
+
+			USubobjectInstancingTestOuterObject* NewOuterObject = NewObject<USubobjectInstancingTestOuterObject>(GetTransientPackage(), InstancingTestClass);
+			CHECK(ObjectReference != ObjectProperty->GetObjectPropertyValue_InContainer(NewOuterObject));
+		}
+
+		SECTION("Non-nullable subobject reference initialized from template")
+		{
+			USubobjectInstancingTestOuterObject* Template = NewObject<USubobjectInstancingTestOuterObject>(GetTransientPackage(), InstancingTestClass, NAME_None, RF_ArchetypeObject);
+			{
+				USubobjectInstancingTestObject* ObjectReference = Cast<USubobjectInstancingTestObject>(ObjectProperty->GetObjectPropertyValue_InContainer(Template));
+
+				REQUIRE(ObjectReference != nullptr);
+				CHECK(ObjectReference->IsIn(Template));
+				CHECK(ObjectReference->TestValue == 100);
+
+				ObjectReference->TestValue = 200;
+			}
+
+			USubobjectInstancingTestOuterObject* OuterObject = NewObject<USubobjectInstancingTestOuterObject>(GetTransientPackage(), InstancingTestClass, NAME_None, RF_NoFlags, Template);
+			{
+				USubobjectInstancingTestObject* ObjectReference = Cast<USubobjectInstancingTestObject>(ObjectProperty->GetObjectPropertyValue_InContainer(OuterObject));
+
+				REQUIRE(ObjectReference != nullptr);
+				CHECK(ObjectReference->IsIn(OuterObject));
+				CHECK(ObjectReference->TestValue == 200);
+
+				CHECK_FALSE(ObjectReference->IsTemplate());
+				CHECK(ObjectReference->IsDefaultSubobject());
+				CHECK_FALSE(ObjectReference->IsTemplateForSubobjects());
+
+				USubobjectInstancingTestOuterObject* NewOuterObject = NewObject<USubobjectInstancingTestOuterObject>(GetTransientPackage(), InstancingTestClass, NAME_None, RF_NoFlags, Template);
+				CHECK(ObjectReference != ObjectProperty->GetObjectPropertyValue_InContainer(NewOuterObject));
+			}
+		}
+
+		SECTION("Non-nullable optional subobject reference property semantics")
+		{
+			CHECK_FALSE(OptionalProperty->HasIntrusiveUnsetOptionalState());
+
+			FProperty* ValueProperty = OptionalProperty->GetValueProperty();
+			{
+				REQUIRE(ValueProperty != nullptr);
+				CHECK(ValueProperty->HasIntrusiveUnsetOptionalState());
+			}
+		}
+
+		SECTION("Non-nullable optional subobject reference initialized from CDO - unset state")
+		{
+			USubobjectInstancingTestOuterObject* OuterObject = NewObject<USubobjectInstancingTestOuterObject>(GetTransientPackage(), InstancingTestClass);
+
+			CHECK_FALSE(OptionalProperty->IsSet(OptionalProperty->ContainerPtrToValuePtr<void>(OuterObject)));
+		}
+
+		SECTION("Non-nullable optional subobject reference initialized from template - unset state")
+		{
+			USubobjectInstancingTestOuterObject* Template = NewObject<USubobjectInstancingTestOuterObject>(GetTransientPackage(), InstancingTestClass, NAME_None, RF_ArchetypeObject);
+			{
+				CHECK_FALSE(OptionalProperty->IsSet(OptionalProperty->ContainerPtrToValuePtr<void>(Template)));
+			}
+
+			USubobjectInstancingTestOuterObject* OuterObject = NewObject<USubobjectInstancingTestOuterObject>(GetTransientPackage(), InstancingTestClass, NAME_None, RF_NoFlags, Template);
+			{
+				CHECK_FALSE(OptionalProperty->IsSet(OptionalProperty->ContainerPtrToValuePtr<void>(OuterObject)));
+			}
+		}
+	}
+
 	void RunSubobjectInstancingTests_Other(UClass* InstancingTestClass)
 	{
 		SECTION("StaticDuplicateObject() instancing")
@@ -1845,6 +1956,7 @@ namespace SubobjectInstancingTest::Private
 		RunSubobjectInstancingTests_Nested(InstancingTestClass);
 		RunSubobjectInstancingTests_Deferred(InstancingTestClass);
 		RunSubobjectInstancingTests_SelfReferencing(InstancingTestClass);
+		RunSubobjectInstancingTests_NonNullable(InstancingTestClass);
 		RunSubobjectInstancingTests_Other(InstancingTestClass);
 	}
 }
diff --git a/Engine/Source/Runtime/CoreUObject/Tests/UObject/SubobjectInstancingTest.h b/Engine/Source/Runtime/CoreUObject/Tests/UObject/SubobjectInstancingTest.h
--- a/Engine/Source/Runtime/CoreUObject/Tests/UObject/SubobjectInstancingTest.h
+++ b/Engine/Source/Runtime/CoreUObject/Tests/UObject/SubobjectInstancingTest.h
@@ -170,9 +170,6 @@ public:
 	UPROPERTY(Instanced)
 	TOptional<TObjectPtr<USubobjectInstancingTestObject>> OptionalNullableInnerObjectUnset;
 
-	UPROPERTY(Instanced)
-	TOptional<TNonNullPtr<USubobjectInstancingTestObject>> OptionalNonNullableInnerObjectUnset;
-
 	UPROPERTY(Instanced)
 	TObjectPtr<USubobjectInstancingTestDirectlyNestedObject> InnerObjectWithDirectlyNestedObject;
 
