[ZenLoader]
- Use race detector scope to cleanup the implementation
- Add a grace period to reduce the toggling rate that can reduce performance
- Add support for race detection from parallel nodes

diff --git a/Engine/Source/Runtime/CoreUObject/Private/Serialization/AsyncLoading2.cpp b/Engine/Source/Runtime/CoreUObject/Private/Serialization/AsyncLoading2.cpp
--- a/Engine/Source/Runtime/CoreUObject/Private/Serialization/AsyncLoading2.cpp
+++ b/Engine/Source/Runtime/CoreUObject/Private/Serialization/AsyncLoading2.cpp
@@ -559,6 +559,14 @@ static FAutoConsoleVariableRef CVarDetectRaceDuringLoading(
 	TEXT("Activate the race detector during loading periods"),
 	ECVF_Default
 );
+
+static float GDetectRaceGracePeriod = 1.0f;
+static FAutoConsoleVariableRef CVarDetectRaceGracePeriod(
+	TEXT("s.DetectRaceGracePeriod"),
+	GDetectRaceGracePeriod,
+	TEXT("Amount of seconds to stay on after all loading is done. Can improve performance by reducing the rate of toggling."),
+	ECVF_Default
+);
 #endif // USING_INSTRUMENTATION
 
 CSV_DECLARE_CATEGORY_MODULE_EXTERN(CORE_API, Basic);
@@ -5985,6 +5993,14 @@ public:
 
 void FEventLoadNode2::ParallelLoadingWork()
 {
+#if USING_INSTRUMENTATION
+	if (GDetectRaceDuringLoading)
+	{
+		UE::Sanitizer::RaceDetector::ToggleFilterOtherThreads(true);
+	}
+	UE::Sanitizer::RaceDetector::FRaceDetectorScope RaceDetectorScope(GDetectRaceDuringLoading, GDetectRaceGracePeriod);
+#endif
+
 	FTaskTagScope Scope(ETaskTag::EParallelLoadingThread);
 	FAsyncLoadingThreadState2 LocalThreadState;
 	FAsyncLoadingThreadState2Scope WorkerStateScope(LocalThreadState);
@@ -10462,22 +10478,11 @@ uint32 FAsyncLoadingThread2::Run()
 			TRACE_CPUPROFILER_EVENT_SCOPE(AsyncLoadingTime);
 
 #if USING_INSTRUMENTATION
-			bool bInstrumentationActive = false;
 			if (GDetectRaceDuringLoading)
 			{
-				bInstrumentationActive = true;
 				UE::Sanitizer::RaceDetector::ToggleFilterOtherThreads(true);
-				UE::Sanitizer::RaceDetector::ToggleRaceDetection(true);
 			}
-
-			ON_SCOPE_EXIT
-			{
-				if (bInstrumentationActive)
-				{
-					UE::Sanitizer::RaceDetector::ToggleRaceDetection(false);
-					UE::Sanitizer::RaceDetector::ToggleFilterOtherThreads(false);
-				}
-			};
+			UE::Sanitizer::RaceDetector::FRaceDetectorScope RaceDetectorScope(GDetectRaceDuringLoading, GDetectRaceGracePeriod);
 #endif
 
 			bool bShouldSuspend = false;
