Adding UE::RemoteObject::Handle::IsValid(RemoteId) global function that mirrors the functionality of the global IsValid(Object) function but for remote objects

diff --git a/Engine/Source/Runtime/CoreUObject/Private/UObject/RemoteObject.cpp b/Engine/Source/Runtime/CoreUObject/Private/UObject/RemoteObject.cpp
--- a/Engine/Source/Runtime/CoreUObject/Private/UObject/RemoteObject.cpp
+++ b/Engine/Source/Runtime/CoreUObject/Private/UObject/RemoteObject.cpp
@@ -184,6 +184,7 @@ public:
 		Stub->SerialNumber = ObjectItem->GetSerialNumber();
 		Stub->Name = Object->GetFName();
 		Stub->ResidentServerId = DestinationServerId;
+		Stub->bWasGarbage = !IsValid(Object);
 
 		return Stub;
 	}
@@ -524,7 +525,7 @@ FRemoteObjectStub::FRemoteObjectStub(UObject* Object)
 	{
 		Id = FObjectHandleUtils::GetRemoteId(Object);
 		Class = FRemoteObjectClass(Object->GetClass());
-
+		bWasGarbage = !IsValid(Object);
 		if (UObject* Outer = FObjectHandleUtils::GetNonAccessTrackedOuterNoResolve(Object))
 		{
 			OuterId = FObjectHandleUtils::GetRemoteId(Outer);
@@ -595,6 +596,26 @@ bool IsOwned(FRemoteObjectId ObjectId)
 	return bResult;
 }
 
+bool IsValid(FRemoteObjectId ObjectId)
+{
+	using namespace UE::RemoteObject::Private;
+
+	// Note: this function needs to reflect the logic of ResolveObject(...) functions
+
+	if (UObject* Object = StaticFindObjectFastInternal(ObjectId))
+	{		
+		return IsValid(Object);
+	}
+
+	if (FRemoteObjectStub* Stub = ObjectMaps->FindRemoteObjectStub(ObjectId))
+	{
+		return !Stub->bWasGarbage;
+	}
+
+	// ObjectId represents an object that has never been migrated
+	return false;
+}
+
 FRemoteServerId GetOwnerServerId(const UObject* Object)
 {
 	using namespace UE::CoreUObject::Private;
diff --git a/Engine/Source/Runtime/CoreUObject/Public/UObject/RemoteObject.h b/Engine/Source/Runtime/CoreUObject/Public/UObject/RemoteObject.h
--- a/Engine/Source/Runtime/CoreUObject/Public/UObject/RemoteObject.h
+++ b/Engine/Source/Runtime/CoreUObject/Public/UObject/RemoteObject.h
@@ -120,6 +120,7 @@ namespace UE::RemoteObject::Handle
 		FRemoteObjectId OuterId;
 		FName Name;
 		FRemoteObjectClass Class;
+		bool bWasGarbage = false;
 
 		/** SerialNumber this object had on this server */
 		int32 SerialNumber = 0;
@@ -210,6 +211,14 @@ namespace UE::RemoteObject::Handle
 	*/
 	COREUOBJECT_API bool IsOwned(FRemoteObjectId ObjectId);
 
+	/**
+	* Checks if an object id represents a valid object in the same sense as the global IsValid(Object)
+	* This information may not be accurate if an object has been marked as garbage on a remote server and hasn't been migrated back yet
+	* @params ObjectId remote object id
+	* @return true if an object id represents an object that is in memory and is not marked as garbage or a remote object which was not marked as garbage when it was being migrated, otherwise false (or when the provided id is not internally known)
+	*/
+	COREUOBJECT_API bool IsValid(FRemoteObjectId ObjectId);
+
 	/**
 	* Get the owner server id for a locally resident object
 	*/
