Fixed CastChecked<To>(NullPtr, ECastCheckedType::NullAllowed) in !DO_CHECK builds.

diff --git a/Engine/Source/Runtime/CoreUObject/Public/Templates/Casts.h b/Engine/Source/Runtime/CoreUObject/Public/Templates/Casts.h
--- a/Engine/Source/Runtime/CoreUObject/Public/Templates/Casts.h
+++ b/Engine/Source/Runtime/CoreUObject/Public/Templates/Casts.h
@@ -285,6 +285,16 @@ PRAGMA_ENABLE_DEPRECATION_WARNINGS
 	template <typename To, typename From>
 	UE_FORCEINLINE_HINT TCopyQualifiersFromTo_T<From, To>* CastChecked(From* Src, ECastCheckedType::Type CheckType)
 	{
+		// The CastChecked call below doesn't handle nulls in the case of interfaces, so handle them here when NullAllowed.
+		// This will also typically compile out when inlining with NullChecked.
+		if constexpr (TIsIInterface<From>::Value || TIsIInterface<To>::Value)
+		{
+			if (CheckType == ECastCheckedType::NullAllowed && !Src)
+			{
+				return nullptr;
+			}
+		}
+
 		return CastChecked<To>(Src);
 	}
 
