[ZenLoader]
- Get rid of postload groups in favor of using proper barriers on dependencies
- Add support for partial loads in the runtime to unify editor/runtime

diff --git a/Engine/Source/Runtime/CoreUObject/Private/Serialization/AsyncLoading2.cpp b/Engine/Source/Runtime/CoreUObject/Private/Serialization/AsyncLoading2.cpp
--- a/Engine/Source/Runtime/CoreUObject/Private/Serialization/AsyncLoading2.cpp
+++ b/Engine/Source/Runtime/CoreUObject/Private/Serialization/AsyncLoading2.cpp
@@ -103,9 +103,6 @@
 
 template<typename PayloadType> using TAsyncAtomic = std::atomic<PayloadType>;
 
-// For now, the partial request behavior is reserved for the editor only.
-#define WITH_PARTIAL_REQUEST_DURING_RECURSION WITH_EDITOR
-
 CSV_DECLARE_CATEGORY_MODULE_EXTERN(CORE_API, FileIO);
 CSV_DEFINE_STAT(FileIO, FrameCompletedExportBundleLoadsKB);
 
@@ -484,14 +481,6 @@ static FAutoConsoleVariableRef CVarGAllowMultithreadedLoading(
 	ECVF_Default
 	);
 
-static bool GUsePostloadGroupsWithIOStore = true;
-static FAutoConsoleVariableRef CVarGUsePostloadGroupsWithIOStore(
-	TEXT("s.UsePostloadGroupsWithIOStore"),
-	GUsePostloadGroupsWithIOStore,
-	TEXT("Postload groups reduces the likelihood of race conditions between the loading thread and the game thread but are slower."),
-	ECVF_Default
-);
-
 static bool GOnlyProcessRequiredPackagesWhenSyncLoading = true;
 static FAutoConsoleVariableRef CVarGOnlyProcessRequiredPackagesWhenSyncLoading(
 	TEXT("s.OnlyProcessRequiredPackagesWhenSyncLoading"),
@@ -3575,8 +3564,6 @@ private:
 };
 #endif // ALT2_ENABLE_NEW_ARCHIVE_FOR_LINKERLOAD
 
-struct FAsyncLoadingPostLoadGroup;
-
 /**
 * Structure containing intermediate data required for async loading of all exports of a package.
 */
@@ -3717,10 +3704,7 @@ private:
 	FAsyncPackageData Data;
 	FAsyncPackageHeaderData HeaderData;
 	FAsyncPackageSerializationState SerializationState;
-
-#if WITH_PARTIAL_REQUEST_DURING_RECURSION
 	TSet<FAsyncPackage2*> AdditionalImportedAsyncPackages;
-#endif
 
 #if WITH_EDITOR
 	TOptional<FAsyncPackageHeaderData> OptionalSegmentHeaderData;
@@ -3755,8 +3739,6 @@ private:
 	UPackage* LinkerRoot = nullptr;
 	// The sync load context associated with this package
 	TAsyncAtomic<uint64>			SyncLoadContextId = 0;
-	FAsyncLoadingPostLoadGroup* PostLoadGroup = nullptr;
-	FAsyncLoadingPostLoadGroup* DeferredPostLoadGroup = nullptr;
 	/** Time load begun. This is NOT the time the load was requested in the case of pending requests.	*/
 	double						LoadStartTime = 0.0;
 	TAsyncAtomic<int32>			RefCount = 0;
@@ -3987,7 +3969,8 @@ private:
 	void ConditionalFinishLoading(FAsyncLoadingThreadState2& ThreadState);
 	void ConditionalBeginPostLoadPhase(FAsyncLoadingThreadState2& ThreadState);
 	void ConditionalReleasePartialRequests(FAsyncLoadingThreadState2& ThreadState);
-
+	void ConditionalBeginPostLoad(FAsyncLoadingThreadState2& ThreadState);
+	void ConditionalBeginDeferredPostLoad(FAsyncLoadingThreadState2& ThreadState);
 #if ALT2_ENABLE_LINKERLOAD_SUPPORT
 	EEventLoadNodeExecutionResult ProcessLinkerLoadPackageSummary(FAsyncLoadingThreadState2& ThreadState);
 	bool CreateLinkerLoadExports(FAsyncLoadingThreadState2& ThreadState);
@@ -4043,11 +4026,6 @@ private:
 	enum EExternalReadAction { ExternalReadAction_Poll, ExternalReadAction_Wait };
 	EAsyncPackageState::Type ProcessExternalReads(FAsyncLoadingThreadState2& ThreadState, EExternalReadAction Action);
 
-	/**
-	* Updates load percentage stat
-	*/
-	void UpdateLoadPercentage();
-
 #if ALT2_DUMP_STATE_ON_HANG
 	void DumpStateImpl(TSet<FAsyncPackage2*>& Set, int32 Indent = 0, TMultiMap<FEventLoadNode2*, FEventLoadNode2*>* MappedNodes = nullptr);
 #endif
@@ -4737,9 +4715,6 @@ private:
 #if WITH_EDITOR
 	void ConditionalProcessEditorCallbacks();
 #endif
-	void ConditionalBeginPostLoad(FAsyncLoadingThreadState2& ThreadState, FAsyncLoadingPostLoadGroup* PostLoadGroup);
-	void ConditionalBeginDeferredPostLoad(FAsyncLoadingThreadState2& ThreadState, FAsyncLoadingPostLoadGroup* DeferredPostLoadGroup);
-	void MergePostLoadGroups(FAsyncLoadingThreadState2& ThreadState, FAsyncLoadingPostLoadGroup* Target, FAsyncLoadingPostLoadGroup* Source, bool bUpdateSyncLoadContext = true);
 
 public:
 	bool ProcessDeferredDeletePackagesQueue(int32 MaxCount = MAX_int32)
@@ -5384,115 +5359,38 @@ void FAsyncLoadingThread2::UpdatePackagePriorityRecursive(FAsyncLoadingThreadSta
 		}
 	}
 
-#if WITH_PARTIAL_REQUEST_DURING_RECURSION
 	for (FAsyncPackage2* ImportedPackage : Package->AdditionalImportedAsyncPackages)
 	{
 		UpdatePackagePriorityRecursive(ThreadState, ImportedPackage, NewPriority);
 	}
-#endif
 }
 
-void FAsyncLoadingThread2::ConditionalBeginPostLoad(FAsyncLoadingThreadState2& ThreadState, FAsyncLoadingPostLoadGroup* PostLoadGroup)
+void FAsyncPackage2::ConditionalBeginPostLoad(FAsyncLoadingThreadState2& ThreadState)
 {
 	TRACE_CPUPROFILER_EVENT_SCOPE(ConditionalBeginPostLoad);
-	check(PostLoadGroup);
-	check(ThreadState.bCanAccessAsyncLoadingThreadData);
-	if (PostLoadGroup->PackagesWithExportsToSerializeCount == 0)
-	{
-		PostLoadGroup->PackagesWithExportsToPostLoadCount = PostLoadGroup->Packages.Num();
-
-		// Release the post load node of packages in the post load group in reverse order that they were added to the group
-		// This usually means that dependencies will be post load first, similarly to how they are also serialized first
-		for (int32 Index = PostLoadGroup->Packages.Num() - 1; Index >= 0; --Index)
-		{
-			FAsyncPackage2* Package = PostLoadGroup->Packages[Index];
-			check(Package->PostLoadGroup == PostLoadGroup);
-			check(Package->AsyncPackageLoadingState == EAsyncPackageLoadingState2::ExportsDone);
 
-			// Move the PostLoadGroup to the DeferredPostLoadGroup so that we do not mistakenly consider post load as not having being triggered yet
-			Package->PostLoadGroup = nullptr;
-			Package->DeferredPostLoadGroup = PostLoadGroup;
-			
-			Package->AsyncPackageLoadingState = EAsyncPackageLoadingState2::PostLoad;
-			Package->ConditionalReleasePartialRequests(ThreadState);
-
-			Package->ConditionalBeginPostLoadPhase(ThreadState);
-		}
-	}
+	check(AsyncPackageLoadingState == EAsyncPackageLoadingState2::ExportsDone);
+	AsyncPackageLoadingState = EAsyncPackageLoadingState2::PostLoad;
+	ConditionalReleasePartialRequests(ThreadState);
+	ConditionalBeginPostLoadPhase(ThreadState);
 }
 
-void FAsyncLoadingThread2::ConditionalBeginDeferredPostLoad(FAsyncLoadingThreadState2& ThreadState, FAsyncLoadingPostLoadGroup* DeferredPostLoadGroup)
+void FAsyncPackage2::ConditionalBeginDeferredPostLoad(FAsyncLoadingThreadState2& ThreadState)
 {
 	TRACE_CPUPROFILER_EVENT_SCOPE(ConditionalBeginDeferredPostLoad);
-	check(DeferredPostLoadGroup);
-	check(ThreadState.bCanAccessAsyncLoadingThreadData);
-	if (DeferredPostLoadGroup->PackagesWithExportsToPostLoadCount == 0)
-	{
-		// Move everything to the second phase so that StaticFind from inside postload can
-		// see objects ready for postload.
-		for (int32 Index = DeferredPostLoadGroup->Packages.Num() - 1; Index >= 0; --Index)
-		{
-			FAsyncPackage2* Package = DeferredPostLoadGroup->Packages[Index];
-			Package->MoveConstructedObjectsToPhase2();
-		}
 
-		// Release the post load node of packages in the post load group in reverse order that they were added to the group
-		// This usually means that dependencies will be post load first, similarly to how they are also serialized first
-		for (int32 Index = DeferredPostLoadGroup->Packages.Num() - 1; Index >= 0; --Index)
-		{
-			FAsyncPackage2* Package = DeferredPostLoadGroup->Packages[Index];
-			check(Package->DeferredPostLoadGroup == DeferredPostLoadGroup);
-			check(Package->AsyncPackageLoadingState == EAsyncPackageLoadingState2::PostLoad);
-			Package->DeferredPostLoadGroup = nullptr;
-			Package->AsyncPackageLoadingState = EAsyncPackageLoadingState2::DeferredPostLoad;
-			Package->ConditionalReleasePartialRequests(ThreadState);
-			for (int32 ExportBundleIndex = 0; ExportBundleIndex < Package->Data.TotalExportBundleCount; ++ExportBundleIndex)
-			{
-				Package->GetExportBundleNode(EEventLoadNode2::ExportBundle_DeferredPostLoad, ExportBundleIndex).ReleaseBarrier(&ThreadState);
-			}
-		}
-		delete DeferredPostLoadGroup;
-	}
-}
+	// Move everything to the second phase so that StaticFind from inside postload can
+	// see objects ready for postload.
+	MoveConstructedObjectsToPhase2();
+	
+	check(AsyncPackageLoadingState == EAsyncPackageLoadingState2::PostLoad);
+	AsyncPackageLoadingState = EAsyncPackageLoadingState2::DeferredPostLoad;
+	ConditionalReleasePartialRequests(ThreadState);
 
-void FAsyncLoadingThread2::MergePostLoadGroups(FAsyncLoadingThreadState2& ThreadState, FAsyncLoadingPostLoadGroup* Target, FAsyncLoadingPostLoadGroup* Source, bool bUpdateSyncLoadContext)
-{
-	if (Target == Source)
-	{
-		return;
-	}
-	TRACE_CPUPROFILER_EVENT_SCOPE(MergePostLoadGroups);
-	check(ThreadState.bCanAccessAsyncLoadingThreadData);
-	for (FAsyncPackage2* Package : Source->Packages)
+	for (int32 ExportBundleIndex = 0; ExportBundleIndex < Data.TotalExportBundleCount; ++ExportBundleIndex)
 	{
-		check(Package->PostLoadGroup == Source);
-		Package->PostLoadGroup = Target;
+		GetExportBundleNode(EEventLoadNode2::ExportBundle_DeferredPostLoad, ExportBundleIndex).ReleaseBarrier(&ThreadState);
 	}
-	Target->Packages.Append(MoveTemp(Source->Packages));
-	Target->PackagesWithExportsToSerializeCount += Source->PackagesWithExportsToSerializeCount;
-	check(Target->PackagesWithExportsToPostLoadCount == 0 && Source->PackagesWithExportsToPostLoadCount == 0);
-
-	// If the intention of the caller of this function is to merge postloads into the caller package so they are executed later after
-	// a partial flush, then we can't update the synccontext of the caller during the merge as it would make the whole hierarchy
-	// flush in a single swoop which is in direct contradiction with the partial flush feature.
-	if (bUpdateSyncLoadContext)
-	{
-		const uint64 SyncLoadContextId = FMath::Max(Source->SyncLoadContextId, Target->SyncLoadContextId);
-		if (SyncLoadContextId)
-		{
-			Target->SyncLoadContextId = SyncLoadContextId;
-			for (FAsyncPackage2* Package : Target->Packages)
-			{
-				Package->SyncLoadContextId = SyncLoadContextId;
-				if (Package->Desc.Priority < MAX_int32)
-				{
-					Package->Desc.Priority = MAX_int32;
-					UpdatePackagePriority(ThreadState, Package);
-				}
-			}
-		}
-	}
-	delete Source;
 }
 
 FAsyncPackage2* FAsyncLoadingThread2::FindOrInsertPackage(FAsyncLoadingThreadState2& ThreadState, FAsyncPackageDesc2& Desc, bool& bInserted, FAsyncPackage2* ImportedByPackage, TUniquePtr<FLoadPackageAsyncDelegate>&& PackageLoadedDelegate, TSharedPtr<FLoadPackageAsyncProgressDelegate> PackageProgressDelegate)
@@ -5556,49 +5454,6 @@ FAsyncPackage2* FAsyncLoadingThread2::FindOrInsertPackage(FAsyncLoadingThreadSta
 		);
 	}
 
-	// PostLoadGroup are to protect some race conditions happening in runtime when ALT is active but
-	// it is causing crashes and soft-locks in editor when FlushAsyncLoading is called on specific requests
-	// that are in the same postload group as their outer, preventing fine grained flush. In runtime it can cause a runtime regression
-	// if you don't need protection from postload races. So we disable by default for the time being until we sort out the above concerns.
-	constexpr bool bIsPostLoadGroupFeatureActive = false;
-#if 0
-	// Prevents activating postload groups during boot because it was causing deadlock into the 
-	// highly recursive InitDefaultMaterials function on some platforms.
-	// Since postload groups are there to protect against race conditions between postloads and serialize,
-	// no such race can exists until the loading thread is started.
-	const bool bIsPostLoadGroupFeatureActive = IsMultithreaded() && GUsePostloadGroupsWithIOStore;
-#endif
-
-	if (bInserted)
-	{
-		// Created a new package, either create a new post load group or use the one from the importing package
-		FAsyncLoadingPostLoadGroup* PostLoadGroup = (bIsPostLoadGroupFeatureActive && ImportedByPackage) ? ImportedByPackage->PostLoadGroup : new FAsyncLoadingPostLoadGroup();
-		++PostLoadGroup->PackagesWithExportsToSerializeCount;
-		PostLoadGroup->Packages.Add(Package);
-		check(!Package->PostLoadGroup);
-		Package->PostLoadGroup = PostLoadGroup;
-	}
-	else if (ImportedByPackage && bIsPostLoadGroupFeatureActive)
-	{
-		// Importing a package that was already being loaded
-		if (!Package->PostLoadGroup)
-		{
-			// The imported package has started postloading, wait for it to finish postloading before serializing any exports
-			for (int32 DependentExportBundleIndex = 0; DependentExportBundleIndex < ImportedByPackage->Data.TotalExportBundleCount; ++DependentExportBundleIndex)
-			{
-				for (int32 DependsOnExportBundleIndex = 0; DependsOnExportBundleIndex < Package->Data.TotalExportBundleCount; ++DependsOnExportBundleIndex)
-				{
-					ImportedByPackage->GetExportBundleNode(ExportBundle_Process, DependentExportBundleIndex).DependsOn(&Package->GetExportBundleNode(ExportBundle_DeferredPostLoad, DependsOnExportBundleIndex));
-				}
-			}
-		}
-		else if (ImportedByPackage->PostLoadGroup != Package->PostLoadGroup)
-		{
-			// The imported package hasn't started postloading yet, merge its post load group with the one for the importing package
-			check(ImportedByPackage->PostLoadGroup);
-			MergePostLoadGroups(ThreadState, ImportedByPackage->PostLoadGroup, Package->PostLoadGroup);
-		}
-	}
 	return Package;
 }
 
@@ -5624,19 +5479,6 @@ void FAsyncLoadingThread2::IncludePackageInSyncLoadContextRecursive(FAsyncLoadin
 		UpdatePackagePriority(ThreadState, Package);
 	}
 
-	FAsyncLoadingPostLoadGroup* PostLoadGroup = Package->PostLoadGroup ? Package->PostLoadGroup : Package->DeferredPostLoadGroup;
-	if (PostLoadGroup && PostLoadGroup->SyncLoadContextId < ContextId)
-	{
-		PostLoadGroup->SyncLoadContextId = ContextId;
-		for (FAsyncPackage2* PackageInPostLoadGroup : PostLoadGroup->Packages)
-		{
-			if (PackageInPostLoadGroup->SyncLoadContextId < ContextId)
-			{
-				IncludePackageInSyncLoadContextRecursive(ThreadState, ContextId, PackageInPostLoadGroup);
-			}
-		}
-	}
-
 	// For LinkerLoad, dependencies are loaded in ProcessPackageSummary so if this is called before,
 	// early out to avoid races and we'll reprocess them in the ProcessDependencies step
 	if (Package->LinkerLoadState.IsSet() &&
@@ -5653,7 +5495,6 @@ void FAsyncLoadingThread2::IncludePackageInSyncLoadContextRecursive(FAsyncLoadin
 		}
 	}
 
-#if WITH_PARTIAL_REQUEST_DURING_RECURSION
 	for (FAsyncPackage2* ImportedPackage : Package->AdditionalImportedAsyncPackages)
 	{
 		if (ImportedPackage->SyncLoadContextId < ContextId)
@@ -5661,7 +5502,6 @@ void FAsyncLoadingThread2::IncludePackageInSyncLoadContextRecursive(FAsyncLoadin
 			IncludePackageInSyncLoadContextRecursive(ThreadState, ContextId, ImportedPackage);
 		}
 	}
-#endif
 }
 
 UE_TRACE_EVENT_BEGIN(CUSTOM_LOADTIMER_LOG, CreateAsyncPackage, NoSync)
@@ -6883,11 +6723,9 @@ void FAsyncPackage2::ImportPackagesRecursive(FAsyncLoadingThreadState2& ThreadSt
 	}
 	bHasStartedImportingPackages = true;
 
-#if WITH_PARTIAL_REQUEST_DURING_RECURSION
 	// No additional imports should be present when ImportPackagesRecursive is called since this should only be called during
 	// the early phases of loading.
 	check(AdditionalImportedAsyncPackages.IsEmpty());
-#endif
 
 	if (Data.ImportedAsyncPackages.IsEmpty())
 	{
@@ -6912,12 +6750,10 @@ void FAsyncPackage2::ImportPackagesRecursive(FAsyncLoadingThreadState2& ThreadSt
 			}
 		}
 
-#if WITH_PARTIAL_REQUEST_DURING_RECURSION
 		for (FAsyncPackage2* ImportedPackage : AdditionalImportedAsyncPackages)
 		{
 			AsyncLoadingThread.IncludePackageInSyncLoadContextRecursive(ThreadState, SyncLoadContextId, ImportedPackage);
 		}
-#endif
 	}
 
 	UE_ASYNC_PACKAGE_LOG_VERBOSE(VeryVerbose, Desc, TEXT("ImportPackages: ImportsDone"),
@@ -7576,11 +7412,7 @@ EEventLoadNodeExecutionResult FAsyncPackage2::ExecutePostLoadLinkerLoadPackageEx
 	// Reset this to be reused for the deferred postload phase
 	LinkerLoadState->PostLoadExportIndex = 0;
 
-	check(DeferredPostLoadGroup);
-	check(DeferredPostLoadGroup->PackagesWithExportsToPostLoadCount > 0);
-	--DeferredPostLoadGroup->PackagesWithExportsToPostLoadCount;
-
-	AsyncLoadingThread.ConditionalBeginDeferredPostLoad(ThreadState, DeferredPostLoadGroup);
+	ConditionalBeginDeferredPostLoad(ThreadState);
 	return EEventLoadNodeExecutionResult::Complete;
 }
 
@@ -9068,7 +8900,6 @@ FAsyncPackage2* FAsyncPackage2::GetCurrentlyExecutingPackage(FAsyncLoadingThread
 
 void FAsyncPackage2::ConditionalReleasePartialRequests(FAsyncLoadingThreadState2& ThreadState)
 {
-#if WITH_PARTIAL_REQUEST_DURING_RECURSION
 	TRACE_CPUPROFILER_EVENT_SCOPE(ConditionalReleasePartialRequests);
 
 	// We need to look at all syncload context now that its possible for any threads to push a context as they are not strictly recursive anymore.
@@ -9129,7 +8960,6 @@ void FAsyncPackage2::ConditionalReleasePartialRequests(FAsyncLoadingThreadState2
 			}
 		}
 	}
-#endif // WITH_PARTIAL_REQUEST_DURING_RECURSION
 }
 
 EEventLoadNodeExecutionResult FAsyncPackage2::Event_ExportsDone(FAsyncLoadingThreadState2& ThreadState, FAsyncPackage2* Package, int32)
@@ -9164,13 +8994,8 @@ EEventLoadNodeExecutionResult FAsyncPackage2::Event_ExportsDone(FAsyncLoadingThr
 
 	Package->CallProgressCallbacks(EAsyncLoadingProgress::Serialized);
 
-	FAsyncLoadingPostLoadGroup* PostLoadGroup = Package->PostLoadGroup;
-	check(PostLoadGroup);
-	check(PostLoadGroup->PackagesWithExportsToSerializeCount > 0);
-	--PostLoadGroup->PackagesWithExportsToSerializeCount;
-
 	Package->ConditionalReleasePartialRequests(ThreadState);
-	Package->AsyncLoadingThread.ConditionalBeginPostLoad(ThreadState, PostLoadGroup);
+	Package->ConditionalBeginPostLoad(ThreadState);
 	return EEventLoadNodeExecutionResult::Complete;
 }
 
@@ -9338,7 +9163,6 @@ FAsyncPackage2* FAsyncPackage2::UpdateDependenciesStateRecursive(FAsyncLoadingTh
 		}
 	}
 
-#if WITH_PARTIAL_REQUEST_DURING_RECURSION
 	if (bContinueToProcessPackage)
 	{
 		for (FAsyncPackage2* ImportedPackage : AdditionalImportedAsyncPackages)
@@ -9349,7 +9173,6 @@ FAsyncPackage2* FAsyncPackage2::UpdateDependenciesStateRecursive(FAsyncLoadingTh
 			}
 		}
 	}
-#endif
 
 	if (Context.P.Top() == this)
 	{
@@ -9594,11 +9417,7 @@ EEventLoadNodeExecutionResult FAsyncPackage2::Event_PostLoadExportBundle(FAsyncL
 	{
 		Package->ProcessedExportBundlesCount = 0;
 
-		FAsyncLoadingPostLoadGroup* DeferredPostLoadGroup = Package->DeferredPostLoadGroup;
-		check(DeferredPostLoadGroup);
-		check(DeferredPostLoadGroup->PackagesWithExportsToPostLoadCount > 0);
-		--DeferredPostLoadGroup->PackagesWithExportsToPostLoadCount;
-		Package->AsyncLoadingThread.ConditionalBeginDeferredPostLoad(ThreadState, DeferredPostLoadGroup);
+		Package->ConditionalBeginDeferredPostLoad(ThreadState);
 	}
 
 	return EEventLoadNodeExecutionResult::Complete;
@@ -9821,7 +9640,6 @@ void FAsyncLoadingThread2::UpdateSyncLoadContext(FAsyncLoadingThreadState2& Thre
 				// after validating that bHasFoundRequestedPackage is true would see the proper value.
 				SyncLoadContext->SetRequestedPackage(i, RequestedPackage);
 
-#if WITH_PARTIAL_REQUEST_DURING_RECURSION
 				if (FAsyncPackage2* RequestingPackage = SyncLoadContext->GetRequestingPackage())
 				{
 					// If the flush is coming from a step before the requesting package is back on GT, there is no way to fully flush
@@ -9835,7 +9653,6 @@ void FAsyncLoadingThread2::UpdateSyncLoadContext(FAsyncLoadingThreadState2& Thre
 						RequestedPackage->ConditionalReleasePartialRequests(ThreadState);
 					}
 				}
-#endif
 
 				IncludePackageInSyncLoadContextRecursive(ThreadState, SyncLoadContext->ContextId, RequestedPackage);
 				++FoundPackages;
@@ -11172,9 +10989,6 @@ FAsyncPackage2::~FAsyncPackage2()
 
 	FMemory::Free(Data.MemoryBuffer0);
 	FMemory::Free(Data.MemoryBuffer1);
-
-	check(PostLoadGroup == nullptr);
-	check(DeferredPostLoadGroup == nullptr);
 }
 
 void FAsyncPackage2::Release()
@@ -11205,13 +11019,11 @@ void FAsyncPackage2::ClearImportedPackages()
 		}
 	}
 
-#if WITH_PARTIAL_REQUEST_DURING_RECURSION
 	TSet<FAsyncPackage2*> LocalAdditionalImportedAsyncPackages = MoveTemp(AdditionalImportedAsyncPackages);
 	for (FAsyncPackage2* ImportedAsyncPackage : LocalAdditionalImportedAsyncPackages)
 	{
 		ImportedAsyncPackage->Release();
 	}
-#endif
 }
 
 void FAsyncPackage2::MoveConstructedObjectsToPhase2()
@@ -11835,12 +11647,8 @@ void FAsyncLoadingThread2::FlushLoading(TConstArrayView<int32> RequestIDs)
 		// Prevent objects still being loaded from being accessible from the game thread.
 		TGuardValue GuardVisibilityFilter(FUObjectThreadContext::Get().AsyncVisibilityFilter, EInternalObjectFlags::AsyncLoadingPhase1);
 
-		// We can't possibly support flushing from async loading thread unless we have the partial request support active.
-#if WITH_PARTIAL_REQUEST_DURING_RECURSION
 		const bool bIsFlushSupportedOnCurrentThread = IsInGameThread() || IsInAsyncLoadingThread() || IsInParallelLoadingThread();
-#else
-		const bool bIsFlushSupportedOnCurrentThread = IsInGameThread();
-#endif
+
 		ELoaderType LoaderType = GetLoaderType();
 
 		if (!bIsFlushSupportedOnCurrentThread)
@@ -11917,26 +11725,6 @@ void FAsyncLoadingThread2::FlushLoading(TConstArrayView<int32> RequestIDs)
 				FAsyncPackage2::FAllDependenciesState::RemoveFromWaitList(&FAsyncPackage2::AllDependenciesFullyLoadedState, CurrentlyExecutingPackage, WaitingPackage);
 				WaitingPackage->ConditionalFinishLoading(*GameThreadState);
 			}
-#if !WITH_PARTIAL_REQUEST_DURING_RECURSION
-			if (GameThreadState->bCanAccessAsyncLoadingThreadData)
-			{
-				if (FAsyncLoadingPostLoadGroup* PostLoadGroup = CurrentlyExecutingPackage->PostLoadGroup)
-				{
-					check(PostLoadGroup->Packages.Contains(CurrentlyExecutingPackage));
-					check(PostLoadGroup->PackagesWithExportsToSerializeCount > 0);
-					if (PostLoadGroup->Packages.Num() > 1)
-					{
-						PostLoadGroup->Packages.Remove(CurrentlyExecutingPackage);
-						--PostLoadGroup->PackagesWithExportsToSerializeCount;
-						ConditionalBeginPostLoad(*GameThreadState, PostLoadGroup);
-						CurrentlyExecutingPackage->PostLoadGroup = new FAsyncLoadingPostLoadGroup();
-						CurrentlyExecutingPackage->PostLoadGroup->SyncLoadContextId = CurrentlyExecutingPackage->SyncLoadContextId;
-						CurrentlyExecutingPackage->PostLoadGroup->Packages.Add(CurrentlyExecutingPackage);
-						CurrentlyExecutingPackage->PostLoadGroup->PackagesWithExportsToSerializeCount = 1;
-					}
-				}
-			}
-#endif
 		}
 
 		FAsyncLoadingSyncLoadContext* SyncLoadContext = nullptr;
@@ -12406,7 +12194,6 @@ void FAsyncPackage2::DumpStateImpl(TSet<FAsyncPackage2*>& Set, int32 Indent, TMu
 		}
 	}
 
-#if WITH_PARTIAL_REQUEST_DURING_RECURSION
 	for (FAsyncPackage2* Import : AdditionalImportedAsyncPackages)
 	{
 		if (Import->AsyncPackageLoadingState < EAsyncPackageLoadingState2::DeferredDelete)
@@ -12414,7 +12201,6 @@ void FAsyncPackage2::DumpStateImpl(TSet<FAsyncPackage2*>& Set, int32 Indent, TMu
 			UE_LOG(LogStreaming, Warning, TEXT(" %s dynamically imports %s package %s state %s"), FCString::Spc(Indent), GetPackageType(Import), *Import->Desc.UPackageName.ToString(), LexToString(Import->AsyncPackageLoadingState));
 		}
 	}
-#endif
 
 	if (AllDependenciesSetupState.WaitingForPackage)
 	{
@@ -12448,7 +12234,6 @@ void FAsyncPackage2::DumpStateImpl(TSet<FAsyncPackage2*>& Set, int32 Indent, TMu
 		}
 	}
 
-#if WITH_PARTIAL_REQUEST_DURING_RECURSION
 	for (FAsyncPackage2* Import : AdditionalImportedAsyncPackages)
 	{
 		if (Import && Import->AsyncPackageLoadingState < EAsyncPackageLoadingState2::DeferredDelete)
@@ -12460,8 +12245,6 @@ void FAsyncPackage2::DumpStateImpl(TSet<FAsyncPackage2*>& Set, int32 Indent, TMu
 			}
 		}
 	}
-#endif
-
 }
 
 #endif // ALT2_DUMP_STATE_ON_HANG
