[StructUtils] FInstancedPropertyBag serializes with the default values. With tagged serialization, the full instance will be skipped if the descriptions and the values are the same. If the descriptions are the same, only the values that are modified will be serialized. If the descriptions changed, then the default is not used.
When using FInstancedPropertyBag in a BP, it might appear that it works without this changelist. This is due to the reinitialization process during compilation. Without this change list, it might not work depending on the loading order.

diff --git a/Engine/Plugins/Experimental/UAF/UAF/Source/UAF/Private/RewindDebugger/AnimNextTrace.cpp b/Engine/Plugins/Experimental/UAF/UAF/Source/UAF/Private/RewindDebugger/AnimNextTrace.cpp
--- a/Engine/Plugins/Experimental/UAF/UAF/Source/UAF/Private/RewindDebugger/AnimNextTrace.cpp
+++ b/Engine/Plugins/Experimental/UAF/UAF/Source/UAF/Private/RewindDebugger/AnimNextTrace.cpp
@@ -223,7 +223,9 @@ bool FAnimNextTrace::OutputAnimNextVariableSet(const TSharedRef<FUAFInstanceVari
 						FMemoryWriter WriterArchive(ArchiveData);
 						FObjectAsTraceIdProxyArchive Archive(WriterArchive);
 						Archive.UsingCustomVersion(CustomVersionGUID);
+PRAGMA_DISABLE_DEPRECATION_WARNINGS
 						Archive.UsingCustomVersion(FPropertyBagCustomVersion::GUID);
+PRAGMA_ENABLE_DEPRECATION_WARNINGS
 						Archive << PropertyDescriptions;
 
 						UE_TRACE_LOG(UAF, InstanceVariableDescriptions, UAFChannel)
diff --git a/Engine/Source/Runtime/CoreUObject/Private/StructUtils/PropertyBag.cpp b/Engine/Source/Runtime/CoreUObject/Private/StructUtils/PropertyBag.cpp
--- a/Engine/Source/Runtime/CoreUObject/Private/StructUtils/PropertyBag.cpp
+++ b/Engine/Source/Runtime/CoreUObject/Private/StructUtils/PropertyBag.cpp
@@ -34,7 +34,9 @@ static TAutoConsoleVariable<bool> CVarValidatePropertyBagDescriptionEquality(
 
 namespace FPropertyBagCustomVersion
 {
-	enum Type
+	enum
+	UE_DEPRECATED(all, "Use a stream custom version. Data made with a custom version for feature do not merge between streams.")
+	Type
 	{
 		// Before any version changes were made in the plugin
 		BeforeCustomVersionWasAdded = 0,
@@ -50,11 +52,14 @@ namespace FPropertyBagCustomVersion
 	};
 };
 
-
+UE_DEPRECATED(all, "Use a stream custom version. Data made with a custom version for feature do not merge between streams.")
 const FGuid FPropertyBagCustomVersion::GUID(0x134A157E, 0xD5E249A3, 0x8D4E843C, 0x98FE9E31);
 
+PRAGMA_DISABLE_DEPRECATION_WARNINGS
 // Register the custom version with core
+UE_DEPRECATED(all, "Use a stream custom version. Data made with a custom version for feature do not merge between streams.")
 FCustomVersionRegistration GPropertyBagCustomVersion(FPropertyBagCustomVersion::GUID, FPropertyBagCustomVersion::LatestVersion, TEXT("PropertyBagCustomVersion"));
+PRAGMA_ENABLE_DEPRECATION_WARNINGS
 
 
 namespace UE::StructUtils::Private
@@ -1711,9 +1716,16 @@ FArchive& operator<<(FArchive& Ar, FPropertyBagPropertyDesc& Bag)
 	Ar << Bag.Name;
 	Ar << Bag.ValueType;
 
-	if (Ar.CustomVer(FPropertyBagCustomVersion::GUID) >= FPropertyBagCustomVersion::ContainerTypes)
+PRAGMA_DISABLE_DEPRECATION_WARNINGS
+	const int32 PropertyBagCustomVersion = Ar.CustomVer(FPropertyBagCustomVersion::GUID);
+	constexpr int32 ContainerTypesCustomVersion = FPropertyBagCustomVersion::ContainerTypes;
+	constexpr int32 NestedContainerTypesCustomVersion = FPropertyBagCustomVersion::NestedContainerTypes;
+	constexpr int32 MetaClassCustomVersion = FPropertyBagCustomVersion::MetaClass;
+PRAGMA_ENABLE_DEPRECATION_WARNINGS
+
+	if (PropertyBagCustomVersion >= ContainerTypesCustomVersion)
 	{
-		if (Ar.IsLoading() && Ar.CustomVer(FPropertyBagCustomVersion::GUID) < FPropertyBagCustomVersion::NestedContainerTypes)
+		if (Ar.IsLoading() && PropertyBagCustomVersion < NestedContainerTypesCustomVersion)
 		{
 			EPropertyBagContainerType TmpContainerType = EPropertyBagContainerType::None;
 			Ar << TmpContainerType;
@@ -1743,7 +1755,7 @@ FArchive& operator<<(FArchive& Ar, FPropertyBagPropertyDesc& Bag)
 #if WITH_EDITORONLY_DATA
 		Ar << Bag.MetaData;
 
-		if (Ar.CustomVer(FPropertyBagCustomVersion::GUID) >= FPropertyBagCustomVersion::MetaClass)
+		if (PropertyBagCustomVersion >= MetaClassCustomVersion)
 		{
 			Ar << Bag.MetaClass;
 		}
@@ -1751,7 +1763,7 @@ FArchive& operator<<(FArchive& Ar, FPropertyBagPropertyDesc& Bag)
 		TArray<FPropertyBagPropertyDescMetaData> TempMetaData; 
 		Ar << TempMetaData;
 
-		if (Ar.CustomVer(FPropertyBagCustomVersion::GUID) >= FPropertyBagCustomVersion::MetaClass)
+		if (PropertyBagCustomVersion >= MetaClassCustomVersion)
 		{
 			TObjectPtr<UClass> TempMetaClass = nullptr;
 			Ar << TempMetaClass;
@@ -3223,9 +3235,9 @@ bool FInstancedPropertyBag::Identical(const FInstancedPropertyBag* Other, const
 	return Other && Value.Identical(&Other->Value, PortFlags);
 }
 
-bool FInstancedPropertyBag::Serialize(FArchive& Ar)
+bool FInstancedPropertyBag::Serialize(FArchive& Ar, UStruct* DefaultStruct, const void* DefaultValue)
 {
-	// Obsolete, use custom version instead.
+	// Obsolete
 	enum class EVersion : uint8
 	{
 		InitialVersion = 0,
@@ -3236,18 +3248,21 @@ bool FInstancedPropertyBag::Serialize(FArchive& Ar)
 	};
 	EVersion Version = EVersion::LatestVersion;
 
+PRAGMA_DISABLE_DEPRECATION_WARNINGS
 	Ar.UsingCustomVersion(FPropertyBagCustomVersion::GUID);
-
 	if (Ar.CustomVer(FPropertyBagCustomVersion::GUID) < FPropertyBagCustomVersion::ContainerTypes)
 	{
 		Ar << Version;
 	}
-	
-	UPropertyBag* BagStruct = Cast<UPropertyBag>(const_cast<UScriptStruct*>(Value.GetScriptStruct()));
+PRAGMA_ENABLE_DEPRECATION_WARNINGS
+
+
+	const UPropertyBag* BagStruct = Cast<const UPropertyBag>(const_cast<UScriptStruct*>(Value.GetScriptStruct()));
+	const FInstancedPropertyBag* DefaultInstancedPropertyBag = static_cast<const FInstancedPropertyBag*>(DefaultValue);
+	const UPropertyBag* DefaultBagStruct = DefaultInstancedPropertyBag ? Cast<const UPropertyBag>(DefaultInstancedPropertyBag->Value.GetScriptStruct()) : nullptr;
+
 	bool bHasData = (BagStruct != nullptr);
-	
 	Ar << bHasData;
-	
 	if (bHasData)
 	{
 		// The script struct class is not serialized, the properties are serialized and type is created based on that.
@@ -3263,9 +3278,22 @@ bool FInstancedPropertyBag::Serialize(FArchive& Ar)
 					Ar.Preload(const_cast<UObject*>(PropDesc.ValueTypeObject.Get()));
 				}
 			}
-			
-			BagStruct = const_cast<UPropertyBag*>(UPropertyBag::GetOrCreateFromDescs(PropertyDescs));
-			Value.InitializeAs(BagStruct);
+
+			BagStruct = UPropertyBag::GetOrCreateFromDescs(PropertyDescs);
+
+			if (BagStruct == DefaultBagStruct)
+			{
+				//Initialize with the defaults
+				const uint8* ValueMemoryDefault = DefaultInstancedPropertyBag ? DefaultInstancedPropertyBag->Value.GetMemory() : nullptr;
+				Value.InitializeAs(BagStruct, ValueMemoryDefault);
+			}
+			else
+			{
+				// Not the same struct definition. Do not use the default.
+				Value.InitializeAs(BagStruct);
+				DefaultBagStruct = nullptr;
+				DefaultInstancedPropertyBag = nullptr;
+			}
 
 			// Size of the serialized memory
 			int32 SerialSize = 0; 
@@ -3277,7 +3305,8 @@ bool FInstancedPropertyBag::Serialize(FArchive& Ar)
 			// BagStruct can be null if it contains structs, classes or enums that could not be found.
 			if (BagStruct != nullptr && Value.GetMutableMemory() != nullptr)
 			{
-				BagStruct->SerializeItem(Ar, Value.GetMutableMemory(), /*Defaults*/nullptr);
+				const void* DefaultMemory = DefaultInstancedPropertyBag ? DefaultInstancedPropertyBag->Value.GetMemory() : nullptr;
+				const_cast<UPropertyBag*>(BagStruct)->SerializeItem(Ar, Value.GetMutableMemory(), DefaultMemory);
 			}
 			else
 			{
@@ -3288,7 +3317,7 @@ bool FInstancedPropertyBag::Serialize(FArchive& Ar)
 		else if (Ar.IsSaving() || Ar.IsObjectReferenceCollector() || Ar.IsCountingMemory())
 		{
 			check(BagStruct);
-			
+
 			TArray<FPropertyBagPropertyDesc> PropertyDescs = BagStruct->PropertyDescs;
 #if WITH_EDITOR
 			// Save primary struct for user defined struct properties.
@@ -3318,7 +3347,8 @@ bool FInstancedPropertyBag::Serialize(FArchive& Ar)
 			const int64 InitialOffset = Ar.Tell(); // Position before struct serialization to compute its serial size
 
 			check(Value.GetMutableMemory() != nullptr);
-			BagStruct->SerializeItem(Ar, Value.GetMutableMemory(), /*Defaults*/nullptr);
+			const void* DefaultMemory = DefaultInstancedPropertyBag ? DefaultInstancedPropertyBag->Value.GetMemory() : nullptr;
+			const_cast<UPropertyBag*>(BagStruct)->SerializeItem(Ar, Value.GetMutableMemory(), DefaultMemory);
 		
 			const int64 FinalOffset = Ar.Tell(); // Keep current offset to reset the archive pos after write the serial size
 
diff --git a/Engine/Source/Runtime/CoreUObject/Public/StructUtils/PropertyBag.h b/Engine/Source/Runtime/CoreUObject/Public/StructUtils/PropertyBag.h
--- a/Engine/Source/Runtime/CoreUObject/Public/StructUtils/PropertyBag.h
+++ b/Engine/Source/Runtime/CoreUObject/Public/StructUtils/PropertyBag.h
@@ -52,6 +52,7 @@ enum class EPropertyBagContainerType : uint8
 
 namespace FPropertyBagCustomVersion
 {
+	UE_DEPRECATED(5.7, "Use a stream custom version. Data made with a custom version for feature do not merge between streams.")
 	COREUOBJECT_API extern const FGuid GUID;
 }
 
@@ -883,7 +884,7 @@ struct FInstancedPropertyBag
 	UE_API TValueOrError<const FPropertyBagSetRef, EPropertyBagResult> GetSetRef(const FName Name) const;
 
 	UE_API bool Identical(const FInstancedPropertyBag* Other, uint32 PortFlags) const;
-	UE_API bool Serialize(FArchive& Ar);
+	UE_API bool Serialize(FArchive& Ar, UStruct* DefaultsStruct = nullptr, const void* Defaults = nullptr);
 	UE_API void AddStructReferencedObjects(FReferenceCollector& Collector);
 	UE_API void GetPreloadDependencies(TArray<UObject*>& OutDeps);
 
