RemoteObjects:  When performing an IsValid on a WeakObjectPtr, we should return false if the UObject is marked as garbage since Get() will return nullptr.  To support this, CanResolveObject now takes a parameter that informs it how to handle different flags (currently we only have one potential flag: Garbage objects are non-resolvable).

#rnx

diff --git a/Engine/Source/Runtime/CoreUObject/Private/UObject/RemoteObject.cpp b/Engine/Source/Runtime/CoreUObject/Private/UObject/RemoteObject.cpp
--- a/Engine/Source/Runtime/CoreUObject/Private/UObject/RemoteObject.cpp
+++ b/Engine/Source/Runtime/CoreUObject/Private/UObject/RemoteObject.cpp
@@ -786,7 +786,7 @@ void TouchResidentObject(UObject* Object)
 	UE::RemoteObject::Transfer::TouchResidentObject(Object);
 }
 
-bool CanResolveObject(FRemoteObjectId ObjectId)
+bool CanResolveObject(FRemoteObjectId ObjectId, ECanResolveObjectBehavior CanResolveObjectBehavior)
 {
 	using namespace UE::RemoteObject::Private;
 
@@ -794,14 +794,27 @@ bool CanResolveObject(FRemoteObjectId ObjectId)
 	
 	if (UObject* Object = StaticFindObjectFastInternal(ObjectId))
 	{
+		// If we want to treat garbage objects as unresolvable, check the flag
+		if (CanResolveObjectBehavior == ECanResolveObjectBehavior::CannotResolveGarbageObjects)
+		{
+			return IsValid(Object);
+		}
+
 		// Object memory is local and even if it's already been migrated we can resolve it
 		return true;
 	}
 
 	if (FRemoteObjectStub* Stub = ObjectMaps->FindRemoteObjectStub(ObjectId))
 	{
-		// A stub exists so the object memory is not local but we can (attempt to) migrate it back if we're not garbage collecting
-		// Note: GIsGarbageCollecting checks are not thread safe
+		// A stub exists so the object memory is not local but we can (attempt to) migrate it back if:
+		// 1. It's not already Garbage (unless we're allowing a migration of Garbage objects)
+		if (CanResolveObjectBehavior == ECanResolveObjectBehavior::CannotResolveGarbageObjects && Stub->bWasGarbage)
+		{
+			return false;
+		}
+
+		// 2. We're not explicitly saying it's unsafe to migrate objects (FUnsafeToMigrateScope)
+		// 3. During Garbage Collection. Note: GIsGarbageCollecting checks are not thread safe
 		return IsSafeToMigrateObjects();
 	}
 
diff --git a/Engine/Source/Runtime/CoreUObject/Private/UObject/WeakObjectPtr.cpp b/Engine/Source/Runtime/CoreUObject/Private/UObject/WeakObjectPtr.cpp
--- a/Engine/Source/Runtime/CoreUObject/Private/UObject/WeakObjectPtr.cpp
+++ b/Engine/Source/Runtime/CoreUObject/Private/UObject/WeakObjectPtr.cpp
@@ -283,9 +283,9 @@ FUObjectItem* FWeakObjectPtr::TryResolveRemoteObject() const
 	return nullptr;
 }
 
-bool FWeakObjectPtr::CanBeResolved() const
+bool FWeakObjectPtr::CanBeResolved(UE::RemoteObject::Handle::ECanResolveObjectBehavior CanResolveObjectBehavior) const
 {
-	return UE::RemoteObject::Handle::CanResolveObject(ObjectRemoteId);
+	return UE::RemoteObject::Handle::CanResolveObject(ObjectRemoteId, CanResolveObjectBehavior);
 }
 
 #endif // UE_WITH_REMOTE_OBJECT_HANDLE
\ No newline at end of file
diff --git a/Engine/Source/Runtime/CoreUObject/Public/UObject/RemoteObject.h b/Engine/Source/Runtime/CoreUObject/Public/UObject/RemoteObject.h
--- a/Engine/Source/Runtime/CoreUObject/Public/UObject/RemoteObject.h
+++ b/Engine/Source/Runtime/CoreUObject/Public/UObject/RemoteObject.h
@@ -178,12 +178,12 @@ namespace UE::RemoteObject::Handle
 	COREUOBJECT_API UClass* GetClass(FRemoteObjectId ObjectId, ERemoteObjectGetClassBehavior GetClassBehavior);
 
 	/**
-	* Resolves a remote object
-	* @param Object Object id to resolve
-	* @param RefType Reference type that wants to resolve an object
-	* @return True if a remote object can be resolved
+	 * Determine if an Object can be resolved (exists in memory, or could be migrated)
+	 * @param ObjectId The ObjectId to consider for resolution
+	 * @param CanResolveObjectBehavior which behavior semantics should we consider when determining if an object is resolvable?
+	 * @return True if a remote object can be resolved using the given behavior semantics
 	*/
-	COREUOBJECT_API bool CanResolveObject(FRemoteObjectId ObjectId);
+	COREUOBJECT_API bool CanResolveObject(FRemoteObjectId ObjectId, ECanResolveObjectBehavior CanResolveObjectBehavior);
 
 	/**
 	* Checks if an object associated with the specified unique id is remote
diff --git a/Engine/Source/Runtime/CoreUObject/Public/UObject/RemoteObjectTypes.h b/Engine/Source/Runtime/CoreUObject/Public/UObject/RemoteObjectTypes.h
--- a/Engine/Source/Runtime/CoreUObject/Public/UObject/RemoteObjectTypes.h
+++ b/Engine/Source/Runtime/CoreUObject/Public/UObject/RemoteObjectTypes.h
@@ -22,6 +22,17 @@ namespace UE::RemoteObject::Private
 {
 	struct FRemoteIdLocalizationHelper;
 }
+namespace UE::RemoteObject::Handle
+{
+	enum class ECanResolveObjectBehavior
+	{
+		/** Do not consider any state when considering if an Object is resolvable */
+		Default,
+
+		/** If an Object is garbage, say it is not resolvable (optimize away nullptr resolves) */
+		CannotResolveGarbageObjects
+	};
+}
 
 constexpr inline int32 REMOTE_OBJECT_SERVER_ID_BIT_SIZE = 10;
 constexpr inline int32 REMOTE_OBJECT_SERIAL_NUMBER_BIT_SIZE = 53;
diff --git a/Engine/Source/Runtime/CoreUObject/Public/UObject/WeakObjectPtr.h b/Engine/Source/Runtime/CoreUObject/Public/UObject/WeakObjectPtr.h
--- a/Engine/Source/Runtime/CoreUObject/Public/UObject/WeakObjectPtr.h
+++ b/Engine/Source/Runtime/CoreUObject/Public/UObject/WeakObjectPtr.h
@@ -392,7 +392,7 @@ private:
 	
 #if UE_WITH_REMOTE_OBJECT_HANDLE
 	COREUOBJECT_API FUObjectItem* TryResolveRemoteObject() const;
-	COREUOBJECT_API bool CanBeResolved() const;
+	COREUOBJECT_API bool CanBeResolved(UE::RemoteObject::Handle::ECanResolveObjectBehavior CanResolveObjectBehavior) const;
 #endif
 
 	/**  
@@ -513,7 +513,14 @@ private:
 		FUObjectItem* ObjectItem = GUObjectArray.IndexToObject(ObjectIndex);
 		if (!ObjectItem || !SerialNumbersMatch(ObjectItem))
 		{
-			return CanBeResolved();
+			if (bEvenIfGarbage)
+			{
+				return CanBeResolved(UE::RemoteObject::Handle::ECanResolveObjectBehavior::Default);
+			}
+			else
+			{
+				return CanBeResolved(UE::RemoteObject::Handle::ECanResolveObjectBehavior::CannotResolveGarbageObjects);
+			}
 		}
 		else if (bThreadsafeTest)
 		{
