Moving ro.InitRemoteServerIdBeforeUObjectInit out of [ConsoleVariables] to [Core.System.Experimental]. The setting can't be changed at runtime so there's no point in keeping it as a console variable

diff --git a/Engine/Source/Runtime/CoreUObject/Private/UObject/RemoteObject.cpp b/Engine/Source/Runtime/CoreUObject/Private/UObject/RemoteObject.cpp
--- a/Engine/Source/Runtime/CoreUObject/Private/UObject/RemoteObject.cpp
+++ b/Engine/Source/Runtime/CoreUObject/Private/UObject/RemoteObject.cpp
@@ -28,11 +28,7 @@
 
 DEFINE_LOG_CATEGORY(LogRemoteObject);
 
-int32 GInitRemoteServerIdBeforeUObjectInit = 0;
-static FAutoConsoleVariableRef CVarInitRemoteServerIdBeforeUObjectInit(
-	TEXT("ro.InitRemoteServerIdBeforeUObjectInit"),
-	GInitRemoteServerIdBeforeUObjectInit,
-	TEXT("Initializes remote server id before UObject system is initialized"));
+
 
 int32 GRemoteIdToStringVerbosity = (int32)ERemoteIdToStringVerbosity::Id;
 static FAutoConsoleVariableRef CVarRemoteIdToStringVerbosity(
@@ -80,6 +76,7 @@ namespace UE::RemoteObject::Private
 std::atomic<uint64> RemoteObjectSerialNumber(1);
 std::atomic<uint64> AssetObjectSerialNumber(1);
 int32 UnsafeToMigrateObjects = 0; // This should go into TLS
+int32 GInitRemoteServerIdBeforeUObjectInit = 0; // Initializes remote server id before UObject system is initialized
 
 class FRemoteObjectStubMap : private TMap<FRemoteObjectId, UE::RemoteObject::Handle::FRemoteObjectStub*>
 {
@@ -306,7 +303,7 @@ void InitRemoteObjects()
 	if (GConfig)
 	{
 		// Get the config value now because console vars are updated automatically after InitRemoteObjects is called
-		GConfig->GetInt(TEXT("ConsoleVariables"), TEXT("ro.InitRemoteServerIdBeforeUObjectInit"), GInitRemoteServerIdBeforeUObjectInit, GEngineIni);
+		GConfig->GetInt(TEXT("[Core.System.Experimental]"), TEXT("ro.InitRemoteServerIdBeforeUObjectInit"), GInitRemoteServerIdBeforeUObjectInit, GEngineIni);
 	}
 	if (GInitRemoteServerIdBeforeUObjectInit)
 	{
