Fixing config section name

diff --git a/Engine/Source/Runtime/CoreUObject/Private/UObject/RemoteObject.cpp b/Engine/Source/Runtime/CoreUObject/Private/UObject/RemoteObject.cpp
--- a/Engine/Source/Runtime/CoreUObject/Private/UObject/RemoteObject.cpp
+++ b/Engine/Source/Runtime/CoreUObject/Private/UObject/RemoteObject.cpp
@@ -302,8 +302,8 @@ void InitRemoteObjects()
 	
 	if (GConfig)
 	{
-		// Get the config value now because console vars are updated automatically after InitRemoteObjects is called
-		GConfig->GetInt(TEXT("[Core.System.Experimental]"), TEXT("ro.InitRemoteServerIdBeforeUObjectInit"), GInitRemoteServerIdBeforeUObjectInit, GEngineIni);
+		// Get the config value now before it's first used
+		GConfig->GetInt(TEXT("Core.System.Experimental"), TEXT("ro.InitRemoteServerIdBeforeUObjectInit"), GInitRemoteServerIdBeforeUObjectInit, GEngineIni);
 	}
 	if (GInitRemoteServerIdBeforeUObjectInit)
 	{
