[Feature] Initial work on Instance Data Transforms
- created IDT repository
- implemented code responsible for reading IDT jsons from class packages and running them on applicable IDOs.
- implented RemoveProperty transform method
- implemented RedirectProperty transform method

diff --git a/Engine/Source/Runtime/CoreUObject/Internal/UObject/InstanceDataTransforms.h b/Engine/Source/Runtime/CoreUObject/Internal/UObject/InstanceDataTransforms.h
new file mode 100644
--- /dev/null
+++ b/Engine/Source/Runtime/CoreUObject/Internal/UObject/InstanceDataTransforms.h
@@ -0,0 +1,19 @@
+// Copyright Epic Games, Inc. All Rights Reserved.
+
+#pragma once
+#include "Containers/UnrealString.h"
+#include "Containers/Map.h"
+#include "Templates/Function.h"
+#include "Templates/SharedPointer.h"
+
+#if !UE_BUILD_SHIPPING_WITH_EDITOR && WITH_EDITORONLY_DATA 
+
+class UObject;
+class FJsonObject;
+
+namespace UE
+{
+	COREUOBJECT_API bool PatchInstanceDataObject(TObjectPtr<UObject>& IDO, TObjectPtr<UObject>& Owner);
+}
+
+#endif // !UE_BUILD_SHIPPING_WITH_EDITOR && WITH_EDITORONLY_DATA 
\ No newline at end of file
diff --git a/Engine/Source/Runtime/CoreUObject/Internal/UObject/PropertyBagRepository.h b/Engine/Source/Runtime/CoreUObject/Internal/UObject/PropertyBagRepository.h
--- a/Engine/Source/Runtime/CoreUObject/Internal/UObject/PropertyBagRepository.h
+++ b/Engine/Source/Runtime/CoreUObject/Internal/UObject/PropertyBagRepository.h
@@ -84,6 +84,9 @@ public:
 	 */
 	UObject* CreateInstanceDataObject(UObject* Owner, FArchive& Archive, int64 StartOffset, int64 EndOffset, bool bIsArchetype = false);
 
+	/** Replace an existing Ido with a new one who's class is generated from the cached unknown property tree */
+	UObject* ReinstanceWithCachedUnknownPropertyTree(UObject* Owner);
+	
 	// called at the end of postload to copy data from Owner to its IDO
 	COREUOBJECT_API void PostLoadInstanceDataObject(const UObject* Owner);
 
diff --git a/Engine/Source/Runtime/CoreUObject/Private/UObject/InstanceDataTransforms.cpp b/Engine/Source/Runtime/CoreUObject/Private/UObject/InstanceDataTransforms.cpp
new file mode 100644
--- /dev/null
+++ b/Engine/Source/Runtime/CoreUObject/Private/UObject/InstanceDataTransforms.cpp
@@ -0,0 +1,472 @@
+// Copyright Epic Games, Inc. All Rights Reserved.
+
+#include "UObject/InstanceDataTransforms.h"
+#if !UE_BUILD_SHIPPING_WITH_EDITOR && WITH_EDITORONLY_DATA 
+#include "UObject/PropertyPathNameTree.h"
+#include "UObject/PropertyStateTracking.h"
+#include "UObject/InstanceDataObjectUtils.h"
+#include "UObject/PropertyPathFunctions.h"
+#include "UObject/UObjectThreadContext.h"
+#include "Serialization/ObjectReader.h"
+#include "Serialization/ObjectWriter.h"
+#include "Misc/StringOutputDevice.h"
+#include "Templates/UnrealTemplate.h"
+#include "UObject/StrProperty.h"
+#include "UObject/TextProperty.h"
+#include "UObject/PropertyBagRepository.h"
+#include "UObject/Package.h"
+#include "Dom/JsonValue.h"
+#include "Dom/JsonObject.h"
+#include "Serialization/JsonReader.h"
+#include "Serialization/JsonSerializer.h"
+#include "HAL/FileManager.h"
+#include "Misc/DateTime.h"
+#include "Misc/FileHelper.h"
+#include "Misc/Paths.h"
+#include "Misc/PackageName.h"
+#include "VerseVM/VVMNames.h"
+
+
+DEFINE_LOG_CATEGORY_STATIC(LogInstanceDataTransforms, Display, Display);
+
+// registers a method that can be looked up by a string name in FInstanceDataTransformRepository::RegisteredOperations
+#define REGISTER_IDT_OPERATION(FunctionName) \
+UE::Private::FRegesterInstanceDataTransformHelper _Registered##FunctionName{TEXT(#FunctionName), FunctionName}; \
+
+
+namespace UE::Private
+{
+	enum EInstanceDataTransformResult
+	{
+		Break,
+		Continue
+	};
+
+	typedef TSharedPtr<FJsonObject> FInstanceDataTransformParams;
+	typedef TFunction<EInstanceDataTransformResult(TObjectPtr<UObject>& IDO, TObjectPtr<UObject>& Owner, const FInstanceDataTransformParams&)> FInstanceDataTransformFunction;
+
+	struct FInstanceDataTransform
+	{
+		FInstanceDataTransformParams Args;
+		FInstanceDataTransformFunction Func;
+	};
+
+	class FInstanceDataTransformRepository
+	{
+	public:
+		static FInstanceDataTransformRepository& Get()
+		{
+			static FInstanceDataTransformRepository Instance;
+			return Instance;
+		}
+
+		const TArray<TSharedPtr<FInstanceDataTransform>>* FindTransforms(const FTopLevelAssetPath& ClassPath)
+		{
+			// the transform config for a class is found in <RootPackageDirectory>/Config/InstanceDataTransforms.json
+			// TODO: I need to some research on wheather any work needs to be done in regards to making sure this file is treated properly by source control and the content worker.
+			FPackagePath PackagePath = FPackagePath::FromPackageNameChecked(ClassPath.GetPackageName());
+			FString Directory = FPaths::GetPath(PackagePath.GetLocalFullPath());
+			if (int32 Found = Directory.Find(TEXT("Content")); Found != INDEX_NONE)
+			{
+				Directory.LeftInline(Found);
+			}
+			Directory /= TEXT("Config");
+			
+			FString ConfigFilepath = Directory / TEXT("InstanceDataTransforms.json");
+			if (IFileManager::Get().FileExists(*ConfigFilepath))
+			{
+				FDateTime FileModificationTime = IFileManager::Get().GetTimeStamp(*ConfigFilepath);
+
+				if (FDateTime* LastReadTime = ConfigFileLastReadTime.Find(ConfigFilepath))
+				{
+					// check if the config file has been updated since the last time we read it
+					if (*LastReadTime < FileModificationTime)
+					{
+						ClassPathToIDTs.Empty(); // kill all cache
+						ReadConfigFile(ConfigFilepath);
+					}
+				}
+				else
+				{
+					ReadConfigFile(ConfigFilepath);
+				}
+			}
+			return ClassPathToIDTs.Find(ClassPath);
+		}
+
+		const TArray<TSharedPtr<FInstanceDataTransform>>* FindTransforms(const TNotNull<UClass*> Class)
+		{
+			return FindTransforms(Class->GetClassPathName());
+		}
+	private:
+
+		void ReadConfigFile(const FString& ConfigFilepath)
+		{
+			TSharedPtr<FJsonObject> ConfigJson;
+			FString JsonRaw;
+			if (FFileHelper::LoadFileToString(JsonRaw, *ConfigFilepath))
+			{
+				TSharedRef<TJsonReader<TCHAR>> JsonReader = TJsonReaderFactory<TCHAR>::Create(JsonRaw);
+				if (FJsonSerializer::Deserialize(JsonReader, ConfigJson))
+				{
+					FString PackagePathString;
+					if (!ConfigJson->TryGetStringField(TEXT("ClassPackage"), PackagePathString))
+					{
+						UE_LOG(LogInstanceDataTransforms, Warning, TEXT("'ClassPackage' feild was missing from '%s'"), *ConfigFilepath);
+						return;
+					}
+
+					FPackagePath PackagePath;
+					if (!FPackagePath::TryFromPackageName(PackagePathString, PackagePath))
+					{
+						UE_LOG(LogInstanceDataTransforms, Warning, TEXT("'ClassPackage' feild was malformed in '%s'"), *ConfigFilepath);
+						return;
+					}
+
+					const TArray<TSharedPtr<FJsonValue>>* ClassTransformsJson = nullptr;
+					if (!ConfigJson->TryGetArrayField(TEXT("ClassTransforms"), ClassTransformsJson) || !ClassTransformsJson)
+					{
+						UE_LOG(LogInstanceDataTransforms, Warning, TEXT("'ClassTransforms' array was missing or malformed in '%s'"), *ConfigFilepath);
+						return;
+					}
+
+					for (const TSharedPtr<FJsonValue>& TransformInfoJson : *ClassTransformsJson)
+					{
+						TSharedPtr<FJsonObject>* TransformInfo;
+						if (!TransformInfoJson->TryGetObject(TransformInfo) || !TransformInfo)
+						{
+							UE_LOG(LogInstanceDataTransforms, Warning, TEXT("Found a malfomed object in the 'ClassTransforms' array in '%s'"), *ConfigFilepath);
+							continue;
+						}
+
+						FString ClassNameString;
+						if (!(*TransformInfo)->TryGetStringField(TEXT("Class"), ClassNameString))
+						{
+							UE_LOG(LogInstanceDataTransforms, Warning, TEXT("an element of the 'ClassTransforms' array in '%s' was missing a 'Class' field"), *ConfigFilepath);
+							continue;
+						}
+						FName ClassName = FName(ClassNameString);
+						FTopLevelAssetPath ClassPath(PackagePath.GetPackageFName(), ClassName);
+
+						FString OpCode;
+						if (!(*TransformInfo)->TryGetStringField(TEXT("OpCode"), OpCode))
+						{
+							UE_LOG(LogInstanceDataTransforms, Warning, TEXT("an element of the 'ClassTransforms' array in '%s' was missing a 'OpCode' field"), *ConfigFilepath);
+							continue;
+						}
+						const FInstanceDataTransformFunction* Func = RegisteredOperations.Find(OpCode);
+						if (!Func)
+						{
+							UE_LOG(LogInstanceDataTransforms, Warning, TEXT("The transform OpCode in '%s' was malformed. '%s' is not a valid operation."), *ConfigFilepath, *OpCode);
+							continue;
+						}
+						FInstanceDataTransformParams Params;
+						if (const FInstanceDataTransformParams* ParamsPtr; (*TransformInfo)->TryGetObjectField(TEXT("Params"), ParamsPtr))
+						{
+							Params = *ParamsPtr;
+						}
+						TArray<TSharedPtr<FInstanceDataTransform>>& ClassTransforms = ClassPathToIDTs.FindOrAdd(ClassPath);
+						ClassTransforms.Add(MakeShared<FInstanceDataTransform>(Params, *Func));
+					}
+				}
+			}
+			ConfigFileLastReadTime.Add(ConfigFilepath, FDateTime::Now());
+		}
+
+		friend struct FRegesterInstanceDataTransformHelper;
+		void RegisterOperation(FStringView Name, FInstanceDataTransformFunction Function)
+		{
+			RegisteredOperations.Add(Name.GetData(), Function);
+		}
+
+		TMap<FString, FDateTime> ConfigFileLastReadTime;
+		TMap<FTopLevelAssetPath, TArray<TSharedPtr<FInstanceDataTransform>>> ClassPathToIDTs;
+		TMap<FString, FInstanceDataTransformFunction> RegisteredOperations;
+	};
+
+	struct FRegesterInstanceDataTransformHelper
+	{
+		FRegesterInstanceDataTransformHelper(FStringView Name, FInstanceDataTransformFunction Function)
+		{
+			FInstanceDataTransformRepository::Get().RegisterOperation(Name, Function);
+		}
+	};
+
+
+	// TODO: this method allows dangerous and narrowing conversions. That's fine for now because the user explicitly asked to do so but in the future it could be good
+	// to include a flag in the RedirectProperty params that specifies weather the user is ok with narrowing conversions
+	bool TryCopyOrConvert(const FProperty* SourceProperty, const void* SourceContainer, int32 SourceIndex, FProperty* DestinationProperty, void* DestinationContainer, int32 DestinationIndex)
+	{
+		const void* SourceData = SourceProperty->ContainerPtrToValuePtr<void>(SourceContainer, SourceIndex == INDEX_NONE ? 0 : SourceIndex);
+		void* DestinationData = DestinationProperty->ContainerPtrToValuePtr<void>(DestinationContainer, DestinationIndex == INDEX_NONE ? 0 : DestinationIndex);
+
+		FUObjectSerializeContext* SerializeContext = FUObjectThreadContext::Get().GetSerializeContext();
+		TGuardValue<bool> ScopedImpersonateProperties(SerializeContext->bImpersonateProperties, true);
+		TArray<uint8> Buffer;
+		FObjectWriter ObjectWriter(Buffer); // using FObjectWriter so that object properties will serialize correctly
+		FStructuredArchiveFromArchive StructuredWriter(ObjectWriter);
+		// todo: handle static arrays
+		FPropertyTag WriteTag(const_cast<FProperty*>(SourceProperty), SourceIndex, (uint8*)SourceData);
+		StructuredWriter.GetSlot() << WriteTag;
+		WriteTag.SerializeTaggedProperty(StructuredWriter.GetSlot(), const_cast<FProperty*>(SourceProperty), (uint8*)SourceData, nullptr);
+
+		FObjectReader ObjectReader(Buffer);
+		FStructuredArchiveFromArchive StructuredReader(ObjectReader);
+
+		FPropertyTag ReadTag;
+		StructuredReader.GetSlot() << ReadTag;
+		ReadTag.SetProperty(DestinationProperty);
+		ReadTag.ArrayIndex = DestinationIndex;
+		ReadTag.Name = DestinationProperty->GetFName();
+
+		bool bResult = false;
+		switch (DestinationProperty->ConvertFromType(ReadTag, StructuredReader.GetSlot(), (uint8*)DestinationContainer, SourceProperty->GetOwnerStruct(), nullptr))
+		{
+		case EConvertFromTypeResult::UseSerializeItem:
+			if (SourceProperty->GetID() == DestinationProperty->GetID())
+			{
+				ReadTag.SerializeTaggedProperty(StructuredReader.GetSlot(), DestinationProperty, (uint8*)DestinationData, nullptr);
+				bResult = true;
+			}
+			break;
+		case EConvertFromTypeResult::Serialized:
+			bResult = true;
+			break;
+		case EConvertFromTypeResult::CannotConvert:
+			break;
+		case EConvertFromTypeResult::Converted:
+			bResult = true;
+			break;
+		}
+
+		if (!bResult)
+		{
+			bool bTryTextSerialize = false;
+
+			const auto IsStringType = [](const FProperty* Property)
+				{
+					static FName VerseStringName = TEXT("VerseStringProperty");
+					return Property->IsA<FStrProperty>() || Property->IsA<FTextProperty>() || Property->IsA<FNameProperty>() || Property->GetID() == VerseStringName;
+				};
+
+			if (IsStringType(SourceProperty) || IsStringType(DestinationProperty))
+			{
+				// if either property is a string, text, or name, use text serialization
+				bTryTextSerialize = true;
+			}
+			else if (const FStructProperty* SourceAsStructProperty = CastField<FStructProperty>(SourceProperty))
+			{
+				if (FStructProperty* DestinationAsStructProperty = CastField<FStructProperty>(DestinationProperty))
+				{
+					if (!SourceAsStructProperty->Struct->UseNativeSerialization() && !DestinationAsStructProperty->Struct->UseNativeSerialization())
+					{
+						// attempt to text serialize structs since ConvertFromType doesn't support them usually
+						bTryTextSerialize = true;
+					}
+				}
+			}
+
+			// use ExportText_Direct and ImportText_Direct
+			if (bTryTextSerialize)
+			{
+				FString StrBuffer;
+				SourceProperty->ExportText_Direct(StrBuffer, SourceData, nullptr, nullptr, PPF_None);
+				FStringOutputDevice ErrorOutput;
+				DestinationProperty->ImportText_Direct(*StrBuffer, DestinationData, nullptr, PPF_None, &ErrorOutput);
+				bResult = ErrorOutput.IsEmpty();
+			}
+		}
+
+		return bResult;
+	}
+
+	bool TryCopyOrConvert(const FPropertyValueInContainer& SourceValue, const FPropertyValueInContainer& DestValue)
+	{
+		return TryCopyOrConvert(
+			SourceValue.Property, SourceValue.Container, SourceValue.ArrayIndex,
+			const_cast<FProperty*>(DestValue.Property), DestValue.Container, DestValue.ArrayIndex);
+	}
+
+	bool InferMisingTypeInfo(FPropertyTypeName& Type, UStruct*& CurrentStruct)
+	{
+		// fix structs that are missing their guids
+		if (Type.GetName() == NAME_StructProperty && Type.GetParameterCount() == 1)
+		{
+			const UE::FPropertyTypeName& StructType = Type.GetParameter(0);
+			if (StructType.GetParameterCount() == 1)
+			{
+				FName StructName = StructType.GetName();
+				FName PackageName = StructType.GetParameter(0).GetName();
+				FTopLevelAssetPath StructAssetPath(PackageName, StructName);
+				if (UScriptStruct* Struct = Cast<UScriptStruct>(StaticFindObject(UStruct::StaticClass(), StructAssetPath)))
+				{
+					// recompile the type with the guid included
+					FPropertyTypeNameBuilder Builder;
+					Builder.AddName(NAME_StructProperty);
+					Builder.BeginParameters();
+					Builder.AddPath(Struct);
+					Builder.AddGuid(Struct->GetCustomGuid());
+					Builder.EndParameters();
+					Type = Builder.Build();
+					CurrentStruct = Struct;
+					return true;
+				}
+			}
+		}
+		return false;
+	}
+
+	void InferMissingPathInfo(FPropertyPathName& Path, UStruct* CurrentStruct)
+	{
+		for (int32 SegmentIndex = 0; SegmentIndex < Path.GetSegmentCount(); ++SegmentIndex)
+		{
+			FPropertyPathNameSegment Segment = Path.GetSegment(SegmentIndex);
+			bool bUpdateSegment = false;
+			if (FPackageName::IsVersePackage(CurrentStruct->GetPackage()->GetName()))
+			{
+				FString PropertyName = Segment.Name.ToString();
+				if (!PropertyName.StartsWith(TEXT("__verse_0x")))
+				{
+					PropertyName = Verse::Names::Private::MangleCasedName(PropertyName, PropertyName);
+					Segment.Name = FName(PropertyName);
+					bUpdateSegment = true;
+				}
+			}
+			bUpdateSegment |= InferMisingTypeInfo(Segment.Type, CurrentStruct);
+
+			if (bUpdateSegment)
+			{
+				Path.SetSegment(SegmentIndex, Segment);
+			}
+		}
+	}
+}
+
+#if WITH_EDITOR
+namespace UE::Private::IDT
+{
+	/**
+	 * Remove a source FPropertyPathName from an IDO if it's an unknown property.
+	 * @param IDO - The instance data object that contains the source unknown property
+	 * @param Owner - the instance associated with the IDO
+	 * @param Params - contains a map with the following parameters
+	 *		- ["src"] the source FPropertyPathName as a json object. If this doesn't point to an unknown property, the operation will be skipped
+	 */
+	EInstanceDataTransformResult RemoveProperty(TObjectPtr<UObject>& IDO, TObjectPtr<UObject>& Owner, const FInstanceDataTransformParams& Params)
+	{
+		FUObjectSerializeContext* SerializeContext = FUObjectThreadContext::Get().GetSerializeContext();
+		TGuardValue<bool> ScopedImpersonateProperties(SerializeContext->bImpersonateProperties, true);
+
+		TSharedPtr<FPropertyPathNameTree> UnknownPropertyTree = FUnknownPropertyTree(Owner).FindOrCreate();
+		FUnknownEnumNames UnknownEnumNames = FUnknownEnumNames(Owner);
+
+		FPropertyPathName SrcPath;
+		if (!ensureAlwaysMsgf(Params && FPropertyPathName::FromJsonValue(Params->TryGetField(TEXT("src")), SrcPath),
+			TEXT("Failed to parse src parameter for a RemoveProperty IDT operation")))
+		{
+			return EInstanceDataTransformResult::Break;
+		}
+		InferMissingPathInfo(SrcPath, Owner->GetClass());
+
+		if (UnknownPropertyTree->Remove(SrcPath))
+		{
+			// create a new IDO with the modified property tree
+			IDO = UE::FPropertyBagRepository::Get().ReinstanceWithCachedUnknownPropertyTree(Owner);
+		}
+
+		return EInstanceDataTransformResult::Continue;
+	};
+	REGISTER_IDT_OPERATION(RemoveProperty);
+
+	/**  
+	 * Map a source FPropertyPathName into a destination FPropertyPathName relative to an object.
+	 * @param IDO - The instance data object that contains the source unknown property
+	 * @param Owner - the instance associated with the IDO
+	 * @param Params - contains a map with the following parameters
+	 *		- ["src"] the source FPropertyPathName as a json object. If this doesn't point to an unknown property, the operation will be skipped
+	 *		- ["dst"] the destination FPropertyPathName as a json object. If this points to a non-existent property, the path will be added to the new IDO instance as an unknown property
+	 */
+	EInstanceDataTransformResult RedirectProperty(TObjectPtr<UObject>& IDO, TObjectPtr<UObject>& Owner, const FInstanceDataTransformParams& Params)
+	{
+		FUObjectSerializeContext* SerializeContext = FUObjectThreadContext::Get().GetSerializeContext();
+		TGuardValue<bool> ScopedImpersonateProperties(SerializeContext->bImpersonateProperties, true);
+
+		TSharedPtr<FPropertyPathNameTree> UnknownPropertyTree = FUnknownPropertyTree(Owner).FindOrCreate();
+		FUnknownEnumNames UnknownEnumNames = FUnknownEnumNames(Owner);
+
+		FPropertyPathName SrcPath;
+		if (!ensureAlwaysMsgf(Params && FPropertyPathName::FromJsonValue(Params->TryGetField(TEXT("src")), SrcPath),
+			TEXT("Failed to parse src parameter for a RedirectProperty IDT operation")))
+		{
+			return EInstanceDataTransformResult::Break;
+		}
+		InferMissingPathInfo(SrcPath, Owner->GetClass());
+
+		FPropertyPathName DstPath;
+		if (!ensureAlwaysMsgf(Params && FPropertyPathName::FromJsonValue(Params->TryGetField(TEXT("dst")), DstPath),
+			TEXT("Failed to parse dst parameter for a RedirectProperty IDT operation")))
+		{
+			return EInstanceDataTransformResult::Break;
+		}
+		InferMissingPathInfo(DstPath, Owner->GetClass());
+
+
+		if (UnknownPropertyTree->Remove(SrcPath))
+		{
+			FPropertyValueInContainer DstInstanceValue = TryResolvePropertyPath(DstPath, Owner);
+			if (!DstInstanceValue)
+			{
+				// DstPath is not in the Instance. This means we'll create a new unknown/loose property at DstPath
+				UnknownPropertyTree->Add(DstPath);
+			}
+
+			// find value in old location (in the old IDO)
+			FPropertyValueInContainer SrcValue = TryResolvePropertyPath(SrcPath, IDO);
+			check(SrcValue); // it should be impossible for SrcPath to be in the old unknown property tree and not be in the old IDO
+
+			// create a new IDO with the modified property tree
+			IDO = UE::FPropertyBagRepository::Get().ReinstanceWithCachedUnknownPropertyTree(Owner);
+
+			// find destination location in the new IDO
+			FEditPropertyChain IdoChangeChain;
+			TSharedPtr<FPropertyChangedChainEvent> IdoChangeEvent;
+			FPropertyValueInContainer DstIDOValue = TryResolvePropertyPathForChangeEvent(DstPath, IDO, EPropertyChangeType::ValueSet, IdoChangeEvent, IdoChangeChain);
+			check(DstIDOValue); // we should've just added this path to the IDO
+
+			// copy property from SrcValue to DstIDOValue (note that because we're using change events, the value will be changed in the instance as well)
+			IDO->PreEditChange(IdoChangeChain);
+			TryCopyOrConvert(SrcValue, DstIDOValue);
+			IDO->PostEditChangeChainProperty(*IdoChangeEvent);
+		}
+
+		return EInstanceDataTransformResult::Continue;
+	};
+	REGISTER_IDT_OPERATION(RedirectProperty);
+}
+#endif
+
+namespace UE
+{
+	bool PatchInstanceDataObject(TObjectPtr<UObject>& IDO, TObjectPtr<UObject>& Owner)
+	{
+		check(IDO);
+		check(Owner);
+
+		if (const TArray<TSharedPtr<Private::FInstanceDataTransform>>* IDTs = Private::FInstanceDataTransformRepository::Get().FindTransforms(Owner->GetClass()))
+		{
+			for (const TSharedPtr<Private::FInstanceDataTransform>& IDT : *IDTs)
+			{
+				Private::EInstanceDataTransformResult Result = IDT->Func(IDO, Owner, IDT->Args);
+				if (Result == Private::EInstanceDataTransformResult::Break)
+				{
+					break;
+				}
+			}
+		}
+
+		return true;
+	}
+
+}
+
+#endif // !UE_BUILD_SHIPPING_WITH_EDITOR && WITH_EDITORONLY_DATA 
\ No newline at end of file
diff --git a/Engine/Source/Runtime/CoreUObject/Private/UObject/PropertyBagRepository.cpp b/Engine/Source/Runtime/CoreUObject/Private/UObject/PropertyBagRepository.cpp
--- a/Engine/Source/Runtime/CoreUObject/Private/UObject/PropertyBagRepository.cpp
+++ b/Engine/Source/Runtime/CoreUObject/Private/UObject/PropertyBagRepository.cpp
@@ -13,6 +13,7 @@
 #include "Serialization/ArchiveCountMem.h"
 #include "UObject/GarbageCollection.h"
 #include "UObject/InstanceDataObjectUtils.h"
+#include "UObject/InstanceDataTransforms.h"
 #include "UObject/LinkerLoad.h"
 #include "UObject/Object.h"
 #include "UObject/Package.h"
@@ -28,6 +29,7 @@
 #include "HAL/IConsoleManager.h"
 #include "Misc/CommandLine.h"
 #endif
+#include <UObject/InstanceDataTransforms.h>
 
 DEFINE_LOG_CATEGORY_STATIC(LogPropertyBagRepository, Log, All);
 
@@ -358,8 +360,8 @@ public:
 				SourceMemoryPtr = SourceMemoryPtr ? SourceField.Get<FProperty>()->ContainerPtrToValuePtr<void>(SourceMemoryPtr) : nullptr;
 				
 				ResolvedChainNode = ResolveCurChainNode(ResolvedStruct, SourceField.Get<FProperty>());
-				ResolvedField = ResolvedChainNode->GetValue();
-				ResolvedMemoryPtr = ResolvedMemoryPtr ? ResolvedField.Get<FProperty>()->ContainerPtrToValuePtr<void>(ResolvedMemoryPtr) : nullptr;
+				ResolvedField = ResolvedChainNode ? ResolvedChainNode->GetValue() : nullptr;
+				ResolvedMemoryPtr = (ResolvedChainNode && ResolvedMemoryPtr) ? ResolvedField.Get<FProperty>()->ContainerPtrToValuePtr<void>(ResolvedMemoryPtr) : nullptr;
 			}
 			else if (UObject* ResolvedObject = ResolvedField.Get<UObject>())
 			{
@@ -532,11 +534,11 @@ public:
 			}
 			if (SourceEvent.PropertyChain.GetActiveNode() == SourceChainNode)
 			{
-				PropertyChain.SetActivePropertyNode(ResolvedChainNode->GetValue());
+				PropertyChain.SetActivePropertyNode(ResolvedChainNode ? ResolvedChainNode->GetValue() : nullptr);
 			}
 			if (SourceEvent.PropertyChain.GetActiveMemberNode() == SourceChainNode)
 			{
-				PropertyChain.SetActiveMemberPropertyNode(ResolvedChainNode->GetValue());
+				PropertyChain.SetActiveMemberPropertyNode(ResolvedChainNode ? ResolvedChainNode->GetValue() : nullptr);
 			}
 		}
 
@@ -549,6 +551,11 @@ public:
 
 	bool ConstructionComplete() const
 	{
+		if (!ResolvedField)
+		{
+			// a field in the chain wasn't in the destination class. This is an unrecoverable failure to resolve.
+			return true;
+		}
 		if (FProperty* SourceProperty = SourceField.Get<FProperty>())
 		{
 			if (SourceEvent.GetArrayIndex(SourceProperty->GetName()) != INDEX_NONE && SourceProperty->Owner.Get<FProperty>() == SourceEvent.Property)
@@ -558,17 +565,17 @@ public:
 			return SourceProperty == SourceEvent.Property;
 		}
 		check(SourceField.IsValid());
-		check(ResolvedField.IsValid());
 		return SourceEvent.PropertyChain.GetActiveMemberNode() == nullptr;
 	}
 
 	FEditPropertyChain::TDoubleLinkedListNode* ResolveCurChainNode(const UStruct* ResolvedStruct, const FProperty* SourceProperty)
 	{
-		FProperty* FoundProperty = FindPropertyByNameAndType(ResolvedStruct, SourceProperty);
-		check(FoundProperty);
-		
-		PropertyChain.AddTail(FoundProperty);
-		return PropertyChain.GetTail();
+		if (FProperty* FoundProperty = FindPropertyByNameAndType(ResolvedStruct, SourceProperty))
+		{
+			PropertyChain.AddTail(FoundProperty);
+			return PropertyChain.GetTail();
+		}
+		return nullptr;
 	}
 
 	static UObject* GetObjectRefFromProperty(const FObjectProperty* Property, const void* Memory, const UObject* OwningObject)
@@ -618,21 +625,19 @@ void FPropertyBagRepository::PostEditChangeChainProperty(UObject* Object, const
 	{
 		FRemappedChainEvent RemappedChainEvent(PropertyChangedEvent, Source, Dest);
 		
-		Dest->PreEditChange(RemappedChainEvent.PropertyChain);
 		
 		FProperty* SourceProperty = RemappedChainEvent.SourceField.Get<FProperty>();
 		FProperty* DestProperty = RemappedChainEvent.ResolvedField.Get<FProperty>();
 		if (SourceProperty && DestProperty && RemappedChainEvent.SourceMemoryPtr && RemappedChainEvent.ResolvedMemoryPtr)
 		{
+			Dest->PreEditChange(RemappedChainEvent.PropertyChain);
 			CopyProperty(
 				SourceProperty,
 				RemappedChainEvent.SourceMemoryPtr,
 				DestProperty,
 				RemappedChainEvent.ResolvedMemoryPtr);
-		}
-
-            
-		Dest->PostEditChangeChainProperty(RemappedChainEvent);
+			Dest->PostEditChangeChainProperty(RemappedChainEvent);
+		} 
 	};
 	
 	if (UObject* Ido = Get().FindInstanceDataObject(Object))
@@ -653,8 +658,12 @@ void FPropertyBagRepository::PostEditChangeChainProperty(UObject* Object, const
 void FPropertyBagRepository::PostLoadInstanceDataObject(const UObject* Owner)
 {
 	// fixups may have been applied to the instance during PostLoad and they need to be copied to its IDO
-	FPropertyBagRepositoryLock LockRepo(this);
-	if (FPropertyBagAssociationData* BagData = AssociatedData.Find(Owner))
+	FPropertyBagAssociationData* BagData = nullptr;
+	{
+		FPropertyBagRepositoryLock LockRepo(this);
+		BagData = AssociatedData.Find(Owner);
+	}
+	if (BagData)
 	{
 		if (BagData->InstanceDataObject)
 		{
@@ -674,6 +683,12 @@ void FPropertyBagRepository::PostLoadInstanceDataObject(const UObject* Owner)
 					PostLoadInstanceDataObject(InstancedSubObject);
 				}
 			}
+
+#if !UE_BUILD_SHIPPING_WITH_EDITOR
+			// Run IDT operations on this instance to recover data
+			TObjectPtr<UObject> InOutInstance = const_cast<UObject*>(Owner);
+			UE::PatchInstanceDataObject(BagData->InstanceDataObject, InOutInstance);
+#endif
 		}
 	}
 }
@@ -799,6 +814,58 @@ const UObject* FPropertyBagRepository::FindInstanceForDataObject(const UObject*
 	return Owner ? *Owner : nullptr;
 }
 
+UObject* FPropertyBagRepository::ReinstanceWithCachedUnknownPropertyTree(UObject* Owner)
+{
+	if (FPropertyBagRepositoryLock LockRepo(this); FPropertyBagAssociationData* BagData = AssociatedData.Find(Owner))
+	{
+		TSharedPtr<FPropertyPathNameTree> UnknownPropertyTree = FUnknownPropertyTree(Owner).FindOrCreate();
+		FUnknownEnumNames UnknownEnumNames = FUnknownEnumNames(Owner);
+
+		UObject* IDO = BagData->InstanceDataObject;
+		UClass* NewIdoClass = CreateInstanceDataObjectClass(UnknownPropertyTree.Get(), &UnknownEnumNames, Owner->GetClass(), IDO->GetClass()->GetOuter());
+
+		UObject* OwnerArchetype = nullptr;
+		if (FArchive* Archive = Owner->GetLinker())
+		{
+			OwnerArchetype = Archive->GetArchetypeFromLoader(Owner);
+		}
+		if (!OwnerArchetype)
+		{
+			OwnerArchetype = Owner->GetArchetype();
+		}
+		UObject* Template = GetTemplateForInstanceDataObject(Owner, OwnerArchetype, NewIdoClass);
+
+		// Construct the new IDO.
+		FStaticConstructObjectParameters ConstructParams(NewIdoClass);
+		ConstructParams.SetFlags |= EObjectFlags::RF_Transactional;
+		// Set the RF_ArchetypeObject flag on all IDO archetypes so that they can be identified as such.
+		if (IDO->HasAllFlags(EObjectFlags::RF_ArchetypeObject))
+		{
+			ConstructParams.SetFlags |= EObjectFlags::RF_ArchetypeObject;
+		}
+
+		ConstructParams.Name = MakeUniqueObjectName(IDO->GetOuter(), ConstructParams.Class, Owner->GetFName(), EUniqueObjectNameOptions::UniversallyUnique);
+		ConstructParams.Outer = IDO->GetOuter();
+		ConstructParams.Template = Template;
+
+		UObject* NewIDO = StaticConstructObject_Internal(ConstructParams);
+
+		CopyTaggedProperties(IDO, NewIDO);
+
+		// replace the old IDO with the new one in the repository
+		InstanceDataObjectToOwner.Remove(BagData->InstanceDataObject);
+		BagData->InstanceDataObject = NewIDO;
+		BagData->bNeedsFixup = UE::StructContainsLooseProperties(NewIdoClass);
+		InstanceDataObjectToOwner.Add(NewIDO, Owner);
+#if WITH_EDITOR
+		FCoreUObjectDelegates::OnObjectsReinstanced.Broadcast({{IDO, NewIDO}});
+#endif
+
+		return NewIDO;
+	}
+	return nullptr;
+}
+
 UObject* FPropertyBagRepository::CreateInstanceDataObject(UObject* Owner, FArchive& Archive, int64 StartOffset, int64 EndOffset, bool bIsArchetype)
 {
 	// Limit the scope of the lock during to the find because this calls itself recursively to handle archetypes.
