[CisFix] silence FPackagePath::GetLocalFullPath header extension warning

diff --git a/Engine/Source/Runtime/CoreUObject/Private/UObject/InstanceDataTransforms.cpp b/Engine/Source/Runtime/CoreUObject/Private/UObject/InstanceDataTransforms.cpp
--- a/Engine/Source/Runtime/CoreUObject/Private/UObject/InstanceDataTransforms.cpp
+++ b/Engine/Source/Runtime/CoreUObject/Private/UObject/InstanceDataTransforms.cpp
@@ -65,31 +65,35 @@ namespace UE::Private
 			// the transform config for a class is found in <RootPackageDirectory>/Config/InstanceDataTransforms.json
 			// TODO: I need to some research on wheather any work needs to be done in regards to making sure this file is treated properly by source control and the content worker.
 			FPackagePath PackagePath = FPackagePath::FromPackageNameChecked(ClassPath.GetPackageName());
-			FString Directory = FPaths::GetPath(PackagePath.GetLocalFullPath());
-			if (int32 Found = Directory.Find(TEXT("Content")); Found != INDEX_NONE)
+			if (PackagePath.IsMountedPath())
 			{
-				Directory.LeftInline(Found);
-			}
-			Directory /= TEXT("Config");
-			
-			FString ConfigFilepath = Directory / TEXT("InstanceDataTransforms.json");
-			if (IFileManager::Get().FileExists(*ConfigFilepath))
-			{
-				FDateTime FileModificationTime = IFileManager::Get().GetTimeStamp(*ConfigFilepath);
+				PackagePath.SetHeaderExtension(EPackageExtension::Asset); // silence missing extension warnings
+				FString Directory = FPaths::GetPath(PackagePath.GetLocalFullPath());
+				if (int32 Found = Directory.Find(TEXT("Content")); Found != INDEX_NONE)
+				{
+					Directory.LeftInline(Found);
+				}
+				Directory /= TEXT("Config");
 
-				if (FDateTime* LastReadTime = ConfigFileLastReadTime.Find(ConfigFilepath))
+				FString ConfigFilepath = Directory / TEXT("InstanceDataTransforms.json");
+				if (IFileManager::Get().FileExists(*ConfigFilepath))
 				{
-					// check if the config file has been updated since the last time we read it
-					if (*LastReadTime < FileModificationTime)
+					FDateTime FileModificationTime = IFileManager::Get().GetTimeStamp(*ConfigFilepath);
+
+					if (FDateTime* LastReadTime = ConfigFileLastReadTime.Find(ConfigFilepath))
+					{
+						// check if the config file has been updated since the last time we read it
+						if (*LastReadTime < FileModificationTime)
+						{
+							ClassPathToIDTs.Empty(); // kill all cache
+							ReadConfigFile(ConfigFilepath);
+						}
+					}
+					else
 					{
-						ClassPathToIDTs.Empty(); // kill all cache
 						ReadConfigFile(ConfigFilepath);
 					}
 				}
-				else
-				{
-					ReadConfigFile(ConfigFilepath);
-				}
 			}
 			return ClassPathToIDTs.Find(ClassPath);
 		}
