Made FFieldPath and TFieldPath use UEOps.h - also fixes a const correctness issue.
Added UE::PointerEqualityComparableWith concept, which describes a pair of types which are comparable as pointers.

diff --git a/Engine/Source/Runtime/Core/Public/Concepts/PointerEqualityComparableWith.h b/Engine/Source/Runtime/Core/Public/Concepts/PointerEqualityComparableWith.h
new file mode 100644
--- /dev/null
+++ b/Engine/Source/Runtime/Core/Public/Concepts/PointerEqualityComparableWith.h
@@ -0,0 +1,16 @@
+// Copyright Epic Games, Inc. All Rights Reserved.
+
+#pragma once
+
+namespace UE
+{
+	/**
+	 * Concept which describes two types which can be compared for equality as pointers.
+	 */
+	template <typename T, typename U>
+	concept CPointerEqualityComparableWith =
+		requires(T* Lhs, U* Rhs)
+		{
+			Lhs == Rhs;
+		};
+}
diff --git a/Engine/Source/Runtime/CoreUObject/Public/UObject/FieldPath.h b/Engine/Source/Runtime/CoreUObject/Public/UObject/FieldPath.h
--- a/Engine/Source/Runtime/CoreUObject/Public/UObject/FieldPath.h
+++ b/Engine/Source/Runtime/CoreUObject/Public/UObject/FieldPath.h
@@ -6,12 +6,14 @@
 
 #pragma once
 
+#include "Concepts/PointerEqualityComparableWith.h"
 #include "Containers/Array.h"
 #include "Containers/Map.h"
 #include "Containers/UnrealString.h"
 #include "CoreTypes.h"
 #include "HAL/ThreadSafeCounter.h"
 #include "Misc/AssertionMacros.h"
+#include "Misc/UEOps.h"
 #include "Templates/Requires.h"
 #include "Templates/TypeHash.h"
 #include "Templates/UnrealTemplate.h"
@@ -252,16 +254,11 @@ public:
 		Path.Empty();
 	}
 
-	UE_FORCEINLINE_HINT bool operator==(const FFieldPath& InOther) const
+	UE_FORCEINLINE_HINT bool UEOpEquals(const FFieldPath& InOther) const
 	{
 		return ResolvedOwner == InOther.ResolvedOwner && Path == InOther.Path;
 	}
 
-	UE_FORCEINLINE_HINT bool operator!=(const FFieldPath& InOther) const
-	{
-		return ResolvedOwner != InOther.ResolvedOwner || Path != InOther.Path;
-	}
-
 	COREUOBJECT_API FString ToString() const;
 
 	COREUOBJECT_API friend FArchive& operator<<(FArchive& Ar, FFieldPath& InOutPropertyPath);
@@ -429,121 +426,24 @@ public:
 	{
 		return Get();
 	}
-};
 
-/**
-* Compare weak pointers for equality
-* @param Lhs weak pointer to compare
-* @param Rhs weak pointer to compare
-**/
-template <typename LhsType, typename RhsType>
-UE_FORCEINLINE_HINT auto operator==(const TFieldPath<LhsType>& Lhs, const TFieldPath<LhsType>& Rhs)
-	-> decltype((LhsType*)nullptr == (RhsType*)nullptr)
-{
-	return *(const FFieldPath*)&Lhs == *(const FFieldPath*)&Rhs;
-}
-
-/**
-* Compare weak pointers for equality
-* @param Lhs weak pointer to compare
-* @param Rhs pointer to compare
-**/
-template <typename LhsType, typename RhsType>
-UE_FORCEINLINE_HINT auto operator==(const TFieldPath<LhsType>& Lhs, const RhsType* Rhs)
-	-> decltype((LhsType*)nullptr == Rhs)
-{
-	return Lhs.Get() == Rhs;
-}
-
-/**
-* Compare weak pointers for equality
-* @param Lhs pointer to compare
-* @param Rhs weak pointer to compare
-**/
-template <typename LhsType, typename RhsType>
-UE_FORCEINLINE_HINT auto operator==(const LhsType* Lhs, const TFieldPath<RhsType>& Rhs)
-	-> decltype(Lhs == (RhsType*)nullptr)
-{
-	return Lhs == Rhs.Get();
-}
-
-/**
-* Test weak pointer for null
-* @param Lhs pointer to test
-**/
-template <typename LhsType>
-UE_FORCEINLINE_HINT bool operator==(TFieldPath<LhsType>& Lhs, TYPE_OF_NULLPTR)
-{
-	return !Lhs.Get();
-}
-
-/**
-* Test weak pointer for null
-* @param Rhs pointer to test
-**/
-template <typename RhsType>
-UE_FORCEINLINE_HINT bool operator==(TYPE_OF_NULLPTR, TFieldPath<RhsType>& Rhs)
-{
-	return !Rhs.Get();
-}
-
-#if !PLATFORM_COMPILER_HAS_GENERATED_COMPARISON_OPERATORS
-/**
-* Compare weak pointers for inequality
-* @param Lhs weak pointer to compare
-* @param Rhs weak pointer to compare
-**/
-template <typename LhsType, typename RhsType>
-UE_FORCEINLINE_HINT auto operator!=(const TFieldPath<LhsType>& Lhs, const TFieldPath<LhsType>& Rhs)
-	-> decltype((LhsType*)nullptr != (RhsType*)nullptr)
-{
-	return !(Lhs == Rhs);
-}
-
-/**
-* Compare weak pointers for inequality
-* @param Lhs weak pointer to compare
-* @param Rhs pointer to compare
-**/
-template <typename LhsType, typename RhsType>
-UE_FORCEINLINE_HINT auto operator!=(const TFieldPath<LhsType>& Lhs, const RhsType* Rhs)
-	-> decltype((LhsType*)nullptr != Rhs)
-{
-	return !(Lhs == Rhs);
-}
-
-/**
-* Compare weak pointers for inequality
-* @param Lhs pointer to compare
-* @param Rhs weak pointer to compare
-**/
-template <typename LhsType, typename RhsType>
-UE_FORCEINLINE_HINT auto operator!=(const LhsType* Lhs, const TFieldPath<RhsType>& Rhs)
-	-> decltype(Lhs != (RhsType*)nullptr)
-{
-	return !(Lhs == Rhs);
-}
+	template <UE::CPointerEqualityComparableWith<PropertyType> OtherPropertyType>
+	bool UEOpEquals(const TFieldPath<OtherPropertyType>& Other) const
+	{
+		return FFieldPath::UEOpEquals(Other);
+	}
 
-/**
-* Test weak pointer for non-null
-* @param Lhs pointer to test
-**/
-template <typename LhsType>
-UE_FORCEINLINE_HINT bool operator!=(TFieldPath<LhsType>& Lhs, TYPE_OF_NULLPTR)
-{
-	return !(Lhs == nullptr);
-}
+	template <UE::CPointerEqualityComparableWith<PropertyType> T>
+	bool UEOpEquals(const T* Other) const
+	{
+		return Get() == Other;
+	}
 
-/**
-* Test weak pointers for non-null
-* @param Rhs pointer to test
-**/
-template <typename RhsType>
-UE_FORCEINLINE_HINT bool operator!=(TYPE_OF_NULLPTR, TFieldPath<RhsType>& Rhs)
-{
-	return !(nullptr == Rhs);
-}
-#endif
+	bool UEOpEquals(TYPE_OF_NULLPTR) const
+	{
+		return !Get();
+	}
+};
 
 // Helper function which deduces the type of the initializer
 template <typename PropertyType>
