Silence PVS warning
#rnx

diff --git a/Engine/Source/Runtime/CoreUObject/Private/Serialization/UnversionedPropertySerialization.cpp b/Engine/Source/Runtime/CoreUObject/Private/Serialization/UnversionedPropertySerialization.cpp
--- a/Engine/Source/Runtime/CoreUObject/Private/Serialization/UnversionedPropertySerialization.cpp
+++ b/Engine/Source/Runtime/CoreUObject/Private/Serialization/UnversionedPropertySerialization.cpp
@@ -705,7 +705,7 @@ protected:
 	{
 		check(NumBits > 0);
 
-		uint32 LastWordMask = ~0u >> ((32u - NumBits) % 32u);
+		uint32 LastWordMask = ~0u >> ((32u - NumBits) % 32u); //-V1064
 		if (NumBits <= 8)
 		{
 			uint8 Word = static_cast<uint8>(*Data & LastWordMask);
