Fixed some copy+pasted comments in FieldPath.h.

#rnx

diff --git a/Engine/Source/Runtime/CoreUObject/Public/UObject/FieldPath.h b/Engine/Source/Runtime/CoreUObject/Public/UObject/FieldPath.h
--- a/Engine/Source/Runtime/CoreUObject/Public/UObject/FieldPath.h
+++ b/Engine/Source/Runtime/CoreUObject/Public/UObject/FieldPath.h
@@ -1,7 +1,7 @@
 // Copyright Epic Games, Inc. All Rights Reserved.
 
 /*=============================================================================
-	SoftObjectPtr.h: Pointer to UObject asset, keeps extra information so that it is works even if the asset is not in memory
+	FieldPath.h: A pointer to a field and the path to it from its outermost owner.
 =============================================================================*/
 
 #pragma once
@@ -324,7 +324,7 @@ public:
 		: FFieldPath(InField, PropertyType::StaticClass()->GetFName())
 	{
 		// This static assert is in here rather than in the body of the class because we want
-		// to be able to define TWeakFieldPtr<UUndefinedClass>.
+		// to be able to define TFieldPath<UUndefinedClass>.
 		static_assert(std::is_convertible_v<PropertyType*, const volatile FField*>, "TFieldPath can only be constructed with FField types");
 	}
 #endif
