Specialized TBaseStructure for const.

diff --git a/Engine/Source/Runtime/CoreUObject/Public/UObject/Class.h b/Engine/Source/Runtime/CoreUObject/Public/UObject/Class.h
--- a/Engine/Source/Runtime/CoreUObject/Public/UObject/Class.h
+++ b/Engine/Source/Runtime/CoreUObject/Public/UObject/Class.h
@@ -5288,6 +5288,11 @@ struct TBaseStructure : TBaseStructureBase<T>
 {
 };
 
+template <typename T>
+struct TBaseStructure<const T> : TBaseStructureBase<T>
+{
+};
+
 template<> struct TBaseStructure<FIntPoint> 
 {
 	static COREUOBJECT_API UScriptStruct* Get(); 
