Matching TSAN_ATOMIC definition of GIsIncrementalReachabilityPending to its declaration

diff --git a/Engine/Source/Runtime/CoreUObject/Private/UObject/GarbageCollection.cpp b/Engine/Source/Runtime/CoreUObject/Private/UObject/GarbageCollection.cpp
--- a/Engine/Source/Runtime/CoreUObject/Private/UObject/GarbageCollection.cpp
+++ b/Engine/Source/Runtime/CoreUObject/Private/UObject/GarbageCollection.cpp
@@ -617,7 +617,7 @@ namespace UE::GC
 	/** EInternalObjectFlag value representing a maybe unreachable object */
 	EInternalObjectFlags GMaybeUnreachableObjectFlag = EInternalObjectFlags::ReachabilityFlag2;
 
-	bool GIsIncrementalReachabilityPending = false;
+	TSAN_ATOMIC(bool) GIsIncrementalReachabilityPending = false;
 } // namespace UE::GC
 
 namespace UE::GC::Private
