Fixed FVector3f fields not being exposed to Python.

diff --git a/Engine/Source/Runtime/CoreUObject/Public/UObject/NoExportTypes.h b/Engine/Source/Runtime/CoreUObject/Public/UObject/NoExportTypes.h
--- a/Engine/Source/Runtime/CoreUObject/Public/UObject/NoExportTypes.h
+++ b/Engine/Source/Runtime/CoreUObject/Public/UObject/NoExportTypes.h
@@ -540,13 +540,13 @@ struct FGuid
 USTRUCT(immutable, noexport, BlueprintType, BlueprintInternalUseOnly, IsAlwaysAccessible, HasDefaults, HasNoOpConstructor, IsCoreType, meta = (ScriptDefaultMake, ScriptDefaultBreak, HasNativeMake = "/Script/Engine.KismetMathLibrary.MakeVector", HasNativeBreak = "/Script/Engine.KismetMathLibrary.BreakVector"))
 struct FVector3f
 {
-	UPROPERTY(EditAnywhere, Category = Vector, SaveGame)
+	UPROPERTY(EditAnywhere, BlueprintReadWrite, Category = Vector, SaveGame)
 	float X;
 
-	UPROPERTY(EditAnywhere, Category = Vector, SaveGame)
+	UPROPERTY(EditAnywhere, BlueprintReadWrite, Category = Vector, SaveGame)
 	float Y;
 
-	UPROPERTY(EditAnywhere, Category = Vector, SaveGame)
+	UPROPERTY(EditAnywhere, BlueprintReadWrite, Category = Vector, SaveGame)
 	float Z;
 };
 
