AsyncLoading2: PartialFlushes: Reduce log volume of the logging of partial flushes by printing out each callstack only once.
#rnx

diff --git a/Engine/Source/Runtime/CoreUObject/Private/Serialization/AsyncLoading2.cpp b/Engine/Source/Runtime/CoreUObject/Private/Serialization/AsyncLoading2.cpp
--- a/Engine/Source/Runtime/CoreUObject/Private/Serialization/AsyncLoading2.cpp
+++ b/Engine/Source/Runtime/CoreUObject/Private/Serialization/AsyncLoading2.cpp
@@ -5053,7 +5053,7 @@ private:
 		struct FPartialFlushInfo
 		{
 			static constexpr int32 MaxCallstackDepth = 32;
-			TArray<uint64, TFixedAllocator<MaxCallstackDepth>> Callstack;
+			int32 UniqueCallstackIndex = -1;
 			FName RequestedPackage;
 			FName RequestingPackage;
 			EAsyncPackageLoadingState2 RequestedPackageState;
@@ -5080,9 +5080,9 @@ private:
 		};
 
 		void Clear();
-		void DumpToLog() const;
+		void DumpToLog();
 		void RecordPartialFlush(const FAsyncPackage2* RequestedPackage, const FAsyncPackage2* RequestingPackage);
-		void Log(FStringBuilderBase& Builder, const FPartialFlushInfo& PartialFlush) const;
+		void Log(FStringBuilderBase& Builder, const FPartialFlushInfo& PartialFlush);
 
 		inline void SetLogMode(ELogMode Mode)
 		{
@@ -5090,6 +5090,13 @@ private:
 			LogMode = Mode;
 #endif
 		}
+		struct FUniqueCallstack
+		{
+			TArray<uint64> Callstack;
+			bool bHasLogged = false;
+		};
+		TArray<FUniqueCallstack> UniqueCallstacks;
+		TMap<TArray<uint64>, int32> UniqueCallstackMap;
 
 #if !UE_BUILD_SHIPPING
 	private:
@@ -11952,7 +11959,7 @@ void FAsyncLoadingThread2::FLoaderStats::Clear()
 #endif
 }
 
-void FAsyncLoadingThread2::FLoaderStats::DumpToLog() const
+void FAsyncLoadingThread2::FLoaderStats::DumpToLog()
 {
 
 #if !UE_BUILD_SHIPPING
@@ -11988,11 +11995,21 @@ void FAsyncLoadingThread2::FLoaderStats::RecordPartialFlush(const FAsyncPackage2
 
 	if (LogMode == ELogMode::DumpWithCallstack)
 	{
-		PartialFlush.Callstack.SetNumUninitialized(FPartialFlushInfo::MaxCallstackDepth);
+		TArray<uint64> Callstack;
+		Callstack.SetNumUninitialized(FPartialFlushInfo::MaxCallstackDepth);
 
 		const int32 NumStackFramesToIgnore = 3; // Ignore this callstack and 2 frames from CaptureStackBackTrace
-		uint32 NumFrames = FPlatformStackWalk::CaptureStackBackTrace(PartialFlush.Callstack.GetData(), PartialFlush.Callstack.Num());
-		PartialFlush.Callstack.SetNum(NumFrames);
+		uint32 NumFrames = FPlatformStackWalk::CaptureStackBackTrace(Callstack.GetData(), Callstack.Num());
+		Callstack.SetNum(NumFrames);
+
+		int32& ExistingIndex = UniqueCallstackMap.FindOrAdd(Callstack, -1);
+		if (ExistingIndex == -1)
+		{
+			ExistingIndex = UniqueCallstacks.Num();
+			FUniqueCallstack& UniqueCallstack = UniqueCallstacks.Emplace_GetRef();
+			UniqueCallstack.Callstack = MoveTemp(Callstack);
+		}
+		PartialFlush.UniqueCallstackIndex = ExistingIndex;
 	}
 
 	if (LogMode == ELogMode::Inline)
@@ -12007,32 +12024,42 @@ void FAsyncLoadingThread2::FLoaderStats::RecordPartialFlush(const FAsyncPackage2
 #endif
 }
 
-void FAsyncLoadingThread2::FLoaderStats::Log(FStringBuilderBase& Builder, const FPartialFlushInfo& PartialFlush) const
+void FAsyncLoadingThread2::FLoaderStats::Log(FStringBuilderBase& Builder, const FPartialFlushInfo& PartialFlush)
 {
 #if !UE_BUILD_SHIPPING
+	FUniqueCallstack* Callstack = PartialFlush.UniqueCallstackIndex < 0 
+		? nullptr : &UniqueCallstacks[PartialFlush.UniqueCallstackIndex];
+
 	//
 	// Note: Update the FLoadingTests_RecursiveLoads_FullFlushFrom_Serialize test if you edit this error message
 	//
 	Builder.Reset();
 	Builder.Appendf(
-		TEXT("Partially loaded package %s (state: %s) to avoid deadlock when it was recursively flushed by package %s (state: %s). Callstack:\n"),
+		TEXT("Partially loaded package %s (state: %s) to avoid deadlock when it was recursively flushed by package %s (state: %s). UniqueCallstack(%d)%s"),
 		*PartialFlush.RequestedPackage.ToString(),
 		LexToString(PartialFlush.RequestedPackageState),
 		*PartialFlush.RequestingPackage.ToString(),
-		LexToString(PartialFlush.RequestingPackageState));
+		LexToString(PartialFlush.RequestingPackageState),
+		PartialFlush.UniqueCallstackIndex,
+		(!Callstack || Callstack->bHasLogged) ? TEXT(".") : TEXT(" - first occurrence - Callstack:\n")
+		);
 
-	FString HumanReadableString;
-	for (const uint64 ProgramCounter : PartialFlush.Callstack)
+	if (Callstack && !Callstack->bHasLogged)
 	{
-		FProgramCounterSymbolInfoEx SymbolInfo;
-		FPlatformStackWalk::ProgramCounterToSymbolInfoEx(ProgramCounter, SymbolInfo);
-		FPlatformStackWalk::SymbolInfoToHumanReadableStringEx(SymbolInfo, HumanReadableString);
+		Callstack->bHasLogged = true;
+		FString HumanReadableString;
+		for (const uint64 ProgramCounter : Callstack->Callstack)
+		{
+			FProgramCounterSymbolInfoEx SymbolInfo;
+			FPlatformStackWalk::ProgramCounterToSymbolInfoEx(ProgramCounter, SymbolInfo);
+			FPlatformStackWalk::SymbolInfoToHumanReadableStringEx(SymbolInfo, HumanReadableString);
 
-		Builder.Append(TEXT("\t"));
-		Builder.Append(HumanReadableString);
-		Builder.Append(TEXT("\n"));
+			Builder.Append(TEXT("\t"));
+			Builder.Append(HumanReadableString);
+			Builder.Append(TEXT("\n"));
 
-		HumanReadableString.Reset();
+			HumanReadableString.Reset();
+		}
 	}
 
 	UE_LOG(LogStreaming, Display, TEXT("%s"), Builder.ToString());
