[ZenLoader]
- Support for linkerload is always active as it is required to support loose cooked files at runtime

#FailedValidations Preflight

diff --git a/Engine/Source/Runtime/CoreUObject/Private/Serialization/AsyncLoading2.cpp b/Engine/Source/Runtime/CoreUObject/Private/Serialization/AsyncLoading2.cpp
--- a/Engine/Source/Runtime/CoreUObject/Private/Serialization/AsyncLoading2.cpp
+++ b/Engine/Source/Runtime/CoreUObject/Private/Serialization/AsyncLoading2.cpp
@@ -301,10 +301,6 @@ void FindAllRuntimeScriptPackages(FRuntimeScriptPackages& OutPackages)
 #define ALT2_VERIFY_LINKERLOAD_MATCHES_IMPORTSTORE 0
 #endif
 
-#ifndef ALT2_ENABLE_LINKERLOAD_SUPPORT
-#define ALT2_ENABLE_LINKERLOAD_SUPPORT 1
-#endif
-
 #ifndef ALT2_ENABLE_NEW_ARCHIVE_FOR_LINKERLOAD
 #define ALT2_ENABLE_NEW_ARCHIVE_FOR_LINKERLOAD 0
 #endif
@@ -565,10 +561,8 @@ TRACE_DECLARE_ATOMIC_INT_COUNTER(AsyncLoadingLoadingPackages, TEXT("AsyncLoading
 TRACE_DECLARE_ATOMIC_INT_COUNTER(AsyncLoadingPackagesWithRemainingWork, TEXT("AsyncLoading/PackagesWithRemainingWork"));
 TRACE_DECLARE_ATOMIC_INT_COUNTER(AsyncLoadingPendingIoRequests, TEXT("AsyncLoading/PendingIoRequests"));
 TRACE_DECLARE_ATOMIC_MEMORY_COUNTER(AsyncLoadingTotalLoaded, TEXT("AsyncLoading/TotalLoaded"));
-#if ALT2_ENABLE_LINKERLOAD_SUPPORT
 TRACE_DECLARE_ATOMIC_INT_COUNTER(AsyncLoadingFileSystemLoads, TEXT("AsyncLoadingThread/FileSystemLoads"));
 TRACE_DECLARE_ATOMIC_INT_COUNTER(AsyncLoadingIoStoreLoads, TEXT("AsyncLoadingThread/IoStoreLoads"));
-#endif
 
 FString FormatPackageId(FPackageId PackageId)
 {
@@ -2281,14 +2275,12 @@ struct FAsyncPackageHeaderData: public FZenPackageHeader
 	TArrayView<FExportBundleEntry> ExportBundleEntriesCopyForPostLoad; // TODO: Can we use ConstructedObjects or Exports instead for posloading?
 };
 
-#if ALT2_ENABLE_LINKERLOAD_SUPPORT
 struct FAsyncPackageLinkerLoadHeaderData
 {
 	TArray<uint64> ImportedPublicExportHashes;
 	TArray<FPackageObjectIndex> ImportMap;
 	TArray<FExportMapEntry> ExportMap; // Note: Currently only using the public export hash field
 };
-#endif
 
 struct FPackageImportStore
 {
@@ -2930,13 +2922,11 @@ enum class EAsyncPackageLoadingState2 : uint8
 	ProcessDependencies,
 	WaitingForDependencies,
 	DependenciesReady,
-#if ALT2_ENABLE_LINKERLOAD_SUPPORT
 // This is the path taken by LinkerLoad packages
 	CreateLinkerLoadExports,
 	WaitingForLinkerLoadDependencies,
 	ResolveLinkerLoadImports,
 	PreloadLinkerLoadExports,
-#endif
 // This is the path taken by Runtime/cooked packages
 	ProcessExportBundles,
 // Both LinkerLoad and Cooked packages should converge at this point
@@ -2962,12 +2952,10 @@ const TCHAR* LexToString(EAsyncPackageLoadingState2 AsyncPackageLoadingState)
 	case EAsyncPackageLoadingState2::ProcessDependencies: return TEXT("ProcessDependencies");
 	case EAsyncPackageLoadingState2::WaitingForDependencies: return TEXT("WaitingForDependencies");
 	case EAsyncPackageLoadingState2::DependenciesReady: return TEXT("DependenciesReady");
-#if ALT2_ENABLE_LINKERLOAD_SUPPORT
 	case EAsyncPackageLoadingState2::CreateLinkerLoadExports: return TEXT("CreateLinkerLoadExports");
 	case EAsyncPackageLoadingState2::WaitingForLinkerLoadDependencies: return TEXT("WaitingForLinkerLoadDependencies");
 	case EAsyncPackageLoadingState2::ResolveLinkerLoadImports: return TEXT("ResolveLinkerLoadImports");
 	case EAsyncPackageLoadingState2::PreloadLinkerLoadExports: return TEXT("PreloadLinkerLoadExports");
-#endif
 	case EAsyncPackageLoadingState2::ProcessExportBundles: return TEXT("ProcessExportBundles");
 	case EAsyncPackageLoadingState2::WaitingForExternalReads: return TEXT("WaitingForExternalReads");
 	case EAsyncPackageLoadingState2::ExportsDone: return TEXT("ExportsDone");
@@ -3305,11 +3293,9 @@ enum EEventLoadNode2 : uint8
 	Package_ProcessSummary,
 	Package_ProcessDependencies,
 	Package_DependenciesReady,
-#if ALT2_ENABLE_LINKERLOAD_SUPPORT
 	Package_CreateLinkerLoadExports,
 	Package_ResolveLinkerLoadImports,
 	Package_PreloadLinkerLoadExports,
-#endif
 	Package_ExportsSerialized,
 	Package_NumPhases,
 
@@ -3688,11 +3674,9 @@ struct FAsyncPackage2 : public FRefCountBase
 	void ImportPackagesRecursive(FAsyncLoadingThreadState2& ThreadState, FIoBatch& IoBatch, FPackageStore& PackageStore);
 	void StartLoading(FAsyncLoadingThreadState2& ThreadState, FIoBatch& IoBatch);
 
-#if ALT2_ENABLE_LINKERLOAD_SUPPORT
 	void InitializeLinkerLoadState(const FLinkerInstancingContext* InstancingContext);
 	void CreateLinker(const FLinkerInstancingContext* InstancingContext);
 	void DetachLinker();
-#endif
 
 private:
 	void ImportPackagesRecursiveInner(FAsyncLoadingThreadState2& ThreadState, FIoBatch& IoBatch, FPackageStore& PackageStore, FAsyncPackageHeaderData& Header);
@@ -3711,7 +3695,7 @@ private:
 	TOptional<FAsyncPackageSerializationState> OptionalSegmentSerializationState;
 	bool bRequestOptionalChunk = false;
 #endif
-#if ALT2_ENABLE_LINKERLOAD_SUPPORT
+
 	struct FLinkerLoadState
 	{
 		FLinkerLoad* Linker = nullptr;
@@ -3730,7 +3714,7 @@ private:
 		FAsyncPackageLinkerLoadHeaderData LinkerLoadHeaderData;
 	};
 	TOptional<FLinkerLoadState> LinkerLoadState;
-#endif // ALT2_ENABLE_LINKERLOAD_SUPPORT
+
 	/** Cached async loading thread object this package was created by */
 	FAsyncLoadingThread2& AsyncLoadingThread;
 	FAsyncLoadEventGraphAllocator& GraphAllocator;
@@ -3855,9 +3839,7 @@ private:
 		}
 	};
 	FAllDependenciesState		AllDependenciesSetupState;
-#if ALT2_ENABLE_LINKERLOAD_SUPPORT
 	FAllDependenciesState		AllDependenciesImportState;
-#endif
 	FAllDependenciesState		AllDependenciesPostLoadReadyState;
 	FAllDependenciesState		AllDependenciesFullyLoadedState;
 
@@ -3896,11 +3878,9 @@ public:
 	/** [EDL] Begin Event driven loader specific stuff */
 
 	static EEventLoadNodeExecutionResult Event_ProcessExportBundle(FAsyncLoadingThreadState2& ThreadState, FAsyncPackage2* Package, int32 ExportBundleIndex);
-#if ALT2_ENABLE_LINKERLOAD_SUPPORT
 	static EEventLoadNodeExecutionResult Event_CreateLinkerLoadExports(FAsyncLoadingThreadState2& ThreadState, FAsyncPackage2* Package, int32);
 	static EEventLoadNodeExecutionResult Event_ResolveLinkerLoadImports(FAsyncLoadingThreadState2& ThreadState, FAsyncPackage2* Package, int32);
 	static EEventLoadNodeExecutionResult Event_PreloadLinkerLoadExports(FAsyncLoadingThreadState2& ThreadState, FAsyncPackage2* Package, int32);
-#endif
 	static EEventLoadNodeExecutionResult Event_ProcessPackageSummary(FAsyncLoadingThreadState2& ThreadState, FAsyncPackage2* Package, int32);
 	static EEventLoadNodeExecutionResult Event_ProcessDependencies(FAsyncLoadingThreadState2& ThreadState, FAsyncPackage2* Package, int32);
 	static EEventLoadNodeExecutionResult Event_DependenciesReady(FAsyncLoadingThreadState2& ThreadState, FAsyncPackage2* Package, int32);
@@ -3971,7 +3951,6 @@ private:
 	void ConditionalReleasePartialRequests(FAsyncLoadingThreadState2& ThreadState);
 	void ConditionalBeginPostLoad(FAsyncLoadingThreadState2& ThreadState);
 	void ConditionalBeginDeferredPostLoad(FAsyncLoadingThreadState2& ThreadState);
-#if ALT2_ENABLE_LINKERLOAD_SUPPORT
 	EEventLoadNodeExecutionResult ProcessLinkerLoadPackageSummary(FAsyncLoadingThreadState2& ThreadState);
 	bool CreateLinkerLoadExports(FAsyncLoadingThreadState2& ThreadState);
 	void ConditionalBeginResolveLinkerLoadImports(FAsyncLoadingThreadState2& ThreadState);
@@ -3979,7 +3958,6 @@ private:
 	bool PreloadLinkerLoadExports(FAsyncLoadingThreadState2& ThreadState);
 	EEventLoadNodeExecutionResult ExecutePostLoadLinkerLoadPackageExports(FAsyncLoadingThreadState2& ThreadState);
 	EEventLoadNodeExecutionResult ExecuteDeferredPostLoadLinkerLoadPackageExports(FAsyncLoadingThreadState2& ThreadState);
-#endif
 
 	void ProcessExportDependencies(const FAsyncPackageHeaderData& Header, int32 LocalExportIndex, FExportBundleEntry::EExportCommandType CommandType);
 	int32 GetPublicExportIndex(uint64 ExportHash, FAsyncPackageHeaderData*& OutHeader);
@@ -4139,14 +4117,6 @@ private:
 
 TAsyncAtomic<uint64> FAsyncLoadingSyncLoadContext::NextContextId{ 1 };
 
-struct FAsyncLoadingPostLoadGroup
-{
-	uint64 SyncLoadContextId = 0;
-	TArray<FAsyncPackage2*> Packages;
-	int32 PackagesWithExportsToSerializeCount = 0;
-	int32 PackagesWithExportsToPostLoadCount = 0;
-};
-
 class FAsyncLoadingThread2 final
 	: public FRunnable
 	, public IAsyncPackageLoader
@@ -4923,11 +4893,9 @@ private:
 
 #if WITH_EDITOR
 		const bool bCanImportPackagesWithIdsOnly = false;
-#elif ALT2_ENABLE_LINKERLOAD_SUPPORT
+#else
 		const bool bIsZenPackage = !AsyncPackage->LinkerLoadState.IsSet();
 		bool bCanImportPackagesWithIdsOnly = bIsZenPackage;
-#else
-		const bool bCanImportPackagesWithIdsOnly = true;
 #endif
 		if (bCanImportPackagesWithIdsOnly)
 		{
@@ -5696,13 +5664,12 @@ bool FAsyncLoadingThread2::CreateAsyncPackagesFromQueue(FAsyncLoadingThreadState
 
 						check(PackageStatus == EPackageStoreEntryStatus::Ok);
 						check(PackageLoader != EPackageLoader::Unknown);
-#if ALT2_ENABLE_LINKERLOAD_SUPPORT
+
 						if (PackageLoader == EPackageLoader::LinkerLoad)
 						{
 							Package->InitializeLinkerLoadState(&PackageDesc.InstancingContext);
 						}
 						else
-#endif
 						{
 							check(PackageLoader == EPackageLoader::Zen);
 							InitializeAsyncPackageFromPackageStore(ThreadState, &IoBatch, Package, PackageEntry);
@@ -6490,7 +6457,7 @@ void FAsyncPackage2::ImportPackagesRecursiveInner(FAsyncLoadingThreadState2& Thr
 #endif
 
 		bool bIsInstanced = false;
-#if WITH_EDITORONLY_DATA && ALT2_ENABLE_LINKERLOAD_SUPPORT
+#if WITH_EDITORONLY_DATA
 		if (bHasImportedPackageNames && LinkerLoadState.IsSet())
 		{
 			// Use the instancing context remap from disk package name to upackage name
@@ -6542,9 +6509,7 @@ void FAsyncPackage2::ImportPackagesRecursiveInner(FAsyncLoadingThreadState2& Thr
 		EPackageExtension ImportedPackageExtension = EPackageExtension::Unspecified;
 
 		FPackagePath::TryFromPackageName(ImportedPackageNameToLoad, ImportedPackagePath);
-#if ALT2_ENABLE_LINKERLOAD_SUPPORT
 		PackageLoader = LinkerLoadState.IsSet() ? EPackageLoader::LinkerLoad : PackageLoader;
-#endif
 
 		if (!ImportedPackageNameToLoad.IsNone() && AsyncLoadingThread.TryGetExistingLoadedPackagePath(ImportedPackageId, ImportedPackageNameToLoad, ImportedPackageExtension, ImportedPackageLoader, ImportedPackageEntry))
 		{
@@ -6572,14 +6537,14 @@ void FAsyncPackage2::ImportPackagesRecursiveInner(FAsyncLoadingThreadState2& Thr
 		FAsyncPackage2* ImportedPackage = nullptr;
 		bool bInserted = false;
 		bool bIsFullyLoaded = ImportedPackageRef.AreAllPublicExportsLoaded();
-#if ALT2_ENABLE_LINKERLOAD_SUPPORT
+
 		if (ImportedPackageLoader == EPackageLoader::LinkerLoad && (!ImportedPackageRef.HasPackage() || (!ImportedPackageRef.GetPackage()->GetLinker() && !ImportedPackageRef.GetPackage()->HasAnyPackageFlags(PKG_InMemoryOnly))))
 		{
 			// If we're importing a linker load package and it doesn't have its linker we need to reload it, otherwise we can't reliably link to its imports
 			// Note: Legacy loader path appears to do this only for uncooked packages in the editor?
 			bIsFullyLoaded = false;
 		}
-#endif
+
 		const bool bIsNotInstalled = GFailLoadOnNotInstalledImport && ImportedPackageStatus == EPackageStoreEntryStatus::NotInstalled;
 		const bool bIsMissing = (ImportedPackageStatus == EPackageStoreEntryStatus::Missing) || (!GFailLoadOnNotInstalledImport && ImportedPackageStatus == EPackageStoreEntryStatus::NotInstalled);
 		if (bIsFullyLoaded)
@@ -6704,14 +6669,13 @@ void FAsyncPackage2::ImportPackagesRecursiveInner(FAsyncLoadingThreadState2& Thr
 			else
 			{
 				check(ImportedPackageStatus == EPackageStoreEntryStatus::Ok);
-#if ALT2_ENABLE_LINKERLOAD_SUPPORT
+
 				if (ImportedPackageLoader == EPackageLoader::LinkerLoad)
 				{
 					// Only propagate the instancing context if the imported package is also instanced and it isn't a zen package
 					ImportedPackage->InitializeLinkerLoadState(PackageLoader == EPackageLoader::Zen || !bIsInstanced ? nullptr : &LinkerLoadState->Linker->GetInstancingContext());
 				}
 				else
-#endif
 				{
 					// TODO: Here we should probably also propagate the instancing context if the imported package is also instanced (similar to the call to InitializeLinkerLoadState done above)
 					AsyncLoadingThread.InitializeAsyncPackageFromPackageStore(ThreadState, &IoBatch, ImportedPackage, ImportedPackageEntry);
@@ -6767,7 +6731,6 @@ void FAsyncPackage2::ImportPackagesRecursive(FAsyncLoadingThreadState2& ThreadSt
 		TEXT("All imported packages are now being loaded."));
 }
 
-#if ALT2_ENABLE_LINKERLOAD_SUPPORT
 void FAsyncPackage2::InitializeLinkerLoadState(const FLinkerInstancingContext* InstancingContext)
 {
 	LinkerLoadState.Emplace();
@@ -6829,7 +6792,6 @@ void FAsyncPackage2::DetachLinker()
 		LinkerLoadState->Linker = nullptr;
 	}
 }
-#endif
 
 namespace UE::AsyncPackage2
 {
@@ -6860,7 +6822,6 @@ void FAsyncPackage2::StartLoading(FAsyncLoadingThreadState2& ThreadState, FIoBat
 	AsyncPackageLoadingState = EAsyncPackageLoadingState2::WaitingForIo;
 
 	FIoReadOptions ReadOptions;
-#if ALT2_ENABLE_LINKERLOAD_SUPPORT
 	if (LinkerLoadState.IsSet())
 	{
 #if ALT2_ENABLE_NEW_ARCHIVE_FOR_LINKERLOAD
@@ -6870,7 +6831,6 @@ void FAsyncPackage2::StartLoading(FAsyncLoadingThreadState2& ThreadState, FIoBat
 #endif
 		return;
 	}
-#endif
 
 #if WITH_EDITOR
 	if (OptionalSegmentHeaderData.IsSet())
@@ -6964,8 +6924,6 @@ void FAsyncPackage2::StartLoading(FAsyncLoadingThreadState2& ThreadState, FIoBat
 	}
 }
 
-#if ALT2_ENABLE_LINKERLOAD_SUPPORT
-
 EEventLoadNodeExecutionResult FAsyncPackage2::ProcessLinkerLoadPackageSummary(FAsyncLoadingThreadState2& ThreadState)
 {
 	TRACE_CPUPROFILER_EVENT_SCOPE(ProcessLinkerLoadPackageSummary);
@@ -7624,8 +7582,6 @@ EEventLoadNodeExecutionResult FAsyncPackage2::Event_PreloadLinkerLoadExports(FAs
 	return EEventLoadNodeExecutionResult::Complete;
 }
 
-#endif // ALT2_ENABLE_LINKERLOAD_SUPPORT
-
 EEventLoadNodeExecutionResult FAsyncPackage2::Event_ProcessPackageSummary(FAsyncLoadingThreadState2& ThreadState, FAsyncPackage2* Package, int32)
 {
 	TRACE_CPUPROFILER_EVENT_SCOPE(Event_ProcessPackageSummary);
@@ -7643,12 +7599,10 @@ EEventLoadNodeExecutionResult FAsyncPackage2::Event_ProcessPackageSummary(FAsync
 
 	Package->CallProgressCallbacks(EAsyncLoadingProgress::Read);
 
-#if ALT2_ENABLE_LINKERLOAD_SUPPORT
 	if (Package->LinkerLoadState.IsSet())
 	{
 		return Package->ProcessLinkerLoadPackageSummary(ThreadState);
 	}
-#endif
 	if (Package->HasLoadFailed())
 	{
 		if (Package->Desc.bCanBeImported)
@@ -7842,11 +7796,8 @@ EEventLoadNodeExecutionResult FAsyncPackage2::Event_ProcessExportBundle(FAsyncLo
 	UE::Core::Private::FPlayInEditorLoadingScope PlayInEditorIDScope(Package->Desc.PIEInstanceID);
 #endif
 
-#if ALT2_ENABLE_LINKERLOAD_SUPPORT
 	// This code path should not be reached for LinkerLoad packages
 	check(!Package->LinkerLoadState.IsSet());
-#endif
-
 	check(InExportBundleIndex < Package->Data.TotalExportBundleCount);
 	
 	if (!Package->HasLoadFailed())
@@ -8290,13 +8241,11 @@ FString FAsyncPackage2::GetNameFromPackageObjectIndex(FPackageObjectIndex Index)
 	FString Out;
 	if (Index.IsExport())
 	{
-#if ALT2_ENABLE_LINKERLOAD_SUPPORT
 		if (LinkerLoadState.IsSet())
 		{
 			Out = LinkerLoadState->Linker->ExportMap[Index.ToExport()].ObjectName.ToString();
 		}
 		else
-#endif
 		{
 			FName Name(NAME_None);
 			HeaderData.NameMap.TryGetName(HeaderData.ExportMap[Index.ToExport()].ObjectName, Name);
@@ -8975,14 +8924,9 @@ EEventLoadNodeExecutionResult FAsyncPackage2::Event_ExportsDone(FAsyncLoadingThr
 	UE_ASYNC_PACKAGE_DEBUG(Package->Desc);
 
 	check(FUObjectThreadContext::Get().AsyncVisibilityFilter == EInternalObjectFlags::None);
-#if ALT2_ENABLE_LINKERLOAD_SUPPORT
 	check(Package->AsyncPackageLoadingState == EAsyncPackageLoadingState2::ProcessExportBundles || 
 		  Package->AsyncPackageLoadingState == EAsyncPackageLoadingState2::WaitingForExternalReads ||
 		  Package->AsyncPackageLoadingState == EAsyncPackageLoadingState2::PreloadLinkerLoadExports);
-#else
-	check(Package->AsyncPackageLoadingState == EAsyncPackageLoadingState2::ProcessExportBundles ||
-		  Package->AsyncPackageLoadingState == EAsyncPackageLoadingState2::WaitingForExternalReads);
-#endif
 
 	Package->AsyncPackageLoadingState = EAsyncPackageLoadingState2::ExportsDone;
 
@@ -9246,14 +9190,12 @@ void FAsyncPackage2::ConditionalBeginProcessPackageExports(FAsyncLoadingThreadSt
 		{
 			check(Package->AsyncPackageLoadingState == EAsyncPackageLoadingState2::DependenciesReady);
 
-#if ALT2_ENABLE_LINKERLOAD_SUPPORT
 			if (Package->LinkerLoadState.IsSet())
 			{
 				Package->AsyncPackageLoadingState = EAsyncPackageLoadingState2::CreateLinkerLoadExports;
 				Package->GetPackageNode(EEventLoadNode2::Package_CreateLinkerLoadExports).ReleaseBarrier(&ThreadState);
 			}
 			else
-#endif
 			{
 				Package->AsyncPackageLoadingState = EAsyncPackageLoadingState2::ProcessExportBundles;
 				if (Package->Data.TotalExportBundleCount > 0)
@@ -9295,34 +9237,11 @@ EEventLoadNodeExecutionResult FAsyncPackage2::Event_PostLoadExportBundle(FAsyncL
 
 	FAsyncPackageScope2 PackageScope(Package);
 
-#if ALT2_ENABLE_LINKERLOAD_SUPPORT
 	if (Package->LinkerLoadState.IsSet())
 	{
 		return Package->ExecutePostLoadLinkerLoadPackageExports(ThreadState);
 	}
-#endif
 
-	/*TSet<FAsyncPackage2*> Visited;
-	TArray<FAsyncPackage2*> ProcessQueue;
-	ProcessQueue.Push(Package);
-	while (ProcessQueue.Num() > 0)
-	{
-		FAsyncPackage2* CurrentPackage = ProcessQueue.Pop();
-		Visited.Add(CurrentPackage);
-		if (CurrentPackage->AsyncPackageLoadingState < EAsyncPackageLoadingState2::ExportsDone)
-		{
-			UE_DEBUG_BREAK();
-		}
-		for (const FPackageId& ImportedPackageId : CurrentPackage->StoreEntry.ImportedPackages)
-		{
-			FAsyncPackage2* ImportedPackage = CurrentPackage->AsyncLoadingThread.GetAsyncPackage(ImportedPackageId);
-			if (ImportedPackage && !Visited.Contains(ImportedPackage))
-			{
-				ProcessQueue.Push(ImportedPackage);
-			}
-		}
-	}*/
-	
 	check(InExportBundleIndex < Package->Data.TotalExportBundleCount);
 
 	EEventLoadNodeExecutionResult LoadingState = EEventLoadNodeExecutionResult::Complete;
@@ -9446,12 +9365,10 @@ EEventLoadNodeExecutionResult FAsyncPackage2::Event_DeferredPostLoadExportBundle
 	FAsyncPackageScope2 PackageScope(Package);
 	TGuardValue<bool> GuardIsRoutingPostLoad(PackageScope.ThreadContext.IsRoutingPostLoad, true);
 
-#if ALT2_ENABLE_LINKERLOAD_SUPPORT
 	if (Package->LinkerLoadState.IsSet())
 	{
 		return Package->ExecuteDeferredPostLoadLinkerLoadPackageExports(ThreadState);
 	}
-#endif
 
 	check(InExportBundleIndex < Package->Data.TotalExportBundleCount);
 	EEventLoadNodeExecutionResult LoadingState = EEventLoadNodeExecutionResult::Complete;
@@ -9943,9 +9860,7 @@ EAsyncPackageState::Type FAsyncLoadingThread2::ProcessLoadedPackagesFromGameThre
 
 				Package->ClearConstructedObjects();
 
-#if ALT2_ENABLE_LINKERLOAD_SUPPORT
 				Package->DetachLinker();
-#endif
 			}
 
 			// Remove the package from the list before we trigger the callbacks, 
@@ -10133,11 +10048,9 @@ FAsyncLoadingThread2::FAsyncLoadingThread2(FIoDispatcher& InIoDispatcher)
 	EventSpecs[EEventLoadNode2::Package_ProcessSummary] = { &FAsyncPackage2::Event_ProcessPackageSummary, &EventQueue, EExecutionType::Parallel, TEXT("ProcessSummary") };
 	EventSpecs[EEventLoadNode2::Package_ProcessDependencies] = { &FAsyncPackage2::Event_ProcessDependencies, &EventQueue, EExecutionType::Normal, TEXT("ProcessDependencies") };
 	EventSpecs[EEventLoadNode2::Package_DependenciesReady] = { &FAsyncPackage2::Event_DependenciesReady, &EventQueue, EExecutionType::Normal, TEXT("DependenciesReady")};
-#if ALT2_ENABLE_LINKERLOAD_SUPPORT
 	EventSpecs[EEventLoadNode2::Package_CreateLinkerLoadExports] = { &FAsyncPackage2::Event_CreateLinkerLoadExports, &EventQueue, EExecutionType::Normal, TEXT("CreateLinkerLoadExports") };
 	EventSpecs[EEventLoadNode2::Package_ResolveLinkerLoadImports] = { &FAsyncPackage2::Event_ResolveLinkerLoadImports, &EventQueue, EExecutionType::Normal, TEXT("ResolveLinkerLoadImports") };
 	EventSpecs[EEventLoadNode2::Package_PreloadLinkerLoadExports] = { &FAsyncPackage2::Event_PreloadLinkerLoadExports, &EventQueue, EExecutionType::Parallel, TEXT("PreloadLinkerLoadExports") };
-#endif
 	EventSpecs[EEventLoadNode2::Package_ExportsSerialized] = { &FAsyncPackage2::Event_ExportsDone, &EventQueue, EExecutionType::Normal, TEXT("ExportsSerialized") };
 
 	EventSpecs[EEventLoadNode2::Package_NumPhases + EEventLoadNode2::ExportBundle_Process] = { &FAsyncPackage2::Event_ProcessExportBundle, &EventQueue, EExecutionType::Normal, TEXT("ExportBundle_Process") };
@@ -12233,12 +12146,10 @@ void FAsyncPackage2::DumpStateImpl(TSet<FAsyncPackage2*>& Set, int32 Indent, TMu
 	{
 		UE_LOG(LogStreaming, Warning, TEXT(" %s AllDependenciesSetupState is waiting on %s"), FCString::Spc(Indent), *FormatPackage(AllDependenciesSetupState.WaitingForPackage));
 	}
-#if ALT2_ENABLE_LINKERLOAD_SUPPORT
 	if (AllDependenciesImportState.WaitingForPackage)
 	{
 		UE_LOG(LogStreaming, Warning, TEXT(" %s AllDependenciesImportState is waiting on %s"), FCString::Spc(Indent), *FormatPackage(AllDependenciesImportState.WaitingForPackage));
 	}
-#endif
 	if (AllDependenciesPostLoadReadyState.WaitingForPackage)
 	{
 		UE_LOG(LogStreaming, Warning, TEXT(" %s AllDependenciesPostLoadReadyState is waiting on %s"), FCString::Spc(Indent), *FormatPackage(AllDependenciesPostLoadReadyState.WaitingForPackage));
