[ZenLoader]
- Fix potential missing wake-up in the multithreaded loader when external reads needs to be processed

#FailedValidations Preflight

diff --git a/Engine/Source/Runtime/CoreUObject/Private/Serialization/AsyncLoading2.cpp b/Engine/Source/Runtime/CoreUObject/Private/Serialization/AsyncLoading2.cpp
--- a/Engine/Source/Runtime/CoreUObject/Private/Serialization/AsyncLoading2.cpp
+++ b/Engine/Source/Runtime/CoreUObject/Private/Serialization/AsyncLoading2.cpp
@@ -4893,7 +4893,7 @@ private:
 
 #if WITH_EDITOR
 		const bool bCanImportPackagesWithIdsOnly = false;
-#else
+#else          
 		const bool bIsZenPackage = !AsyncPackage->LinkerLoadState.IsSet();
 		bool bCanImportPackagesWithIdsOnly = bIsZenPackage;
 #endif
@@ -7468,10 +7468,10 @@ EEventLoadNodeExecutionResult FAsyncPackage2::Event_CreateLinkerLoadExports(FAsy
 
 	check(Package->LinkerLoadState.IsSet());
 
-	if (!Package->CreateLinkerLoadExports(ThreadState))
-	{
-		return EEventLoadNodeExecutionResult::Timeout;
-	}
+		if (!Package->CreateLinkerLoadExports(ThreadState))
+		{
+			return EEventLoadNodeExecutionResult::Timeout;
+		}
 
 	Package->AsyncPackageLoadingState = EAsyncPackageLoadingState2::WaitingForLinkerLoadDependencies;
 	Package->ConditionalBeginResolveLinkerLoadImports(ThreadState);
@@ -7578,6 +7578,10 @@ EEventLoadNodeExecutionResult FAsyncPackage2::Event_PreloadLinkerLoadExports(FAs
 	{
 		Package->AsyncPackageLoadingState = EAsyncPackageLoadingState2::WaitingForExternalReads;
 		Package->AsyncLoadingThread.ExternalReadQueue.Enqueue(Package);
+		if (IsInParallelLoadingThread())
+		{
+			Package->AsyncLoadingThread.AltZenaphore.NotifyOne();
+		}
 	}
 	return EEventLoadNodeExecutionResult::Complete;
 }
@@ -7796,7 +7800,7 @@ EEventLoadNodeExecutionResult FAsyncPackage2::Event_ProcessExportBundle(FAsyncLo
 	UE::Core::Private::FPlayInEditorLoadingScope PlayInEditorIDScope(Package->Desc.PIEInstanceID);
 #endif
 
-	// This code path should not be reached for LinkerLoad packages
+		// This code path should not be reached for LinkerLoad packages
 	check(!Package->LinkerLoadState.IsSet());
 	check(InExportBundleIndex < Package->Data.TotalExportBundleCount);
 	
@@ -12621,4 +12625,4 @@ bool FPublicExportMapPreInsertPublicExportsTest::RunTest(const FString& Paramete
 	return true;
 }
 
-#endif // WITH_LOW_LEVEL_TESTS
+#endif // WITH_LOW_LEVEL_TESTS
\ No newline at end of file
