[AsyncLoading2]

Prevent packages from loading if any of their imported packages cannot be loaded (missing or notinstalled) when CVar s.FailLoadOnFailedImport is true.

- Adjusts the existing logging for missing/notinstalled packages to no longer be conditional on the package name existing as that is uncommon when importing packages since the header is not loaded yet
- Adds a bIsNotInstalled state to FLoadedPackageRef so that we can differentiate between NotInstalled and Missing where NotInstalled implies that the package is expected to possibly be missing but may eventually become available, whereas Missing implies it's unexpected to not exist and it is unexpected for such a package to later become available (generates a warning in both cases)
- Adds LexToString for PackageStoreEntryStatus to simplify logging

#rnx

#FailedValidations Preflight

diff --git a/Engine/Source/Runtime/CoreUObject/Private/Serialization/AsyncLoading2.cpp b/Engine/Source/Runtime/CoreUObject/Private/Serialization/AsyncLoading2.cpp
--- a/Engine/Source/Runtime/CoreUObject/Private/Serialization/AsyncLoading2.cpp
+++ b/Engine/Source/Runtime/CoreUObject/Private/Serialization/AsyncLoading2.cpp
@@ -485,11 +485,11 @@ static FAutoConsoleVariableRef CVarGOnlyProcessRequiredPackagesWhenSyncLoading(
 	ECVF_Default
 );
 
-static bool GFailLoadOnNotInstalledImport = true;
-static FAutoConsoleVariableRef CVarGFailLoadOnNotInstalledImport(
-	TEXT("s.FailLoadOnNotInstalledImport"),
-	GFailLoadOnNotInstalledImport,
-	TEXT("Fail package load if an imported package is not installed"),
+static bool GFailLoadOnFailedImport = true;
+static FAutoConsoleVariableRef CVarGFailLoadOnFailedImport(
+	TEXT("s.FailLoadOnFailedImport"),
+	GFailLoadOnFailedImport,
+	TEXT("Fail package load if an imported package cannot be loaded for any reason"),
 	ECVF_Default
 );
 
@@ -1302,6 +1302,7 @@ private:
 	EPackageExtension PackageHeaderExtension = EPackageExtension::Unspecified;
 	EPackageLoader PackageLoader = EPackageLoader::Unknown;
 	bool bAreAllPublicExportsLoaded = false;
+	bool bIsNotInstalled = false;
 	bool bIsMissing = false;
 	bool bHasFailed = false;
 	bool bHasBeenLoadedDebug = false;
@@ -1334,7 +1335,22 @@ public:
 
 	inline bool HasErrors() const
 	{
-		return bHasFailed || bIsMissing;
+		return bHasFailed || bIsMissing || bIsNotInstalled;
+	}
+
+	inline bool IsFailed() const
+	{
+		return bHasFailed;
+	}
+
+	inline bool IsMissing() const
+	{
+		return bIsMissing;
+	}
+
+	inline bool IsNotInstalled() const
+	{
+		return bIsNotInstalled;
 	}
 
 	inline EPackageLoader GetPackageLoader() const
@@ -1360,6 +1376,7 @@ public:
 	{
 		check(!bAreAllPublicExportsLoaded);
 		check(!bIsMissing);
+		check(!bIsNotInstalled);
 		check(!bHasFailed);
 		check(!HasPackage());
 		if (InPackage)
@@ -1400,6 +1417,7 @@ public:
 	inline void SetAllPublicExportsLoaded(bool bSnapshotExportCount)
 	{
 		check(!bIsMissing);
+		check(!bIsNotInstalled);
 		check(!bHasFailed);
 		check(HasPackage());
 		bIsMissing = false;
@@ -1421,6 +1439,14 @@ public:
 		bAreAllPublicExportsLoaded = false;
 	}
 
+	inline void SetIsNotInstalledPackage()
+	{
+		check(!bAreAllPublicExportsLoaded);
+		check(!HasPackage());
+		bIsNotInstalled = true;
+		bAreAllPublicExportsLoaded = false;
+	}
+
 	inline void ClearErrorFlags()
 	{
 		bIsMissing = false;
@@ -1451,6 +1477,7 @@ public:
 	void RemovePublicExport(uint64 ExportHash, FName ObjectName = NAME_None)
 	{
 		check(!bIsMissing);
+		check(!bIsNotInstalled);
 		check(HasPackage());
 		if (PublicExportMap.Remove(ExportHash))
 		{
@@ -1748,6 +1775,7 @@ public:
 
 					if (PackageRef.bIsMissing)
 					{
+						check(!PackageRef.bIsNotInstalled);
 						PackageRef.bIsMissing = false;
 						UE_LOG(LogStreaming, Warning,
 							TEXT("FGlobalImportStore:AddPackageRef: Found reference to previously missing package %s (0x%s)"),
@@ -1755,6 +1783,18 @@ public:
 							*LexToString(PackageId)
 						);
 					}
+					else if (PackageRef.bIsNotInstalled)
+					{
+						check(!PackageRef.bIsMissing);
+						PackageRef.bIsNotInstalled = false;
+						// Note, unlike missing packages, it is expected that we may later find a reference to a 
+						// previously not installed package so we do not emit a warning
+						UE_LOG(LogStreaming, Log,
+							TEXT("FGlobalImportStore:AddPackageRef: Found reference to previously not installed package %s (0x%s)"),
+							*FoundPackage->GetName(),
+							*LexToString(PackageId)
+						);
+					}
 
 					PackageRef.SetPackage(FoundPackage);
 					FoundPackage->SetCanBeImportedFlag(true);
@@ -1831,23 +1871,24 @@ public:
 		--PackageRef.RefCount;
 
 #if DO_CHECK
-		ensureMsgf(!PackageRef.bHasBeenLoadedDebug || PackageRef.bAreAllPublicExportsLoaded || PackageRef.bIsMissing || PackageRef.bHasFailed,
+		ensureMsgf(!PackageRef.bHasBeenLoadedDebug || PackageRef.bAreAllPublicExportsLoaded || PackageRef.bIsMissing || PackageRef.bIsNotInstalled || PackageRef.bHasFailed,
 			TEXT("LoadedPackageRef from None (0x%s) to %s (0x%s) should not have been released when the package is not complete.")
-			TEXT("RefCount=%d, AreAllExportsLoaded=%d, IsMissing=%d, HasFailed=%d, HasBeenLoaded=%d"),
+			TEXT("RefCount=%d, AreAllExportsLoaded=%d, IsMissing=%d, IsNotInstalled=%d, HasFailed=%d, HasBeenLoaded=%d"),
 			*LexToString(FromPackageId),
 			*PackageRef.GetOriginalPackageName().ToString(),
 			*LexToString(PackageId),
 			PackageRef.RefCount,
 			PackageRef.bAreAllPublicExportsLoaded,
 			PackageRef.bIsMissing,
+			PackageRef.bIsNotInstalled,
 			PackageRef.bHasFailed,
 			PackageRef.bHasBeenLoadedDebug);
 
 		if (PackageRef.bAreAllPublicExportsLoaded)
 		{
-			check(!PackageRef.bIsMissing);
+			check(!PackageRef.bIsMissing && !PackageRef.bIsNotInstalled);
 		}
-		if (PackageRef.bIsMissing)
+		if (PackageRef.bIsMissing || PackageRef.bIsNotInstalled)
 		{
 			check(!PackageRef.bAreAllPublicExportsLoaded);
 		}
@@ -4657,6 +4698,77 @@ private:
 
 		if (PackageStatus != EPackageStoreEntryStatus::Missing)
 		{
+			if (PackageStatus == EPackageStoreEntryStatus::Ok && GFailLoadOnFailedImport)
+			{
+				// We need to confirm that this package can be fully imported so walk its transitive dependencies
+				// If we have a loaded package already that isn't failed we don't need to check it's dependencies
+				bool bFoundMissing = false;
+				bool bFoundNotInstalled = false;
+				FPackageId UnloadableImportPackageId;
+				TArray<FPackageId> ImportsToCheck(OutPackageStoreEntry.ImportedPackageIds);
+				TSet<FPackageId> ImportsSeen(OutPackageStoreEntry.ImportedPackageIds); // Out of paranoia, don't loop forever if there is somehow a circular import chain
+				while (!ImportsToCheck.IsEmpty() && !UnloadableImportPackageId.IsValid())
+				{
+					FPackageId ImportId = ImportsToCheck.Pop();
+					if (FLoadedPackageRef* PackageRef = GlobalImportStore.FindPackageRef(ImportId))
+					{
+						if (!PackageRef->HasErrors())
+						{
+							continue;
+						}
+						else if(PackageRef->IsNotInstalled())
+						{
+							bFoundMissing = PackageRef->IsMissing();
+							bFoundNotInstalled = PackageRef->IsNotInstalled();
+							UnloadableImportPackageId = ImportId;
+							break;
+						}
+					}
+
+					FPackageStoreEntry ImportEntry;
+					FName ImportName;
+					EPackageStoreEntryStatus ImportStatus = PackageStore.GetPackageStoreEntry(ImportId, ImportName, ImportEntry);
+					bFoundMissing |= ImportStatus == EPackageStoreEntryStatus::Missing;
+					bFoundNotInstalled |= ImportStatus == EPackageStoreEntryStatus::NotInstalled;
+					if (bFoundMissing || bFoundNotInstalled)
+					{
+						UnloadableImportPackageId = ImportId;
+						break;
+					}
+
+					// Append all the imports we haven't checked yet
+					for (FPackageId ImportToAdd : ImportEntry.ImportedPackageIds)
+					{
+						bool bIsAlreadyInSet = false;
+						ImportsSeen.Add(ImportToAdd, &bIsAlreadyInSet);
+						if (!bIsAlreadyInSet)
+						{
+							ImportsToCheck.Add(ImportToAdd);
+						}
+					}
+				}
+
+				if (UnloadableImportPackageId.IsValid())
+				{
+					check(bFoundMissing || bFoundNotInstalled);
+					//
+					// Note: Update the FAsyncLoadingTestFailedImportPropagation test if you edit these messages
+					//
+					if (bFoundMissing)
+					{
+						UE_LOG(LogStreaming, Display, TEXT("Package %s (0x%s) will not be loaded due to Missing imported package (0x%s)"),
+							*InOutPackageNameToLoad.ToString(), *LexToString(InOutPackageIdToLoad), *LexToString(UnloadableImportPackageId));
+						return EPackageStoreEntryStatus::Missing;
+					}
+					else
+					{
+						UE_LOG(LogStreaming, Display, TEXT("Package %s (0x%s) will not be loaded due to the NotInstalled imported package (0x%s), however it may become available later."),
+							*InOutPackageNameToLoad.ToString(), *LexToString(InOutPackageIdToLoad), *LexToString(UnloadableImportPackageId));
+						return EPackageStoreEntryStatus::NotInstalled;
+					}
+				}
+			}
+
 #if WITH_EDITORONLY_DATA
 			if (OutPackageStoreEntry.LoaderType == EPackageLoader::LinkerLoad)
 			{
@@ -5527,6 +5639,18 @@ bool FAsyncLoadingThread2::CreateAsyncPackagesFromQueue(FAsyncLoadingThreadState
 			// Remove from PendingPackages
 			It.RemoveCurrent();
 		}
+		else if (PendingPackageStatus == EPackageStoreEntryStatus::NotInstalled)
+		{
+			// Initialize package with a fake package store entry
+			FPackageStoreEntry FakePackageEntry;
+			InitializeAsyncPackageFromPackageStore(ThreadState, &IoBatch, PendingPackage, FakePackageEntry);
+			// Simulate StartLoading() getting back a failed IoRequest and let it go through all package states
+			PendingPackage->AsyncPackageLoadingState = EAsyncPackageLoadingState2::WaitingForIo;
+			PendingPackage->LoadStatus = EAsyncLoadingResult::FailedNotInstalled;
+			PendingPackage->GetPackageNode(EEventLoadNode2::Package_ProcessSummary).ReleaseBarrier(&ThreadState);
+			// Remove from PendingPackages
+			It.RemoveCurrent();
+		}
 	}
 	for (;;)
 	{
@@ -6530,9 +6654,17 @@ void FAsyncPackage2::ImportPackagesRecursiveInner(FAsyncLoadingThreadState2& Thr
 		}
 		ImportedPackagePath.SetHeaderExtension(ImportedPackageExtension);
 
+		// Add the imported package reference and confirm that its loader and extension match what was found from the package store or import store directly, 
+		// and if not, ensure the package was marked as missing or not installed
 		FLoadedPackageRef& ImportedPackageRef = ImportStore.AddImportedPackageReference(ImportedPackageId, ImportedPackageUPackageName, ImportedPackageLoader, ImportedPackagePath.GetHeaderExtension());
-		checkf(ImportedPackageStatus == EPackageStoreEntryStatus::Missing || ImportedPackageRef.GetPackageLoader() == ImportedPackageLoader, TEXT("ImportPackageStatus(%d) != EPackageStoreEntryStatus::Missing(%d) || ImportedPackageRef.GetPackageLoader()(%s) != ImportedPackageLoader(%s)"), ImportedPackageStatus, EPackageStoreEntryStatus::Missing, LexToString(ImportedPackageRef.GetPackageLoader()), LexToString(ImportedPackageLoader));
-		checkf(ImportedPackageStatus == EPackageStoreEntryStatus::Missing || ImportedPackageRef.GetPackageHeaderExtension() == ImportedPackagePath.GetHeaderExtension(), TEXT("ImportPackageStatus(%d) != EPackageStoreEntryStatus::Missing(%d) || ImportedPackageRef.GetPackageHeaderExtension()(%d) != ImportedPackagePath.GetHeaderExtension()(%d)"), ImportedPackageStatus, EPackageStoreEntryStatus::Missing, ImportedPackageRef.GetPackageHeaderExtension(), ImportedPackagePath.GetHeaderExtension());
+		checkf((ImportedPackageStatus == EPackageStoreEntryStatus::Missing || ImportedPackageStatus == EPackageStoreEntryStatus::NotInstalled) 
+			|| ImportedPackageRef.GetPackageLoader() == ImportedPackageLoader, 
+			TEXT("ImportPackageStatus=%s, ImportedPackageRef.GetPackageLoader()=%s, ImportedPackageLoader=%s"), 
+			LexToString(ImportedPackageStatus), LexToString(ImportedPackageRef.GetPackageLoader()), LexToString(ImportedPackageLoader));
+		checkf((ImportedPackageStatus == EPackageStoreEntryStatus::Missing || ImportedPackageStatus == EPackageStoreEntryStatus::NotInstalled) 
+			|| ImportedPackageRef.GetPackageHeaderExtension() == ImportedPackagePath.GetHeaderExtension(), 
+			TEXT("ImportPackageStatus=%s, ImportedPackageRef.GetPackageHeaderExtension()=%s, ImportedPackagePath.GetHeaderExtension()=%s"), 
+			LexToString(ImportedPackageStatus), LexToString(ImportedPackageRef.GetPackageHeaderExtension()), LexToString(ImportedPackagePath.GetHeaderExtension()));
 
 		FAsyncPackage2* ImportedPackage = nullptr;
 		bool bInserted = false;
@@ -6545,8 +6677,8 @@ void FAsyncPackage2::ImportPackagesRecursiveInner(FAsyncLoadingThreadState2& Thr
 			bIsFullyLoaded = false;
 		}
 
-		const bool bIsNotInstalled = GFailLoadOnNotInstalledImport && ImportedPackageStatus == EPackageStoreEntryStatus::NotInstalled;
-		const bool bIsMissing = (ImportedPackageStatus == EPackageStoreEntryStatus::Missing) || (!GFailLoadOnNotInstalledImport && ImportedPackageStatus == EPackageStoreEntryStatus::NotInstalled);
+		const bool bIsNotInstalled = ImportedPackageStatus == EPackageStoreEntryStatus::NotInstalled;
+		const bool bIsMissing = ImportedPackageStatus == EPackageStoreEntryStatus::Missing;
 		if (bIsFullyLoaded)
 		{
 			ImportedPackage = AsyncLoadingThread.FindAsyncPackage(ImportedPackageId);
@@ -6558,37 +6690,21 @@ void FAsyncPackage2::ImportPackagesRecursiveInner(FAsyncLoadingThreadState2& Thr
 		}
 		else if (bIsNotInstalled)
 		{
-			LoadStatus = EAsyncLoadingResult::FailedNotInstalled;
-			ImportedPackageRef.SetHasFailed();
-
-			UE_ASYNC_PACKAGE_CLOG(!ImportedPackageUPackageName.IsNone(), Log, Desc, TEXT("ImportPackages: SkipPackage"),
-				TEXT("Skipping not installed imported package %s (0x%s)"), *ImportedPackageNameToLoad.ToString(), *LexToString(ImportedPackageId));
-			if (!ImportedPackageUPackageName.IsNone())
-			{
-				FName PackagePathToLoadName = Desc.PackagePathToLoad.GetPackageFName();
-				ExecuteOnGameThread(TEXT("ExplanationForNotInstalledPackage"), [ImportedPackageNameToLoad, PackagePathToLoadName]()
-				{
-					FString ImportedPackageNameToLoadString = ImportedPackageNameToLoad.ToString();
-					FString PackagePathToLoadString = PackagePathToLoadName.ToString();
-					FMessageLog("LoadErrors").Warning(FText::FormatNamed(NSLOCTEXT("Core", "AsyncLoading_SkippedPackage_NotInstalled", "While trying to load package {Package}, a dependent package {DependentPackage} was unavailable because it was not installed."),
-						TEXT("Package"), FText::FromString(PackagePathToLoadString),
-						TEXT("DependentPackage"), FText::FromString(ImportedPackageNameToLoadString)));
-					return false;
-				});
-			}
-
+			UE_ASYNC_PACKAGE_LOG(Log, Desc, TEXT("ImportPackages: SkipPackage"),
+				TEXT("Skipping imported package %s (0x%s) because it or one of its dependencies is not installed"), *ImportedPackageNameToLoad.ToString(), *LexToString(ImportedPackageId));
+			ImportedPackageRef.SetIsNotInstalledPackage();
 			continue;
 		}
 		else if (bIsMissing)
 		{
 			if (!ImportedPackageRef.HasPackage()) // If we found a package it's not actually missing but we can't load it anyway
 			{
-				UE_ASYNC_PACKAGE_CLOG(!ImportedPackageUPackageName.IsNone(), Log, Desc, TEXT("ImportPackages: SkipPackage"),
+				UE_ASYNC_PACKAGE_LOG(Log, Desc, TEXT("ImportPackages: SkipPackage"),
 					TEXT("Skipping non mounted imported package %s (0x%s)"), *ImportedPackageNameToLoad.ToString(), *LexToString(ImportedPackageId));
-				if (!ImportedPackageUPackageName.IsNone() && !FLinkerLoad::IsKnownMissingPackage(ImportedPackageNameToLoad))
+				if (!FLinkerLoad::IsKnownMissingPackage(ImportedPackageNameToLoad))
 				{
 					FName PackagePathToLoadName = Desc.PackagePathToLoad.GetPackageFName();
-					ExecuteOnGameThread(TEXT("GetExplanationForUnavailablePackage"), [ImportedPackageNameToLoad, PackagePathToLoadName]()
+					ExecuteOnGameThread(TEXT("GetExplanationForUnavailablePackage"), [ImportedPackageNameToLoad, ImportedPackageId, PackagePathToLoadName]()
 					{
 						FString ImportedPackageNameToLoadString = ImportedPackageNameToLoad.ToString();
 						FString PackagePathToLoadString = PackagePathToLoadName.ToString();
@@ -6596,16 +6712,18 @@ void FAsyncPackage2::ImportPackagesRecursiveInner(FAsyncLoadingThreadState2& Thr
 						FPackageName::GetExplanationForUnavailablePackage(ImportedPackageNameToLoad, Explanation);
 						if (Explanation.Len())
 						{
-							FMessageLog("LoadErrors").Warning(FText::FormatNamed(NSLOCTEXT("Core", "AsyncLoading_SkippedPackage_Explanation", "While trying to load package {MissingPackage}, a dependent package {DependentPackage} was not available. Additional explanatory information follows:\n{Explanation}"),
+							FMessageLog("LoadErrors").Warning(FText::FormatNamed(NSLOCTEXT("Core", "AsyncLoading_SkippedPackage_Explanation", "While trying to load package {MissingPackage}, a dependent package {DependentPackage} ({DependentPackageId}) was not available. Additional explanatory information follows:\n{Explanation}"),
 								TEXT("MissingPackage"), FText::FromString(PackagePathToLoadString),
 								TEXT("DependentPackage"), FText::FromString(ImportedPackageNameToLoadString),
+								TEXT("DependentPackageId"), FText::FromString(LexToString(ImportedPackageId)),
 								TEXT("Explanation"), FText::FromString(Explanation.ToString())));
 						}
 						else
 						{
-							FMessageLog("LoadErrors").Warning(FText::FormatNamed(NSLOCTEXT("Core", "AsyncLoading_SkippedPackage_NoExplanation", "While trying to load package {MissingPackage}, a dependent package {DependentPackage} was not available. No additional explanation was available."),
+							FMessageLog("LoadErrors").Warning(FText::FormatNamed(NSLOCTEXT("Core", "AsyncLoading_SkippedPackage_NoExplanation", "While trying to load package {MissingPackage}, a dependent package {DependentPackage} ({DependentPackageId}) was not available. No additional explanation was available."),
 								TEXT("MissingPackage"), FText::FromString(PackagePathToLoadString),
-								TEXT("DependentPackage"), FText::FromString(ImportedPackageNameToLoadString)));
+								TEXT("DependentPackage"), FText::FromString(ImportedPackageNameToLoadString),
+								TEXT("DependentPackageId"), FText::FromString(LexToString(ImportedPackageId))));
 						}
 
 						return false;
@@ -11416,7 +11534,7 @@ void FAsyncLoadingThread2::QueueUnresolvedPackage(FAsyncLoadingThreadState2& Thr
 	if (PackageStatus == EPackageStoreEntryStatus::NotInstalled)
 	{
 		UE_CLOG(bIssueWarning, LogStreaming, Warning,
-			TEXT("LoadPackage: SkipPackage: %s (0x%s) - The package to load does not exist on disk or in the loader but may be installed on demand"),
+			TEXT("LoadPackage: SkipPackage: %s (0x%s) - The package to load, or one of its dependencies (see logs above), does not exist on disk or in the loader but may be installed on demand later."),
 			*FailedPackageName.ToString(), *LexToString(PackageDesc.PackageIdToLoad));
 	}
 	else
diff --git a/Engine/Source/Runtime/CoreUObject/Private/Serialization/PackageStore.cpp b/Engine/Source/Runtime/CoreUObject/Private/Serialization/PackageStore.cpp
--- a/Engine/Source/Runtime/CoreUObject/Private/Serialization/PackageStore.cpp
+++ b/Engine/Source/Runtime/CoreUObject/Private/Serialization/PackageStore.cpp
@@ -27,6 +27,20 @@
 
 DEFINE_LOG_CATEGORY_STATIC(LogPackageStore, Log, All);
 
+const TCHAR* LexToString(EPackageStoreEntryStatus PackageLoader)
+{
+	switch (PackageLoader)
+	{
+	case EPackageStoreEntryStatus::None: return TEXT("None");
+	case EPackageStoreEntryStatus::Missing: return TEXT("Missing");
+	case EPackageStoreEntryStatus::NotInstalled: return TEXT("NotInstalled");
+	case EPackageStoreEntryStatus::Pending: return TEXT("Pending");
+	case EPackageStoreEntryStatus::Ok: return TEXT("Ok");
+	default:
+		checkNoEntry();
+		return TEXT("Not Implemented");
+	}
+};
 
 ///////////////////////////////////////
 // FLinkerLoadPackageStoreBackend
diff --git a/Engine/Source/Runtime/CoreUObject/Public/Serialization/PackageStore.h b/Engine/Source/Runtime/CoreUObject/Public/Serialization/PackageStore.h
--- a/Engine/Source/Runtime/CoreUObject/Public/Serialization/PackageStore.h
+++ b/Engine/Source/Runtime/CoreUObject/Public/Serialization/PackageStore.h
@@ -7,6 +7,7 @@
 #include "CoreMinimal.h"
 #include "CoreTypes.h"
 #include "Delegates/Delegate.h"
+#include "HAL/Platform.h"
 #include "HAL/PlatformCrt.h"
 #include "Misc/EnumClassFlags.h"
 #include "Misc/PackagePath.h"
@@ -42,6 +43,7 @@ enum class EPackageStoreEntryStatus
 	Pending,
 	Ok,
 };
+const TCHAR* LexToString(EPackageStoreEntryStatus PackageLoader);
 
 /**
  * Package store entry.
diff --git a/Engine/Source/Runtime/Engine/Private/Tests/Loading/AsyncLoadingTests.cpp b/Engine/Source/Runtime/Engine/Private/Tests/Loading/AsyncLoadingTests.cpp
--- a/Engine/Source/Runtime/Engine/Private/Tests/Loading/AsyncLoadingTests.cpp
+++ b/Engine/Source/Runtime/Engine/Private/Tests/Loading/AsyncLoadingTests.cpp
@@ -7,6 +7,7 @@
 #include "AssetRegistry/IAssetRegistry.h"
 #include "Async/ParallelFor.h"
 #include "AsyncLoadingTests_Shared.h"
+#include "Serialization/PackageStore.h"
 #include "StructUtils/UserDefinedStruct.h"
 #include "UObject/LinkerLoad.h"
 #include "UObject/UObjectHash.h"
@@ -73,7 +74,7 @@ bool FAsyncLoadingTestInvalidateExportDuringPreload::RunTest(const FString& Para
 	auto VerifyLoad = [this](FLoadingTestsScope& LoadingTestScope, bool bExpectToFindObject)
 		{
 			UPackage* Package = LoadPackage(nullptr, LoadingTestScope.PackagePath1, LOAD_None);
-			TestTrue(TEXT("The pacakge should load successfully"), Package != nullptr);
+			TestTrue(TEXT("The package should load successfully"), Package != nullptr);
 			
 			// Exclude garbage objects as the GC won't have run yet but invalidated objects shdould be marked as garbage by this point
 			UAsyncLoadingTests_Shared* Object1 = FindObjectFast<UAsyncLoadingTests_Shared>(Package, LoadingTestScope.ObjectName, EFindObjectFlags::ExactClass, EObjectFlags::RF_MirroredGarbage);
@@ -116,5 +117,130 @@ bool FAsyncLoadingTestInvalidateExportDuringPreload::RunTest(const FString& Para
 	return true;
 }
 
+/**
+ * Validates imported package failures are propagated such that failing imports prevent loading of the importing package
+ */
+IMPLEMENT_SIMPLE_AUTOMATION_TEST(FAsyncLoadingTestFailedImportPropagation, TEXT(TEST_NAME_ROOT ".FailedImportPropagation"), EAutomationTestFlags::EditorContext | EAutomationTestFlags::EngineFilter)
+bool FAsyncLoadingTestFailedImportPropagation::RunTest(const FString& Parameters)
+{
+	class FMockPackageStoreBackend final : public IPackageStoreBackend
+	{
+	public:
+		virtual EPackageLoader GetSupportedLoaders() override { return EPackageLoader::Zen | EPackageLoader::LinkerLoad; }
+		virtual void OnMounted(TSharedRef<const FPackageStoreBackendContext> Context) override {}
+		virtual void BeginRead() override {}
+		virtual void EndRead() override {}
+		virtual bool GetPackageRedirectInfo(FPackageId PackageId, FName& OutSourcePackageName, FPackageId& OutRedirectedToPackageId) override { return false; }
+
+		FMockPackageStoreBackend(const FLoadingTestsScope& LoadingTestScope)
+		{
+			const FPackageId PackageId1 = FPackageId::FromName(FName(LoadingTestScope.PackagePath1));
+			const FPackageId PackageId2 = FPackageId::FromName(FName(LoadingTestScope.PackagePath2));
+			const FPackageId PackageId3 = FPackageId::FromName(FName(LoadingTestScope.PackagePath3));
+
+			PackageImportPackageIds.Add(PackageId1, { PackageId2, PackageId3 });
+			PackageImportPackageIds.Add(PackageId2, { PackageId3 });
+			PackageImportPackageIds.Add(PackageId3, { });
+
+			Entries.Add(PackageId1,
+				{
+#if WITH_EDITORONLY_DATA
+					.LinkerLoadCaseCorrectedPackageName = LoadingTestScope.PackagePath1,
+#endif
+					.PackageExtension = EPackageExtension::Asset,
+					.LoaderType = EPackageLoader::LinkerLoad,
+					.ImportedPackageIds = *PackageImportPackageIds.Find(PackageId1)
+				});
+			Entries.Add(PackageId2,
+				{
+#if WITH_EDITORONLY_DATA
+					.LinkerLoadCaseCorrectedPackageName = LoadingTestScope.PackagePath2,
+#endif
+					.PackageExtension = EPackageExtension::Asset,
+					.LoaderType = EPackageLoader::LinkerLoad,
+					.ImportedPackageIds = *PackageImportPackageIds.Find(PackageId2)
+				});
+			Entries.Add(PackageId3,
+				{
+#if WITH_EDITORONLY_DATA
+					.LinkerLoadCaseCorrectedPackageName = LoadingTestScope.PackagePath3,
+#endif
+					.PackageExtension = EPackageExtension::Asset,
+					.LoaderType = EPackageLoader::LinkerLoad,
+					.ImportedPackageIds = *PackageImportPackageIds.Find(PackageId3)
+				});
+		}
+
+		virtual EPackageStoreEntryStatus GetPackageStoreEntry(FPackageId PackageId, FName PackageName, FPackageStoreEntry& OutPackageStoreEntry) override
+		{
+			if (FPackageStoreEntry* Entry = Entries.Find(PackageId))
+			{
+				OutPackageStoreEntry = *Entry;
+
+				if (NotInstalledPackages.Contains(PackageId))
+				{
+					return EPackageStoreEntryStatus::NotInstalled;
+				}
+
+				return EPackageStoreEntryStatus::Ok;
+			}
+			return EPackageStoreEntryStatus::Missing;
+		}
+
+		void MarkPackageAsNotInstalled(const FPackageId PackageId)
+		{
+			NotInstalledPackages.Add(PackageId);
+		}
+
+		void MarkPackageAsAvailable(const FPackageId PackageId)
+		{
+			NotInstalledPackages.Remove(PackageId);
+		}
+	private:
+		TSet<FPackageId> NotInstalledPackages;
+		TMap<FPackageId, TArray<FPackageId>> PackageImportPackageIds;
+		TMap<FPackageId, FPackageStoreEntry> Entries;
+	};
+
+	// Package 1 has a hard ref to package 2 which has a hard ref to package 3.
+	auto MutateObjects =
+		[](FLoadingTestsScope& Scope)
+		{
+			Scope.Object1->HardReference = Scope.Object2;
+			Scope.Object2->HardReference = Scope.Object3;
+		};
+
+	FLoadingTestsScope LoadingTestScope(this, MutateObjects);
+	const FPackageId PackageId3 = FPackageId::FromName(FName(LoadingTestScope.PackagePath3));
+
+	TSharedPtr<FMockPackageStoreBackend> MockBackend = MakeShared<FMockPackageStoreBackend>(LoadingTestScope);
+	FPackageStore::Get().Mount(MockBackend.ToSharedRef(), INT_MAX /*Priority*/);
+
+	TArray<const TCHAR*> PathsToLoad = { LoadingTestScope.PackagePath3, LoadingTestScope.PackagePath2, LoadingTestScope.PackagePath1 };
+
+	MockBackend->MarkPackageAsNotInstalled(PackageId3);
+
+	// We should the following log message twice; once for Package1 and again for Package2 due to failed Imports. Package3 fails because it does not exist, not due to failed imports
+	AddExpectedMessage(TEXT("will not be loaded due to the NotInstalled imported package"), EAutomationExpectedErrorFlags::Contains, 2 /*Occurrences*/);
+
+	// Verify that no matter which package we load, since a leaf imported package is NotInstalled, none of the packages are loaded
+	for (const TCHAR* PackagePath : PathsToLoad)
+	{
+		TestTrue(TEXT("Package with a failed transitive import should not have loaded"), LoadPackage(nullptr, PackagePath, LOAD_None) == nullptr);
+		TestTrue(TEXT("Package1 should not have been loaded"), FindPackage(nullptr, LoadingTestScope.PackagePath1) == nullptr);
+		TestTrue(TEXT("Package2 should not have been loaded"), FindPackage(nullptr, LoadingTestScope.PackagePath2) == nullptr);
+		TestTrue(TEXT("Package3 should not have been loaded"), FindPackage(nullptr, LoadingTestScope.PackagePath3) == nullptr);
+	}
+	
+	MockBackend->MarkPackageAsAvailable(PackageId3);
+
+	// Now that Package3 is available, verify all packages are now loadable
+	TestTrue(TEXT("Package1 is expected to be loaded"), LoadPackage(nullptr, LoadingTestScope.PackagePath1, LOAD_None) != nullptr);
+	TestTrue(TEXT("Package2 is expected to be implicitly loaded by Package1"), FindPackage(nullptr, LoadingTestScope.PackagePath2) != nullptr);
+	TestTrue(TEXT("Package3 is expected to be implicitly loaded by Package2"), FindPackage(nullptr, LoadingTestScope.PackagePath3) != nullptr);
+
+	return true;
+}
+
 #undef TEST_NAME_ROOT
 #endif // WITH_DEV_AUTOMATION_TESTS
