[AsyncLoading2]

Add missing include affecting mobile platforms

#rnx

diff --git a/Engine/Source/Runtime/CoreUObject/Private/Serialization/AsyncLoading2.cpp b/Engine/Source/Runtime/CoreUObject/Private/Serialization/AsyncLoading2.cpp
--- a/Engine/Source/Runtime/CoreUObject/Private/Serialization/AsyncLoading2.cpp
+++ b/Engine/Source/Runtime/CoreUObject/Private/Serialization/AsyncLoading2.cpp
@@ -16,6 +16,7 @@
 #include "HAL/ThreadManager.h"
 #include "HAL/RunnableThread.h"
 #include "HAL/PlatformMisc.h"
+#include "HAL/PlatformStackWalk.h"
 #include "Misc/ScopeLock.h"
 #include "Stats/StatsMisc.h"
 #include "Containers/StripedMap.h"
