[AsyncLoading2]

Re-enable load failures for missing imports and demote LoadError logs to Info rather than warning matching the existing loader log levels for similar issues.

#rnx

diff --git a/Engine/Source/Runtime/CoreUObject/Private/Serialization/AsyncLoading2.cpp b/Engine/Source/Runtime/CoreUObject/Private/Serialization/AsyncLoading2.cpp
--- a/Engine/Source/Runtime/CoreUObject/Private/Serialization/AsyncLoading2.cpp
+++ b/Engine/Source/Runtime/CoreUObject/Private/Serialization/AsyncLoading2.cpp
@@ -486,7 +486,7 @@ static FAutoConsoleVariableRef CVarGOnlyProcessRequiredPackagesWhenSyncLoading(
 	ECVF_Default
 );
 
-static bool GFailLoadOnFailedImport = false;
+static bool GFailLoadOnFailedImport = true;
 static FAutoConsoleVariableRef CVarGFailLoadOnFailedImport(
 	TEXT("s.FailLoadOnFailedImport"),
 	GFailLoadOnFailedImport,
@@ -6713,7 +6713,7 @@ void FAsyncPackage2::ImportPackagesRecursiveInner(FAsyncLoadingThreadState2& Thr
 						FPackageName::GetExplanationForUnavailablePackage(ImportedPackageNameToLoad, Explanation);
 						if (Explanation.Len())
 						{
-							FMessageLog("LoadErrors").Warning(FText::FormatNamed(NSLOCTEXT("Core", "AsyncLoading_SkippedPackage_Explanation", "While trying to load package {MissingPackage}, a dependent package {DependentPackage} ({DependentPackageId}) was not available. Additional explanatory information follows:\n{Explanation}"),
+							FMessageLog("LoadErrors").Info(FText::FormatNamed(NSLOCTEXT("Core", "AsyncLoading_SkippedPackage_Explanation", "While trying to load package {MissingPackage}, a dependent package {DependentPackage} ({DependentPackageId}) was not available. Additional explanatory information follows:\n{Explanation}"),
 								TEXT("MissingPackage"), FText::FromString(PackagePathToLoadString),
 								TEXT("DependentPackage"), FText::FromString(ImportedPackageNameToLoadString),
 								TEXT("DependentPackageId"), FText::FromString(LexToString(ImportedPackageId)),
@@ -6721,7 +6721,7 @@ void FAsyncPackage2::ImportPackagesRecursiveInner(FAsyncLoadingThreadState2& Thr
 						}
 						else
 						{
-							FMessageLog("LoadErrors").Warning(FText::FormatNamed(NSLOCTEXT("Core", "AsyncLoading_SkippedPackage_NoExplanation", "While trying to load package {MissingPackage}, a dependent package {DependentPackage} ({DependentPackageId}) was not available. No additional explanation was available."),
+							FMessageLog("LoadErrors").Info(FText::FormatNamed(NSLOCTEXT("Core", "AsyncLoading_SkippedPackage_NoExplanation", "While trying to load package {MissingPackage}, a dependent package {DependentPackage} ({DependentPackageId}) was not available. No additional explanation was available."),
 								TEXT("MissingPackage"), FText::FromString(PackagePathToLoadString),
 								TEXT("DependentPackage"), FText::FromString(ImportedPackageNameToLoadString),
 								TEXT("DependentPackageId"), FText::FromString(LexToString(ImportedPackageId))));
@@ -7587,10 +7587,10 @@ EEventLoadNodeExecutionResult FAsyncPackage2::Event_CreateLinkerLoadExports(FAsy
 
 	check(Package->LinkerLoadState.IsSet());
 
-		if (!Package->CreateLinkerLoadExports(ThreadState))
-		{
-			return EEventLoadNodeExecutionResult::Timeout;
-		}
+	if (!Package->CreateLinkerLoadExports(ThreadState))
+	{
+		return EEventLoadNodeExecutionResult::Timeout;
+	}
 
 	Package->AsyncPackageLoadingState = EAsyncPackageLoadingState2::WaitingForLinkerLoadDependencies;
 	Package->ConditionalBeginResolveLinkerLoadImports(ThreadState);
@@ -7919,7 +7919,7 @@ EEventLoadNodeExecutionResult FAsyncPackage2::Event_ProcessExportBundle(FAsyncLo
 	UE::Core::Private::FPlayInEditorLoadingScope PlayInEditorIDScope(Package->Desc.PIEInstanceID);
 #endif
 
-		// This code path should not be reached for LinkerLoad packages
+	// This code path should not be reached for LinkerLoad packages
 	check(!Package->LinkerLoadState.IsSet());
 	check(InExportBundleIndex < Package->Data.TotalExportBundleCount);
 	
@@ -12744,4 +12744,4 @@ bool FPublicExportMapPreInsertPublicExportsTest::RunTest(const FString& Paramete
 	return true;
 }
 
-#endif // WITH_LOW_LEVEL_TESTS
\ No newline at end of file
+#endif // WITH_LOW_LEVEL_TESTS
