Fixed conversion from TArray<bool> in FArrayProperty

#FailedValidations Preflight

diff --git a/Engine/Source/Runtime/CoreUObject/Private/UObject/PropertyArray.cpp b/Engine/Source/Runtime/CoreUObject/Private/UObject/PropertyArray.cpp
--- a/Engine/Source/Runtime/CoreUObject/Private/UObject/PropertyArray.cpp
+++ b/Engine/Source/Runtime/CoreUObject/Private/UObject/PropertyArray.cpp
@@ -1340,12 +1340,30 @@ EConvertFromTypeResult FArrayProperty::ConvertFromType(const FPropertyTag& Tag,
 	FUObjectSerializeContext* Context = FUObjectThreadContext::Get().GetSerializeContext();
 	FStructuredArchive::FStream ValueStream = Slot.EnterStream();
 
-	EConvertFromTypeResult ConvertResult = Inner->ConvertFromType(InnerPropertyTag, ValueStream.EnterElement(), ScriptArrayHelper.GetRawPtr(0), DefaultsStruct, nullptr);
+	TOptional<FBoolProperty> InnerBool;
+	if (InnerPropertyTag.Type == NAME_BoolProperty)
+	{
+		// ConvertFromType expects a bool value to be stored on the property tag.
+		// Call SerializeItem to populate BoolVal before calling ConvertFromType.
+		InnerBool.Emplace(this, GetFName(), RF_NoFlags);
+	}
+
+	FStructuredArchive::FSlot FirstElementSlot = ValueStream.EnterElement();
+	if (InnerBool)
+	{
+		InnerBool->SerializeItem(FirstElementSlot, &InnerPropertyTag.BoolVal, nullptr);
+	}
+	EConvertFromTypeResult ConvertResult = Inner->ConvertFromType(InnerPropertyTag, FirstElementSlot, ScriptArrayHelper.GetRawPtr(0), DefaultsStruct, nullptr);
 	if (ConvertResult == EConvertFromTypeResult::Converted || ConvertResult == EConvertFromTypeResult::Serialized)
 	{
 		for (int32 ElementIndex = 1; ElementIndex < ElementCount; ++ElementIndex)
 		{
-			ConvertResult = Inner->ConvertFromType(InnerPropertyTag, ValueStream.EnterElement(), ScriptArrayHelper.GetRawPtr(ElementIndex), DefaultsStruct, nullptr);
+			FStructuredArchive::FSlot ElementSlot = ValueStream.EnterElement();
+			if (InnerBool)
+			{
+				InnerBool->SerializeItem(ElementSlot, &InnerPropertyTag.BoolVal, nullptr);
+			}
+			ConvertResult = Inner->ConvertFromType(InnerPropertyTag, ElementSlot, ScriptArrayHelper.GetRawPtr(ElementIndex), DefaultsStruct, nullptr);
 			check(ConvertResult == EConvertFromTypeResult::Converted || ConvertResult == EConvertFromTypeResult::Serialized);
 		}
 		return ConvertResult;
diff --git a/Engine/Source/Runtime/CoreUObject/Tests/UObject/TaggedPropertySerializationTest.cpp b/Engine/Source/Runtime/CoreUObject/Tests/UObject/TaggedPropertySerializationTest.cpp
new file mode 100644
--- /dev/null
+++ b/Engine/Source/Runtime/CoreUObject/Tests/UObject/TaggedPropertySerializationTest.cpp
@@ -0,0 +1,78 @@
+// Copyright Epic Games, Inc. All Rights Reserved.
+
+#include "TaggedPropertySerializationTest.h"
+
+#include UE_INLINE_GENERATED_CPP_BY_NAME(TaggedPropertySerializationTest)
+
+#if WITH_TESTS
+
+#include "Serialization/Formatters/JsonArchiveInputFormatter.h"
+#include "Serialization/Formatters/JsonArchiveOutputFormatter.h"
+#include "Serialization/MemoryReader.h"
+#include "Serialization/MemoryWriter.h"
+#include "Serialization/StructuredArchive.h"
+#include "Tests/TestHarnessAdapter.h"
+
+namespace UE
+{
+
+TEST_CASE_NAMED(FTaggedPropertySerializationTest, "CoreUObject::Serialization::TaggedPropertySerialization", "[Core][UObject][EngineFilter]")
+{
+	SECTION("Convert TArray<bool> to TArray<int32>")
+	{
+		TArray<uint8> Data;
+		{
+			FTestTaggedPropertySerializationBoolArray Struct;
+			Struct.Array = { true, false, true, false };
+
+			FMemoryWriter Ar(Data, /*bIsPersistent*/ true);
+			FTestTaggedPropertySerializationBoolArray::StaticStruct()->SerializeItem(Ar, &Struct, nullptr);
+		}
+		{
+			FTestTaggedPropertySerializationInt32Array Struct;
+
+			FMemoryReader Ar(Data, /*bIsPersistent*/ true);
+			FTestTaggedPropertySerializationInt32Array::StaticStruct()->SerializeItem(Ar, &Struct, nullptr);
+
+			REQUIRE(Struct.Array.Num() == 4);
+			CHECK(Struct.Array[0] == 1);
+			CHECK(Struct.Array[1] == 0);
+			CHECK(Struct.Array[2] == 1);
+			CHECK(Struct.Array[3] == 0);
+		}
+	#if WITH_TEXT_ARCHIVE_SUPPORT
+		TArray<uint8> JsonData;
+		{
+			FTestTaggedPropertySerializationBoolArray Struct;
+			Struct.Array = { true, false, true, false };
+
+			FMemoryWriter Ar(JsonData, /*bIsPersistent*/ true);
+			FJsonArchiveOutputFormatter Formatter(Ar);
+			FStructuredArchive StructuredAr(Formatter);
+			FStructuredArchiveRecord Record = StructuredAr.Open().EnterRecord();
+			FStructuredArchiveSlot Slot = Record.EnterField(TEXT("Value"));
+			FTestTaggedPropertySerializationBoolArray::StaticStruct()->SerializeItem(Slot, &Struct, nullptr);
+		}
+		{
+			FTestTaggedPropertySerializationInt32Array Struct;
+
+			FMemoryReader Ar(JsonData, /*bIsPersistent*/ true);
+			FJsonArchiveInputFormatter Formatter(Ar);
+			FStructuredArchive StructuredAr(Formatter);
+			FStructuredArchiveRecord Record = StructuredAr.Open().EnterRecord();
+			FStructuredArchiveSlot Slot = Record.EnterField(TEXT("Value"));
+			FTestTaggedPropertySerializationInt32Array::StaticStruct()->SerializeItem(Slot, &Struct, nullptr);
+
+			REQUIRE(Struct.Array.Num() == 4);
+			CHECK(Struct.Array[0] == 1);
+			CHECK(Struct.Array[1] == 0);
+			CHECK(Struct.Array[2] == 1);
+			CHECK(Struct.Array[3] == 0);
+		}
+	#endif
+	}
+}
+
+} // UE
+
+#endif // WITH_TESTS
diff --git a/Engine/Source/Runtime/CoreUObject/Tests/UObject/TaggedPropertySerializationTest.h b/Engine/Source/Runtime/CoreUObject/Tests/UObject/TaggedPropertySerializationTest.h
new file mode 100644
--- /dev/null
+++ b/Engine/Source/Runtime/CoreUObject/Tests/UObject/TaggedPropertySerializationTest.h
@@ -0,0 +1,33 @@
+// Copyright Epic Games, Inc. All Rights Reserved.
+
+#pragma once
+
+#include "CoreTypes.h"
+#include "Containers/Array.h"
+#include "UObject/Class.h"
+
+#include "TaggedPropertySerializationTest.generated.h"
+
+#if WITH_TESTS
+
+USTRUCT()
+struct FTestTaggedPropertySerializationBoolArray
+{
+	GENERATED_BODY()
+
+public:
+	UPROPERTY()
+	TArray<bool> Array;
+};
+
+USTRUCT()
+struct FTestTaggedPropertySerializationInt32Array
+{
+	GENERATED_BODY()
+
+public:
+	UPROPERTY()
+	TArray<int32> Array;
+};
+
+#endif // WITH_TESTS
