dstm perf:
- allow the enqueueing of transactional "work lists". These are dependent lists of transactions to execute - we can't execute the 2nd until the 1st is complete, etc.
- allow enqueueing of transactions from within transactions
- change PostTickDispatch to first run a transaction that is just meant to enqueue any pending RPCs, which are now enqueued as separate transactions but in a work list so they are guaranteed to run in-order
- this is meant to address issues where a single character gets 20+ move RPCs to process in a single frame, and we end up having to re-run all of them when any of them aborts, causing lots of wasted perf

diff --git a/Engine/Source/Runtime/CoreUObject/Private/UObject/RemoteExecutor.cpp b/Engine/Source/Runtime/CoreUObject/Private/UObject/RemoteExecutor.cpp
--- a/Engine/Source/Runtime/CoreUObject/Private/UObject/RemoteExecutor.cpp
+++ b/Engine/Source/Runtime/CoreUObject/Private/UObject/RemoteExecutor.cpp
@@ -52,6 +52,27 @@ FArchive& operator<<(FArchive& Ar, FRemoteTransactionId& Id)
 	return Ar;
 }
 
+struct FRemoteExecutorWorkList
+{
+	// bWorking is false until the executor starts working on it
+	bool bWorking = false;
+
+	FName Name;
+	FRemoteWorkPriority Priority;
+	bool bIsTransactional = false;
+
+	TArray<TFunction<void(void)>> PendingWork;
+
+	void Reset()
+	{
+		bWorking = false;
+		Name = FName{};
+		Priority = FRemoteWorkPriority{};
+		bIsTransactional = false;
+		PendingWork.Empty();
+	}
+};
+
 struct FRemoteExecutorWork
 {
 	FName Name;
@@ -68,6 +89,10 @@ struct FRemoteExecutorWork
 	
 	bool bRequiresMultiServerCommit = false;
 	FString RequiresMultiServerCommitReason;
+
+    // for work that belongs to a work list, this lets 
+	// us find the worklist in order to enqueue the next item
+	int32 WorkListIndex = INDEX_NONE;
 };
 
 class FRemoteExecutor
@@ -77,6 +102,7 @@ class FRemoteExecutor
 	uint32 NextTransactionRequestId = 0;
 
 	TArray<FRemoteExecutorWork> PendingWorks;
+	TArray<FRemoteExecutorWorkList> PendingWorkLists;
 
 	// GRemoteExecutorForcedNetPumpIterationCount controls how many iterations
 	// we're allowed to go before we unconditionally try to pump the network
@@ -114,14 +140,15 @@ public:
 		return FRemoteTransactionId((RawRequestId++ % 0x800000u) + 1000u);
 	}
 
-	void EnqueueWork(FName WorkName, FRemoteWorkPriority InWorkPriority, bool bIsTransactional, const TFunction<void(void)>& InWork)
+	void EnqueueWork(FName InWorkName, FRemoteWorkPriority InWorkPriority, bool bIsTransactional, int32 InWorkListIndex, const TFunction<void(void)>& InWork)
 	{
 		FRemoteExecutorWork& NewWork = PendingWorks.Emplace_GetRef();
-		NewWork.Name = WorkName;
+		NewWork.Name = InWorkName;
 		NewWork.Priority = InWorkPriority;
 		NewWork.Work = InWork;
 		NewWork.bIsTransactional = bIsTransactional;
 		NewWork.RequestId = GenerateNextTransactionId();
+		NewWork.WorkListIndex = InWorkListIndex;
 
 		if (bIsTransactional)
 		{
@@ -129,15 +156,60 @@ public:
 		}
 	}
 
+	void EnqueueWorkList(FName WorkName, FRemoteWorkPriority InWorkPriority, bool bIsTransactional, const TArray<TFunction<void(void)>>& InWorkList)
+	{
+		FRemoteExecutorWorkList& NewWorkList = PendingWorkLists.Emplace_GetRef();
+
+		NewWorkList.Name = WorkName;
+		NewWorkList.Priority = InWorkPriority;
+		NewWorkList.bIsTransactional = bIsTransactional;
+		NewWorkList.PendingWork = InWorkList;
+	}
+
 	void ExecutePendingWork()
 	{
 		check(ExecutingWork == nullptr);
 
-		if (PendingWorks.Num() > 0)
+		if ((PendingWorks.Num() > 0) || (PendingWorkLists.Num() > 0))
 		{
-			UE_LOG(LogRemoteExec, VeryVerbose, TEXT("ExecutePendingWork BEGIN with %d enqueued work"), PendingWorks.Num());
+			UE_LOG(LogRemoteExec, VeryVerbose, TEXT("ExecutePendingWork BEGIN with %d work items and %d work lists"), PendingWorks.Num(), PendingWorkLists.Num());
 		}
 
+		auto EnqueueFirstWorkFromList = [this](int32 WorkListIndex) -> void
+			{
+				if (WorkListIndex != INDEX_NONE)
+				{
+					FRemoteExecutorWorkList& WorkList = PendingWorkLists[WorkListIndex];
+
+					// is bWorking is already true then we already have work enqueued from this list
+					if (!WorkList.bWorking)
+					{
+						WorkList.bWorking = true;
+
+						if (WorkList.PendingWork.Num() > 0)
+						{
+							EnqueueWork(WorkList.Name, WorkList.Priority, WorkList.bIsTransactional, WorkListIndex, WorkList.PendingWork[0]);
+							WorkList.PendingWork.RemoveAt(0, EAllowShrinking::No);
+						}
+					}
+				}
+			};
+
+		auto EnqueueNextWorkFromList = [this](int32 WorkListIndex) -> void
+			{
+				if (WorkListIndex != INDEX_NONE)
+				{
+					FRemoteExecutorWorkList& WorkList = PendingWorkLists[WorkListIndex];
+					check(WorkList.bWorking);
+
+					if (WorkList.PendingWork.Num() > 0)
+					{
+						EnqueueWork(WorkList.Name, WorkList.Priority, WorkList.bIsTransactional, WorkListIndex, WorkList.PendingWork[0]);
+						WorkList.PendingWork.RemoveAt(0, EAllowShrinking::No);
+					}
+				}
+			};
+
 		double LastStallPrintTime = FPlatformTime::Seconds();
 
 		for (int32 LocalIterationNumber = 0; ; LocalIterationNumber++)
@@ -199,6 +271,12 @@ public:
 				continue;
 			}
 
+			// enqueue the first work item from any new work lists
+			for (int32 WorkListIndex = 0; WorkListIndex < PendingWorkLists.Num(); WorkListIndex++)
+			{
+				EnqueueFirstWorkFromList(WorkListIndex);
+			}
+
 			// subsystem ticking
 			for (int32 SubsystemIndex = 0; SubsystemIndex < Subsystems.Num(); SubsystemIndex++)
 			{
@@ -625,6 +703,7 @@ public:
 						UE::RemoteExecutor::OnTransactionReleasedDelegate.Broadcast(ExecutingWork->RequestId);
 
 						// remove this work from the list and do not increment PendingWorkIndex
+						EnqueueNextWorkFromList(ExecutingWork->WorkListIndex);
 						PendingWorks.RemoveAt(PendingWorkIndex, EAllowShrinking::No);
 					}
 					else
@@ -639,6 +718,7 @@ public:
 					ExecutingWork->Work();
 
 					// remove this work from the list and do not increment PendingWorkIndex
+					EnqueueNextWorkFromList(ExecutingWork->WorkListIndex);
 					PendingWorks.RemoveAt(PendingWorkIndex, EAllowShrinking::No);
 				}
 
@@ -652,6 +732,8 @@ public:
 				}
 			}
 		}
+
+		PendingWorkLists.Empty();
 	}
 };
 
@@ -662,7 +744,7 @@ void ExecuteTransactionalInternal(FName WorkName, FRemoteWorkPriority InWorkPrio
 #if UE_WITH_REMOTE_OBJECT_HANDLE
 	check(!AutoRTFM::IsClosed());
 
-	UE::RemoteExecutor::EnqueueWorkWithExplicitPriority(WorkName, InWorkPriority, true, [&InWork]() { InWork(); });
+	GRemoteExecutor.EnqueueWork(WorkName, InWorkPriority, true, INDEX_NONE, [&InWork]() { InWork(); });
 	UE::RemoteExecutor::ExecutePendingWork();
 #else
 	InWork();
@@ -955,11 +1037,7 @@ void ExecutePendingWork()
 
 void ExecuteTransactionalWithExplicitPriority(FName WorkName, FRemoteWorkPriority WorkPriority, const TFunctionRef<void(void)>& Work)
 {
-#if UE_WITH_REMOTE_OBJECT_HANDLE
 	ExecuteTransactionalInternal(WorkName, WorkPriority, Work);
-#else
-	Work();
-#endif
 }
 
 void ExecuteTransactional(FName WorkName, const TFunctionRef<void(void)>& Work)
@@ -971,7 +1049,7 @@ void ExecuteTransactional(FName WorkName, const TFunctionRef<void(void)>& Work)
 void EnqueueWorkWithExplicitPriority(FName WorkName, FRemoteWorkPriority WorkPriority, bool bIsTransactional, const TFunction<void(void)>& InWork)
 {
 #if UE_WITH_REMOTE_OBJECT_HANDLE
-	GRemoteExecutor.EnqueueWork(WorkName, WorkPriority, bIsTransactional, InWork);
+	GRemoteExecutor.EnqueueWork(WorkName, WorkPriority, bIsTransactional, INDEX_NONE, InWork);
 #else
 	InWork();
 #endif
@@ -979,11 +1057,29 @@ void EnqueueWorkWithExplicitPriority(FName WorkName, FRemoteWorkPriority WorkPri
 
 void EnqueueWork(FName WorkName, bool bIsTransactional, const TFunction<void(void)>& InWork)
 {
-	check(!AutoRTFM::IsClosed());
 	FRemoteWorkPriority RootWorkPriority = FRemoteWorkPriority::CreateRootWorkPriority(FRemoteServerId::GetLocalServerId(), GRemoteExecutor.GenerateNextTransactionId());
 	EnqueueWorkWithExplicitPriority(WorkName, RootWorkPriority, bIsTransactional, InWork);
 }
 
+void EnqueueWorkListWithExplicitPriority(FName WorkName, FRemoteWorkPriority WorkPriority, bool bIsTransactional, const TArray<TFunction<void(void)>>& InWorkList)
+{
+#if UE_WITH_REMOTE_OBJECT_HANDLE
+	GRemoteExecutor.EnqueueWorkList(WorkName, WorkPriority, bIsTransactional, InWorkList);
+#else
+	for(int32 WorkIndex = 0; WorkIndex < InWorkList.Num(); WorkIndex++)
+	{
+		const TFunction<void(void)>& Work = InWorkList[WorkIndex];
+		Work();
+	}
+#endif
+}
+
+void EnqueueWorkList(FName WorkName, bool bIsTransactional, const TArray<TFunction<void(void)>>& InWorkList)
+{
+	FRemoteWorkPriority RootWorkPriority = FRemoteWorkPriority::CreateRootWorkPriority(FRemoteServerId::GetLocalServerId(), GRemoteExecutor.GenerateNextTransactionId());
+	EnqueueWorkListWithExplicitPriority(WorkName, RootWorkPriority, bIsTransactional, InWorkList);
+}
+
 } // namespace UE::RemoteExecutor
 
 FString FRemoteWorkPriority::ToString() const
diff --git a/Engine/Source/Runtime/CoreUObject/Public/UObject/RemoteExecutor.h b/Engine/Source/Runtime/CoreUObject/Public/UObject/RemoteExecutor.h
--- a/Engine/Source/Runtime/CoreUObject/Public/UObject/RemoteExecutor.h
+++ b/Engine/Source/Runtime/CoreUObject/Public/UObject/RemoteExecutor.h
@@ -228,6 +228,8 @@ namespace UE::RemoteExecutor
 	COREUOBJECT_API void ExecuteTransactionalWithExplicitPriority(FName WorkName, FRemoteWorkPriority WorkPriority, const TFunctionRef<void(void)>& Work);
 	COREUOBJECT_API void EnqueueWorkWithExplicitPriority(FName WorkName, FRemoteWorkPriority WorkPriority, bool bIsTransactional, const TFunction<void(void)>& Work);
 	COREUOBJECT_API void EnqueueWork(FName WorkName, bool bIsTransactional, const TFunction<void(void)>& Work);
+	COREUOBJECT_API void EnqueueWorkListWithExplicitPriority(FName WorkName, FRemoteWorkPriority WorkPriority, bool bIsTransactional, const TArray<TFunction<void(void)>>& WorkList);
+	COREUOBJECT_API void EnqueueWorkList(FName WorkName, bool bIsTransactional, const TArray<TFunction<void(void)>>& WorkList);
 	COREUOBJECT_API void ExecutePendingWork();
 
 	COREUOBJECT_API void AbortTransactionRequiresDependencies(FStringView Description);
diff --git a/Engine/Source/Runtime/Engine/Private/DataReplication.cpp b/Engine/Source/Runtime/Engine/Private/DataReplication.cpp
--- a/Engine/Source/Runtime/Engine/Private/DataReplication.cpp
+++ b/Engine/Source/Runtime/Engine/Private/DataReplication.cpp
@@ -2623,7 +2623,7 @@ void FObjectReplicator::UpdateUnmappedObjects(bool & bOutHasMoreUnmapped)
 	}
 }
 
-void FObjectReplicator::ExecutePendingTransactionalRPCs()
+void FObjectReplicator::EnqueuePendingTransactionalRPCs()
 {
 	ensureMsgf(AutoRTFM::IsClosed(), TEXT("%hs is meant to be executed inside of a transaction"), __func__);
 	if (!ensure(Connection) || !ensure(Connection->Driver) || !ensure(Connection->Driver->NetCache))
@@ -2644,66 +2644,76 @@ void FObjectReplicator::ExecutePendingTransactionalRPCs()
 	}
 
 	const UObject* Object = GetObject();
-	const FClassNetCache* ClassCache = Connection->Driver->NetCache->GetClassNetCache(ObjectClass);
+	FString ObjectFullName = Object->GetFullName();
 
 	// Handle pending RPCs, in order
 	TArray<FRPCPendingLocalCall> ServiceTheseRPCs;
 	Swap(PendingTransactionalRPCs, ServiceTheseRPCs);
 
+	TArray<TFunction<void(void)>> TransactionsToEnqueue;
+	TransactionsToEnqueue.Reserve(ServiceTheseRPCs.Num());
+
 	for (const FRPCPendingLocalCall& Pending : ServiceTheseRPCs)
 	{
-		ensureMsgf(Pending.UnmappedGuids.IsEmpty(), TEXT("PendingTransactionalRPCs are never expected to have UnmappedGuids (UnmappedGuids are in PendingLocalRPCs)"));
-
-		bool bSuccess;
-		const FFieldNetCache* FieldCache = ClassCache->GetFromIndex(Pending.RPCFieldIndex);			
-		if (FieldCache == nullptr)
-		{
-			UE_LOG(LogNet, Warning, TEXT("%hs: FieldCache not found. Object: %s"), __func__, *Object->GetFullName());
-			bSuccess = false;
-		}
-		else
-		{
-			// Help us keep track of this function's execution time
-			FName FunctionName = FieldCache->Field.GetFName();
-			UE::RemoteExecutor::OnRegionBeginDelegate.ExecuteIfBound(FunctionName);
-			AutoRTFM::PushOnAbortHandler(&Pending, []()
+		TransactionsToEnqueue.Emplace([this, Pending, ObjectFullName]()
 			{
-				UE::RemoteExecutor::OnRegionEndDelegate.ExecuteIfBound(TEXT("Abort"));
-			});
+				const FClassNetCache* ClassCache = Connection->Driver->NetCache->GetClassNetCache(ObjectClass);
 
-			FNetBitReader Reader(Connection->PackageMap, Pending.Buffer.GetData(), Pending.NumBits);
-			Connection->SetNetVersionsOnArchive(Reader);
+				ensureMsgf(Pending.UnmappedGuids.IsEmpty(), TEXT("PendingTransactionalRPCs are never expected to have UnmappedGuids (UnmappedGuids are in PendingLocalRPCs)"));
 
-			TSet<FNetworkGUID> UnmappedGuids;
-			bool bSkipped = false;
-			bSuccess = ReceivedRPC(Reader, UnmappedGuids, bSkipped, Pending.RepFlags, FieldCache, ESkipRpcBehavior::NeverSkip);
+				bool bSuccess;
+				const FFieldNetCache* FieldCache = ClassCache->GetFromIndex(Pending.RPCFieldIndex);			
+				if (FieldCache == nullptr)
+				{
+					UE_LOG(LogNet, Warning, TEXT("%hs: FieldCache not found. Object: %s"), __func__, *ObjectFullName);
+					bSuccess = false;
+				}
+				else
+				{
+					// Help us keep track of this function's execution time
+					FName FunctionName = FieldCache->Field.GetFName();
+					UE::RemoteExecutor::OnRegionBeginDelegate.ExecuteIfBound(FunctionName);
+					AutoRTFM::PushOnAbortHandler(&Pending, []()
+					{
+						UE::RemoteExecutor::OnRegionEndDelegate.ExecuteIfBound(TEXT("Abort"));
+					});
 
-			ensureMsgf(UnmappedGuids.IsEmpty(), TEXT("%hs: Servers are expected to never have UnmappedGuids"), __func__);
-			ensureMsgf(!bSkipped, TEXT("%hs: We should never skip an RPC when requesting NeverSkip"), __func__);
+					FNetBitReader Reader(Connection->PackageMap, Pending.Buffer.GetData(), Pending.NumBits);
+					Connection->SetNetVersionsOnArchive(Reader);
 
-			UE::RemoteExecutor::OnRegionEndDelegate.ExecuteIfBound(TEXT(""));
-			AutoRTFM::PopOnAbortHandler(&Pending);
-		}
+					TSet<FNetworkGUID> UnmappedGuids;
+					bool bSkipped = false;
+					bSuccess = ReceivedRPC(Reader, UnmappedGuids, bSkipped, Pending.RepFlags, FieldCache, ESkipRpcBehavior::NeverSkip);
 
-		if (!bSuccess)
-		{
-			FString FunctionName = FieldCache ? FieldCache->Field.GetName() : TEXT("Unknown");
+					ensureMsgf(UnmappedGuids.IsEmpty(), TEXT("%hs: Servers are expected to never have UnmappedGuids"), __func__);
+					ensureMsgf(!bSkipped, TEXT("%hs: We should never skip an RPC when requesting NeverSkip"), __func__);
 
-			const bool bDemoMode = Connection->IsInternalAck();
-			if (!bDemoMode)
-			{
-				// Close our connection and abort rpcs as things are invalid
-				UE_LOG(LogNet, Error, TEXT("%hs: Failed executing delayed RPC %s on Object %s, closing connection!"), __func__, *FunctionName, *Object->GetFullName());
+					UE::RemoteExecutor::OnRegionEndDelegate.ExecuteIfBound(TEXT(""));
+					AutoRTFM::PopOnAbortHandler(&Pending);
+				}
 
-				Connection->Close();
-				return;
-			}
-			else
-			{
-				UE_LOG(LogNet, Warning, TEXT("%hs: Failed executing delayed RPC %s on Object %s, skipping RPC!"), __func__, *FunctionName, *Object->GetFullName());
-			}
-		}
+				if (!bSuccess)
+				{
+					FString FunctionName = FieldCache ? FieldCache->Field.GetName() : TEXT("Unknown");
+
+					const bool bDemoMode = Connection->IsInternalAck();
+					if (!bDemoMode)
+					{
+						// Close our connection and abort rpcs as things are invalid
+						UE_LOG(LogNet, Error, TEXT("%hs: Failed executing delayed RPC %s on Object %s, closing connection!"), __func__, *FunctionName, *ObjectFullName);
+
+						Connection->Close();
+						return;
+					}
+					else
+					{
+						UE_LOG(LogNet, Warning, TEXT("%hs: Failed executing delayed RPC %s on Object %s, skipping RPC!"), __func__, *FunctionName, *ObjectFullName);
+					}
+				}
+			});
 	}
+
+	UE::RemoteExecutor::EnqueueWorkList(Object->GetFName(), true, TransactionsToEnqueue);
 }
 
 void FObjectReplicator::QueuePropertyRepNotify(
diff --git a/Engine/Source/Runtime/Engine/Private/NetDriver.cpp b/Engine/Source/Runtime/Engine/Private/NetDriver.cpp
--- a/Engine/Source/Runtime/Engine/Private/NetDriver.cpp
+++ b/Engine/Source/Runtime/Engine/Private/NetDriver.cpp
@@ -2783,7 +2783,7 @@ void UNetDriver::PostTickDispatch()
 
 			for (FObjectReplicator* Replicator : TransactingReplicators)
 			{
-				static const FName WorkName { __func__ };
+				static const FName WorkName { TEXT("EnqueuePendingTransactionalRPCs") };
 				constexpr bool bIsTransactional = true;
 				UE::RemoteExecutor::EnqueueWork(WorkName, bIsTransactional, [Replicator, &TransactingReplicators]()
 				{
@@ -2800,10 +2800,10 @@ void UNetDriver::PostTickDispatch()
 							UE::RemoteExecutor::OnRegionEndDelegate.ExecuteIfBound(TEXT("Commit"));
 						};
 
-						Replicator->ExecutePendingTransactionalRPCs();
-					}
+						Replicator->EnqueuePendingTransactionalRPCs();
 
-					TransactingReplicators.Remove(Replicator);
+						TransactingReplicators.Remove(Replicator);
+					}
 				});
 			}
 
diff --git a/Engine/Source/Runtime/Engine/Public/Net/DataReplication.h b/Engine/Source/Runtime/Engine/Public/Net/DataReplication.h
--- a/Engine/Source/Runtime/Engine/Public/Net/DataReplication.h
+++ b/Engine/Source/Runtime/Engine/Public/Net/DataReplication.h
@@ -223,10 +223,10 @@ public:
 	ENGINE_API void UpdateUnmappedObjects(bool& bOutHasMoreUnmapped);
 
 	/**
-	 * Execute any pending RPCs meant to be done within an AutoRTFM transaction.
+	 * Enqueue any pending RPCs meant to be done within an AutoRTFM transaction.
 	 * @see UNetDriver::PostTickDispatch
 	 */
-	ENGINE_API void ExecutePendingTransactionalRPCs();
+	ENGINE_API void EnqueuePendingTransactionalRPCs();
 
 	inline TWeakObjectPtr<UObject>	GetWeakObjectPtr() const
 	{
@@ -363,7 +363,7 @@ private:
 
 	/**
 	 * These are RPCs that are saved so that they may be executed transactionally at a later time
-	 * @see ExecutePendingTransactionalRPCs
+	 * @see EnqueuePendingTransactionalRPCs
 	 */
 	TArray<FRPCPendingLocalCall> PendingTransactionalRPCs;
 
