[Mass] BitSet changes to support "single element bitsets"

* alternative implementations of bitwise operators
* default to 8 elements as bitarray's initial allocation size
* modified some functions' signatures to support more flexibility

diff --git a/Engine/Source/Runtime/CoreUObject/Private/StructUtils/StructTypeBitSet.cpp b/Engine/Source/Runtime/CoreUObject/Private/StructUtils/StructTypeBitSet.cpp
--- a/Engine/Source/Runtime/CoreUObject/Private/StructUtils/StructTypeBitSet.cpp
+++ b/Engine/Source/Runtime/CoreUObject/Private/StructUtils/StructTypeBitSet.cpp
@@ -125,6 +125,8 @@ int32 FStructTracker::RegisterImplementation(const UStruct& InStructType, const
 	check(StructTypesList.Num() == NewIndex);
 	StructTypesList.Add(&InStructType);
 
+	OnTypeRegistered.Broadcast(&InStructType, NewIndex);
+
 	if (bIsSerializable)
 	{
 		// first-time SerializationHash initialization
diff --git a/Engine/Source/Runtime/CoreUObject/Public/StructUtils/StructTypeBitSet.h b/Engine/Source/Runtime/CoreUObject/Public/StructUtils/StructTypeBitSet.h
--- a/Engine/Source/Runtime/CoreUObject/Public/StructUtils/StructTypeBitSet.h
+++ b/Engine/Source/Runtime/CoreUObject/Public/StructUtils/StructTypeBitSet.h
@@ -16,12 +16,12 @@ namespace FStructTypeBitSet
 {
 
 /**
- * A constant bitset container that extends TBitArray<> with additional utilities.
+ * A constant bitset container that extends TBitArray with additional utilities.
  * Used for operations such as bitwise checks and hashing
  */
-struct FConstBitSetContainer : TBitArray<>
+struct FConstBitSetContainer : TBitArray<TInlineAllocator<8>>
 {
-	using Super = TBitArray<>;
+	using FBitArrayBase = TBitArray;
 
 	FConstBitSetContainer() = default;
 
@@ -29,7 +29,7 @@ struct FConstBitSetContainer : TBitArray<>
 	 * Checks if all bits set in the Other bit array are also set in this bitset
 	 * @return True if all bits in Other are set in this bitset
 	 */
-	inline bool HasAll(const TBitArray<>& Other) const
+	bool HasAll(const FBitArrayBase& Other) const
 	{
 		FConstWordIterator ThisIterator(*this);
 		FConstWordIterator OtherIterator(Other);
@@ -54,7 +54,7 @@ struct FConstBitSetContainer : TBitArray<>
 	 * Checks if any bits in the Other bit array are also set in this bitset
 	 * @return True if any bits in Other are set in this bitset
 	 */
-	inline bool HasAny(const TBitArray<>& Other) const
+	bool HasAny(const FBitArrayBase& Other) const
 	{
 		FConstWordIterator ThisIterator(*this);
 		FConstWordIterator OtherIterator(Other);
@@ -75,10 +75,76 @@ struct FConstBitSetContainer : TBitArray<>
 		return false;
 	}
 
+	FConstBitSetContainer operator&(const FBitArrayBase& Other) const
+	{
+		const int32 MinNumBits = FMath::Min(Num(), Other.Num());
+		if (MinNumBits > 0)
+		{
+			FBitArrayBase Result(false, MinNumBits);
+
+			FConstWordIterator IteratorA(*this);
+			FConstWordIterator IteratorB(Other);
+
+			FWordIterator IteratorResult(Result);
+
+			for ( ; IteratorResult; ++IteratorResult, ++IteratorA, ++IteratorB)
+			{
+				const uint32 NewValue = IteratorA.GetWord() & IteratorB.GetWord();
+				IteratorResult.SetWord(NewValue);
+			}
+
+			return FConstBitSetContainer(MoveTemp(Result));
+		}
+
+		return {};
+	}
+
+	FConstBitSetContainer operator|(const FBitArrayBase& Other) const
+	{
+		const int32 MaxNumBits = FMath::Max(Num(), Other.Num());
+		if (MaxNumBits > 0)
+		{
+			FBitArrayBase Result(false, MaxNumBits);
+
+			FConstWordIterator IteratorA(*this);
+			FConstWordIterator IteratorB(Other);
+
+			FWordIterator IteratorResult(Result);
+
+			/*const int32 MinNumBits = FMath::Min(Num(), Other.Num());
+
+			for (int32 WordIndex = 0; IteratorResult; ++IteratorResult, ++IteratorA, ++IteratorB)
+			{
+				const uint32 NewValue = IteratorA.GetWord() | IteratorB.GetWord();
+				IteratorResult.SetWord(NewValue);
+			}*/
+
+			for (; IteratorA && IteratorB; ++IteratorResult, ++IteratorA, ++IteratorB)
+			{
+				const uint32 NewValue = IteratorA.GetWord() | IteratorB.GetWord();
+				IteratorResult.SetWord(NewValue);
+			}
+			for (; IteratorA; ++IteratorResult, ++IteratorA)
+			{
+				const uint32 NewValue = IteratorA.GetWord();
+				IteratorResult.SetWord(NewValue);
+			}
+			for (; IteratorB; ++IteratorResult, ++IteratorB)
+			{
+				const uint32 NewValue = IteratorB.GetWord();
+				IteratorResult.SetWord(NewValue);
+			}
+
+			return FConstBitSetContainer(MoveTemp(Result));
+		}
+
+		return {};
+	}
+
 	/**
 	 * @return whether the bitset is empty (no bits are set)
 	 */
-	inline bool IsEmpty() const
+	bool IsEmpty() const
 	{
 		FConstWordIterator Iterator(*this);
 
@@ -96,14 +162,14 @@ struct FConstBitSetContainer : TBitArray<>
 	 */
 	bool operator==(const FConstBitSetContainer& Other) const
 	{
-		return (IsEmpty() && Other.IsEmpty()) || Super::operator==(Other);
+		return (IsEmpty() && Other.IsEmpty()) || FBitArrayBase::operator==(Other);
 	}
 
 	/**
 	 * Computes the hash value for the given bitset
 	 * @return The computed hash value
 	 */
-	inline friend uint32 GetTypeHash(const FConstBitSetContainer& Instance)
+	friend uint32 GetTypeHash(const FConstBitSetContainer& Instance)
 	{
 		FConstWordIterator Iterator(Instance);
 		uint32 Hash = 0;
@@ -150,8 +216,8 @@ protected:
 	 * Copy-constructor for creating a bitset from another bit array
 	 * @param Source - The source bit array
 	 */
-	explicit FConstBitSetContainer(const TBitArray<>& Source)
-		: Super(Source)
+	explicit FConstBitSetContainer(const FBitArrayBase& Source)
+		: FBitArrayBase(Source)
 	{
 	}
 
@@ -159,8 +225,8 @@ protected:
 	 * Move-constructor for creating a bitset from another bit array
 	 * @param Source - The source bit array to move from
 	 */
-	explicit FConstBitSetContainer(TBitArray<>&& Source)
-		: Super(Forward<TBitArray<>>(Source))
+	explicit FConstBitSetContainer(FBitArrayBase&& Source)
+		: FBitArrayBase(Forward<FBitArrayBase>(Source))
 	{
 	}
 };
@@ -174,20 +240,24 @@ struct FBitSetContainer : FConstBitSetContainer
 
 	FBitSetContainer() = default;
 
-	explicit FBitSetContainer(const TBitArray<>& Source)
+	explicit FBitSetContainer(const FBitArrayBase& Source)
 		: FConstBitSetContainer(Source)
 	{
 	}
-	explicit FBitSetContainer(TBitArray<>&& Source)
-		: FConstBitSetContainer(Forward<TBitArray<>>(Source))
+	explicit FBitSetContainer(FBitArrayBase&& Source)
+		: FConstBitSetContainer(Forward<FBitArrayBase>(Source))
+	{
+	}
+	FBitSetContainer(FConstBitSetContainer&& Source)
+		: FConstBitSetContainer(Forward<FConstBitSetContainer>(Source))
 	{
 	}
 
-	FBitSetContainer& operator=(const TBitArray<>& Other)
+	FBitSetContainer& operator=(const FBitArrayBase& Other)
 	{
-		static_assert(sizeof(FBitSetContainer) == sizeof(TBitArray<>)
-			, "This operator requires FBitSetContainer and TBitArray<> sizes to match (i.e. FBitSetContainer cannot have member variables nor virtual funtions)");
-		*((TBitArray<>*)this) = Other;
+		static_assert(sizeof(FBitSetContainer) == sizeof(FBitArrayBase)
+			, "This operator requires FBitSetContainer and FBitArrayBase sizes to match (i.e. FBitSetContainer cannot have member variables nor virtual funtions)");
+		*((FBitArrayBase*)this) = Other;
 		return *this;
 	}
 
@@ -206,7 +276,7 @@ struct FBitSetContainer : FConstBitSetContainer
 	 * Removes bits set in the Other bit array from this bitset
 	 * @param Other - The bit array to remove
 	 */
-	inline void operator-=(const TBitArray<>& Other)
+	FBitSetContainer& operator-=(const FBitArrayBase& Other)
 	{
 		FWordIterator ThisIterator(*this);
 		FConstWordIterator OtherIterator(Other);
@@ -218,6 +288,26 @@ struct FBitSetContainer : FConstBitSetContainer
 			++ThisIterator;
 			++OtherIterator;
 		}
+		return *this;
+	}
+
+	FBitSetContainer& operator+=(const FBitArrayBase& Other)
+	{
+		// make sure we have enough space for all the contents of Other
+		// not that it's fine to pass in a negative number, TBitArray::Add handles that gracefully. 
+		Add(false, Other.Num() - Num());
+		
+		FWordIterator ThisIterator(*this);
+		FConstWordIterator OtherIterator(Other);
+
+		while (OtherIterator)
+		{
+			ThisIterator.SetWord(ThisIterator.GetWord() | OtherIterator.GetWord());
+
+			++ThisIterator;
+			++OtherIterator;
+		}
+		return *this;
 	}
 
 	/**
@@ -276,6 +366,8 @@ struct FStructTracker
 	using FBaseStructGetter = TFunction<UStruct*()>;
 	using FTypeValidation = TFunction<bool(const UStruct*)>;
 
+	DECLARE_MULTICAST_DELEGATE_TwoParams(FOnTypeRegistered, TNotNull<const UStruct*> /*Type*/, int32 /*Index*/);
+
 	/** 
 	 * The input parameter is a function that fetches the UStruct representing the base class for all the stored structs.
 	 * We're unable to get the UStruct parameter directly since FStructTracker instances are being created during
@@ -337,6 +429,8 @@ struct FStructTracker
 	 */
 	UE_API void Serialize(FArchive& Ar, FStructTypeBitSet::FBitSetContainer& StructTypesBitArray);
 
+	FOnTypeRegistered OnTypeRegistered;
+
 #if WITH_STRUCTUTILS_DEBUG
 	/**
 	 * @return Name identifying the given struct type or NAME_None if it has never been used/seen before.
@@ -404,6 +498,7 @@ struct TTypeBitSetBase
 	using FUStructType = TStructType;
 	using FBaseStruct = TBaseStruct;
 	using FContainer = TBitSetContainer;
+	using FBitArrayBase = FStructTypeBitSet::FConstBitSetContainer::FBitArrayBase;
 
 	/**
 	 * Iterator for traversing indices of bits with a specified value (true or false).
@@ -445,7 +540,7 @@ struct TTypeBitSetBase
 		}
 
 	private:
-		TBitArray<>::FConstIterator It;
+		FBitArrayBase::FConstIterator It;
 		const bool bValueToCheck;
 	};
 
@@ -481,7 +576,7 @@ struct TTypeBitSetBase
 	 * Retrieves the base UStruct type used for validation.
 	 * @return Pointer to the base UStruct.
 	 */
-	inline static const UStruct* GetBaseUStruct()
+	static const UStruct* GetBaseUStruct()
 	{
 		if constexpr (bTestInheritanceAtRuntime)
 		{
@@ -591,7 +686,7 @@ struct TTypeBitSetBase
 	 */
 	UE_FORCEINLINE_HINT TImplementation operator&(const TImplementation& Other) const
 	{
-		return TImplementation(TBitArray<>::BitwiseAND(StructTypesBitArray, Other.StructTypesBitArray, EBitwiseOperatorFlags::MinSize));
+		return TImplementation(StructTypesBitArray & Other.StructTypesBitArray);
 	}
 
 	/**
@@ -601,7 +696,7 @@ struct TTypeBitSetBase
 	 */
 	UE_FORCEINLINE_HINT TImplementation operator|(const TImplementation& Other) const
 	{
-		return TImplementation(TBitArray<>::BitwiseOR(StructTypesBitArray, Other.StructTypesBitArray, EBitwiseOperatorFlags::MaxSize));
+		return TImplementation(StructTypesBitArray | Other.StructTypesBitArray);
 	}
 
 	UE_FORCEINLINE_HINT TImplementation GetOverlap(const TImplementation& Other) const
@@ -620,6 +715,35 @@ struct TTypeBitSetBase
 		return StructTypesBitArray.CompareSetBits(Other.StructTypesBitArray, /*bMissingBitValue=*/false);
 	}
 
+	bool IsEquivalentFiltered(const TImplementation& Other, const TImplementation& Filter) const
+	{
+		constexpr uint32 MissingBitsFill = 0;
+
+		FBitArrayBase::FConstWordIterator ThisIterator(StructTypesBitArray);
+		FBitArrayBase::FConstWordIterator OtherIterator(Other.StructTypesBitArray);
+		FBitArrayBase::FConstWordIterator FilterIterator(Filter.StructTypesBitArray);
+
+		ThisIterator.FillMissingBits(MissingBitsFill);
+		OtherIterator.FillMissingBits(MissingBitsFill);
+		FilterIterator.FillMissingBits(MissingBitsFill);
+
+		while (FilterIterator && (ThisIterator || OtherIterator))
+		{
+			const uint32 A = ThisIterator  ? (ThisIterator.GetWord() & FilterIterator.GetWord()) : MissingBitsFill;
+			const uint32 B = OtherIterator ? (OtherIterator.GetWord() & FilterIterator.GetWord()) : MissingBitsFill;
+			if (A != B)
+			{
+				return false;
+			}
+
+			++ThisIterator;
+			++OtherIterator;
+			++FilterIterator;
+		}
+
+		return true;
+	}
+
 	UE_FORCEINLINE_HINT bool HasAll(const TImplementation& Other) const
 	{
 		return StructTypesBitArray.HasAll(Other.StructTypesBitArray);
@@ -658,18 +782,20 @@ struct TTypeBitSetBase
 	 * Adds the bits from another bitset to this one (union).
 	 * @param Other - The other bitset to add.
 	 */
-	UE_FORCEINLINE_HINT void operator+=(const TImplementation& Other)
+	TTypeBitSetBase& operator+=(const TImplementation& Other)
 	{
-		StructTypesBitArray = TBitArray<>::BitwiseOR(StructTypesBitArray, Other.StructTypesBitArray, EBitwiseOperatorFlags::MaxSize);
+		StructTypesBitArray += Other.StructTypesBitArray;
+		return *this;
 	}
 
 	/**
 	 * Removes the bits from another bitset from this one (difference).
 	 * @param Other - The other bitset to subtract.
 	 */
-	UE_FORCEINLINE_HINT void operator-=(const TImplementation& Other)
+	UE_FORCEINLINE_HINT TTypeBitSetBase& operator-=(const TImplementation& Other)
 	{
 		StructTypesBitArray -= Other.StructTypesBitArray;
+		return *this;
 	}
 
 	/**
@@ -677,7 +803,7 @@ struct TTypeBitSetBase
 	 * @param NewElement - The struct type to add.
 	 * @return A new bitset with the added struct type.
 	 */
-	inline TImplementation operator+(const FUStructType& NewElement) const
+	TImplementation operator+(const FUStructType& NewElement) const
 	{
 		TImplementation Result = GetImplementation();
 		Result.Add(NewElement);
@@ -689,7 +815,7 @@ struct TTypeBitSetBase
 	 * @param NewElement - The struct type to remove.
 	 * @return A new bitset with the struct type removed.
 	 */
-	inline TImplementation operator-(const FUStructType& NewElement) const
+	TImplementation operator-(const FUStructType& NewElement) const
 	{
 		TImplementation Result = GetImplementation();
 		Result.Remove(NewElement);
@@ -714,7 +840,7 @@ struct TTypeBitSetBase
 	template<typename TOutStructType, typename Allocator>
 	void ExportTypes(TArray<const TOutStructType*, Allocator>& OutTypes) const
 	{
-		TBitArray<>::FConstIterator It(StructTypesBitArray);
+		FBitArrayBase::FConstIterator It(StructTypesBitArray);
 		while (It)
 		{
 			if (It.GetValue())
@@ -775,15 +901,20 @@ protected:
 	 * Constructor accepting a container.
 	 * @param InContainer - The bitset container to use.
 	 */
-	TTypeBitSetBase(FContainer& InContainer)
+	explicit TTypeBitSetBase(const FContainer& InContainer)
 		: StructTypesBitArray(InContainer)
 	{
 	}
 
+	explicit TTypeBitSetBase(FBitArrayBase&& InContainer)
+		: StructTypesBitArray(MoveTemp(InContainer))
+	{
+	}
+
 #if WITH_STRUCTUTILS_DEBUG
 	// For unit testing purposes only
-	const TBitArray<>& DebugGetStructTypesBitArray() const { return StructTypesBitArray; }
-	TBitArray<>& DebugGetMutableStructTypesBitArray() { return StructTypesBitArray; }
+	const FBitArrayBase& DebugGetStructTypesBitArray() const { return StructTypesBitArray; }
+	FBitArrayBase& DebugGetMutableStructTypesBitArray() { return StructTypesBitArray; }
 #endif // WITH_STRUCTUTILS_DEBUG
 
 	/** The bitset container storing the bits representing struct types */
@@ -824,6 +955,8 @@ struct TStructTypeBitSet : TTypeBitSetBase<TStructTypeBitSet<TBaseStruct, TStruc
 	using FStructTrackerWrapper = TStructTrackerWrapper;
 	using FUStructType = TUStructType;
 	using FBaseStruct = TBaseStruct;
+	using FBitSetContainer = FStructTypeBitSet::FBitSetContainer;
+	using FBitArrayBase = FStructTypeBitSet::FConstBitSetContainer::FBitArrayBase;
 
 	friend Super;
 	using Super::StructTypesBitArray;
@@ -838,6 +971,10 @@ struct TStructTypeBitSet : TTypeBitSetBase<TStructTypeBitSet<TBaseStruct, TStruc
 	using Super::ExportTypes;
 
 	TStructTypeBitSet() = default;
+	TStructTypeBitSet(const TStructTypeBitSet&) = default;
+	TStructTypeBitSet(TStructTypeBitSet&&) = default;
+	TStructTypeBitSet& operator=(const TStructTypeBitSet&) = default;
+	TStructTypeBitSet& operator=(TStructTypeBitSet&&) = default;
 
 	/**
 	 * Constructor that initializes the bitset with a single struct type.
@@ -901,17 +1038,17 @@ private:
 	 * Note that this constructor needs to remain private to ensure consistency of stored values with data tracked 
 	 * by the TStructTrackerWrapper::StructTracker.
 	 */
-	TStructTypeBitSet(TBitArray<>&& Source)
+	TStructTypeBitSet(FBitArrayBase&& Source)
+		: Super(MoveTemp(Source))
 	{
-		StructTypesBitArray = Forward<TBitArray<>>(Source);
 	}
 
-	TStructTypeBitSet(const TBitArray<>& Source)
+	TStructTypeBitSet(const FBitArrayBase& Source)
 	{
 		StructTypesBitArray = Source;
 	}
 
-	TStructTypeBitSet(const int32 BitToSet)
+	explicit TStructTypeBitSet(const int32 BitToSet)
 	{
 		StructTypesBitArray.AddAtIndex(BitToSet);
 	}
@@ -993,7 +1130,7 @@ public:
 	 * @param T - The struct type to add.
 	 */
 	template<typename T>
-	inline void Add()
+	void Add()
 	{
 		static_assert(TIsDerivedFrom<T, TBaseStruct>::IsDerived, "Given struct type doesn't match the expected base struct type.");
 		const int32 StructTypeIndex = GetTypeIndex<T>();
@@ -1005,7 +1142,7 @@ public:
 	 * @param T - The struct type to remove.
 	 */
 	template<typename T>
-	inline void Remove()
+	void Remove()
 	{
 		static_assert(TIsDerivedFrom<T, TBaseStruct>::IsDerived, "Given struct type doesn't match the expected base struct type.");
 		const int32 StructTypeIndex = GetTypeIndex<T>();
@@ -1027,7 +1164,7 @@ public:
 	 * @return True if the struct type is in the bitset; false otherwise.
 	 */
 	template<typename T>
-	inline bool Contains() const
+	bool Contains() const
 	{
 		static_assert(TIsDerivedFrom<T, TBaseStruct>::IsDerived, "Given struct type doesn't match the expected base struct type.");
 		const int32 StructTypeIndex = GetTypeIndex<T>();
@@ -1039,11 +1176,9 @@ public:
 	 * @param Other - The other bitset to combine with.
 	 * @return A new bitset representing the union.
 	 */
-	inline TStructTypeBitSet operator+(const TStructTypeBitSet& Other) const
+	TStructTypeBitSet operator+(const TStructTypeBitSet& Other) const
 	{
-		TStructTypeBitSet Result;
-		Result.StructTypesBitArray = TBitArray<>::BitwiseOR(StructTypesBitArray, Other.StructTypesBitArray, EBitwiseOperatorFlags::MaxSize);
-		return MoveTemp(Result);
+		return (*this) | Other;
 	}
 
 	/**
@@ -1080,7 +1215,7 @@ public:
 	 * @param Other - The other bitset to subtract.
 	 * @return A new bitset representing the difference.
 	 */
-	inline TStructTypeBitSet operator-(const TStructTypeBitSet& Other) const
+	TStructTypeBitSet operator-(const TStructTypeBitSet& Other) const
 	{
 		TStructTypeBitSet Result = *this;
 		Result -= Other;
@@ -1097,7 +1232,7 @@ public:
 	 */
 	void ExportTypes(TFunctionRef<bool(const FUStructType*)> Callback) const
 	{
-		TBitArray<>::FConstIterator It(Super::StructTypesBitArray);
+		FBitArrayBase::FConstIterator It(Super::StructTypesBitArray);
 		bool bKeepGoing = true;
 		while (bKeepGoing && It)
 		{
@@ -1174,7 +1309,7 @@ public:
 	 * @param Instance - The bitset instance to hash.
 	 * @return The hash value.
 	 */
-	inline friend uint32 GetTypeHash(const TStructTypeBitSet& Instance)
+	friend uint32 GetTypeHash(const TStructTypeBitSet& Instance)
 	{
 		const uint32 StoredTypeHash = PointerHash(GetBaseUStruct());
 		const uint32 BitArrayHash = GetTypeHash(Instance.StructTypesBitArray);
@@ -1196,7 +1331,7 @@ public:
 	 * @param Instance - The bitset instance to serialize.
 	 * @return The archive after serialization.
 	 */
-	inline friend FArchive& operator<<(FArchive& Ar, TStructTypeBitSet& Instance)
+	friend FArchive& operator<<(FArchive& Ar, TStructTypeBitSet& Instance)
 	{
 		Instance.Serialize(Ar);
 		return Ar;
@@ -1215,7 +1350,7 @@ public:
 		EXPORTED_API static FStructTracker StructTracker; \
 	}; \
 	using ContainerTypeName = TStructTypeBitSet<BaseType, ContainerTypeName##StructTrackerWrapper, BaseUStructType>; \
-	static_assert(std::is_move_constructible_v<ContainerTypeName> && std::is_move_assignable_v<ContainerTypeName>)
+	static_assert(std::is_move_constructible_v<ContainerTypeName> && std::is_move_assignable_v<ContainerTypeName>, "TStructTypeBitSet as configured is expected to be movable")
 
 #define DECLARE_STRUCTTYPEBITSET_EXPORTED(EXPORTED_API, ContainerTypeName, BaseStructType) _DECLARE_TYPEBITSET_IMPL(EXPORTED_API, ContainerTypeName, BaseStructType, UScriptStruct)
 #define DECLARE_STRUCTTYPEBITSET(ContainerTypeName, BaseStructType) _DECLARE_TYPEBITSET_IMPL(, ContainerTypeName, BaseStructType, UScriptStruct)
