[AutoRTFM] Silence MSVC unreachable code warnings when building with MSVC

AutoRTFM::Close() will call abort() when compiling without the verse-clang compiler or with AutoRTFM instrumentation disabled.
abort() is annotated with [[noreturn]] and so code that immediately follows the call to AutoRTFM::Close() will be treated as unreachable.

Disable compilation of a handful of AutoRTFM::Close() calls that were unconditionally compiled.

diff --git a/Engine/Source/Runtime/Core/Private/Internationalization/Text.cpp b/Engine/Source/Runtime/Core/Private/Internationalization/Text.cpp
--- a/Engine/Source/Runtime/Core/Private/Internationalization/Text.cpp
+++ b/Engine/Source/Runtime/Core/Private/Internationalization/Text.cpp
@@ -296,6 +296,7 @@ FText& FText::operator=(FText&& Other)
 	return *this;
 }
 
+#if UE_AUTORTFM
 void FText::AutoRTFMAssignFromOpenToClosed(FText& Closed, const FText& Open)
 {
 	// Copy (instead of move) in a closed transaction to ensure the ref-counts
@@ -306,6 +307,7 @@ void FText::AutoRTFMAssignFromOpenToClosed(FText& Closed, const FText& Open)
 	});
 	ensure(AutoRTFM::EContextStatus::OnTrack == Status);
 }
+#endif // UE_AUTORTFM
 
 const ITextData& FText::GetTextData() const
 {
diff --git a/Engine/Source/Runtime/Core/Public/Internationalization/Text.h b/Engine/Source/Runtime/Core/Public/Internationalization/Text.h
--- a/Engine/Source/Runtime/Core/Public/Internationalization/Text.h
+++ b/Engine/Source/Runtime/Core/Public/Internationalization/Text.h
@@ -397,9 +397,11 @@ public:
 	CORE_API FText(FText&& Other);
 	CORE_API FText& operator=(FText&& Other);
 
+#if UE_AUTORTFM
 	// Safely copies an FText from the open to a closed transaction.
 	// See AutoRTFM::Open().
 	static CORE_API void AutoRTFMAssignFromOpenToClosed(FText& Closed, const FText& Open);
+#endif // UE_AUTORTFM
 
 	/**
 	 * Generate an FText that represents the passed number in the current culture
diff --git a/Engine/Source/Runtime/CoreUObject/Private/Serialization/AsyncPackageLoader.cpp b/Engine/Source/Runtime/CoreUObject/Private/Serialization/AsyncPackageLoader.cpp
--- a/Engine/Source/Runtime/CoreUObject/Private/Serialization/AsyncPackageLoader.cpp
+++ b/Engine/Source/Runtime/CoreUObject/Private/Serialization/AsyncPackageLoader.cpp
@@ -269,7 +269,7 @@ void InitAsyncThread()
 		}
 	}
 
-#if UE_AUTORTFM
+#if UE_AUTORTFM_ENABLED
 	// If we are using AutoRTFM we need to layer a transactionally safe package loader over
 	// whatever loader was to be used.
 	GPackageLoader.Reset(MakeTransactionallySafeAsyncPackageLoader(GPackageLoader.Release()));
diff --git a/Engine/Source/Runtime/CoreUObject/Private/Serialization/TransactionallySafeAsyncLoading.cpp b/Engine/Source/Runtime/CoreUObject/Private/Serialization/TransactionallySafeAsyncLoading.cpp
--- a/Engine/Source/Runtime/CoreUObject/Private/Serialization/TransactionallySafeAsyncLoading.cpp
+++ b/Engine/Source/Runtime/CoreUObject/Private/Serialization/TransactionallySafeAsyncLoading.cpp
@@ -10,6 +10,8 @@
 #include "Misc/ScopeLock.h"
 #include "HAL/IConsoleManager.h"
 
+#if UE_AUTORTFM_ENABLED
+
 bool CVar_TransactionallySafeAsyncLoadingLegacyTrashedPackages = false;
 static FAutoConsoleVariableRef CVarTransactionallySafeAsyncLoadingLegacyTrashedPackages(TEXT("TransactionallySafeAsyncLoading.LegacyTrashedPackages"), CVar_TransactionallySafeAsyncLoadingLegacyTrashedPackages, TEXT("Use the legacy trashed packages handling (EG. ignore them!)."), ECVF_Default);
 
@@ -676,7 +678,14 @@ private:
 
 const char* const FTransactionallySafeAsyncPackageLoader::UnreachableMessage = "Cannot call function within a transaction!";
 
+#endif // UE_AUTORTFM_ENABLED
+
 IAsyncPackageLoader* MakeTransactionallySafeAsyncPackageLoader(IAsyncPackageLoader* const InWrappedPackageLoader)
 {
+#if UE_AUTORTFM_ENABLED
 	return new FTransactionallySafeAsyncPackageLoader(InWrappedPackageLoader);
+#else
+	checkf(UE_AUTORTFM_ENABLED, TEXT("MakeTransactionallySafeAsyncPackageLoader() called without AutoRTFM instrumentation enabled for CoreUObject"));
+	return InWrappedPackageLoader;
+#endif
 }
diff --git a/Engine/Source/Runtime/CoreUObject/Private/UObject/UObjectGlobals.cpp b/Engine/Source/Runtime/CoreUObject/Private/UObject/UObjectGlobals.cpp
--- a/Engine/Source/Runtime/CoreUObject/Private/UObject/UObjectGlobals.cpp
+++ b/Engine/Source/Runtime/CoreUObject/Private/UObject/UObjectGlobals.cpp
@@ -4326,8 +4326,10 @@ void FObjectInitializer::PostConstructInit()
 		{
 			if (AutoRTFM::IsTransactional())
 			{
+#if UE_AUTORTFM_ENABLED // #if purely to silence unreachable code warning.
 				AutoRTFM::EContextStatus Status = AutoRTFM::Close([&] { UE::FInitializedPropertyValueState(Obj).ActivateTracking(); });
 				check(Status == AutoRTFM::EContextStatus::OnTrack);
+#endif
 			}
 			else
 			{
