Replaced PREPROCESSOR_COMMA_SEPARATED.

diff --git a/Engine/Source/Runtime/CoreUObject/Public/UObject/ScriptMacros.h b/Engine/Source/Runtime/CoreUObject/Public/UObject/ScriptMacros.h
--- a/Engine/Source/Runtime/CoreUObject/Public/UObject/ScriptMacros.h
+++ b/Engine/Source/Runtime/CoreUObject/Public/UObject/ScriptMacros.h
@@ -28,32 +28,32 @@ enum {MAX_VARIABLE_SIZE = 0x0FFF };
 
 #define ZERO_INIT(Type,ParamName) FMemory::Memzero(&ParamName,sizeof(Type));
 
-#define PARAM_PASSED_BY_VAL(ParamName, PropertyType, ParamType)									\
-	ParamType ParamName;																		\
+#define PARAM_PASSED_BY_VAL(ParamName, PropertyType, ParamType)																							\
+	UE_REMOVE_OPTIONAL_PARENS(ParamType) ParamName;																										\
 	Stack.StepCompiledIn<PropertyType>(&ParamName);
 
-#define PARAM_PASSED_BY_VAL_ZEROED(ParamName, PropertyType, ParamType)							\
-	ParamType ParamName = (ParamType)0;															\
+#define PARAM_PASSED_BY_VAL_ZEROED(ParamName, PropertyType, ParamType)																					\
+	UE_REMOVE_OPTIONAL_PARENS(ParamType) ParamName = (UE_REMOVE_OPTIONAL_PARENS(ParamType))0;															\
 	Stack.StepCompiledIn<PropertyType>(&ParamName);
 
-#define PARAM_PASSED_BY_VAL_INITED(ParamName, PropertyType, ParamType, ...)                     \
-	ParamType ParamName{__VA_ARGS__};                                                           \
+#define PARAM_PASSED_BY_VAL_INITED(ParamName, PropertyType, ParamType, ...)																				\
+	UE_REMOVE_OPTIONAL_PARENS(ParamType) ParamName{__VA_ARGS__};																						\
 	Stack.StepCompiledIn<PropertyType>(&ParamName);
 
-#define PARAM_PASSED_BY_REF(ParamName, PropertyType, ParamType)									\
-	ParamType ParamName##Temp;																	\
-	ParamType& ParamName = Stack.StepCompiledInRef<PropertyType, ParamType>(&ParamName##Temp);
+#define PARAM_PASSED_BY_REF(ParamName, PropertyType, ParamType)																							\
+	UE_REMOVE_OPTIONAL_PARENS(ParamType) ParamName##Temp;																								\
+	UE_REMOVE_OPTIONAL_PARENS(ParamType)& ParamName = Stack.StepCompiledInRef<PropertyType, UE_REMOVE_OPTIONAL_PARENS(ParamType)>(&ParamName##Temp);
 
-#define PARAM_PASSED_BY_REF_ZEROED(ParamName, PropertyType, ParamType)							\
-	ParamType ParamName##Temp = (ParamType)0;													\
-	ParamType& ParamName = Stack.StepCompiledInRef<PropertyType, ParamType>(&ParamName##Temp);
+#define PARAM_PASSED_BY_REF_ZEROED(ParamName, PropertyType, ParamType)																					\
+	UE_REMOVE_OPTIONAL_PARENS(ParamType) ParamName##Temp = (UE_REMOVE_OPTIONAL_PARENS(ParamType))0;														\
+	UE_REMOVE_OPTIONAL_PARENS(ParamType)& ParamName = Stack.StepCompiledInRef<PropertyType, UE_REMOVE_OPTIONAL_PARENS(ParamType)>(&ParamName##Temp);
 
-#define P_GET_PROPERTY(PropertyType, ParamName)													\
-	PropertyType::TCppType ParamName = PropertyType::GetDefaultPropertyValue();					\
+#define P_GET_PROPERTY(PropertyType, ParamName)																											\
+	PropertyType::TCppType ParamName = PropertyType::GetDefaultPropertyValue();																			\
 	Stack.StepCompiledIn<PropertyType>(&ParamName);
 
-#define P_GET_PROPERTY_REF(PropertyType, ParamName)												\
-	PropertyType::TCppType ParamName##Temp = PropertyType::GetDefaultPropertyValue();			\
+#define P_GET_PROPERTY_REF(PropertyType, ParamName)																										\
+	PropertyType::TCppType ParamName##Temp = PropertyType::GetDefaultPropertyValue();																	\
 	PropertyType::TCppType& ParamName = Stack.StepCompiledInRef<PropertyType, PropertyType::TCppType>(&ParamName##Temp);
 
 
@@ -65,8 +65,8 @@ enum {MAX_VARIABLE_SIZE = 0x0FFF };
 #define P_GET_UBOOL64(ParamName)					uint64 ParamName=0;                             Stack.StepCompiledIn<FBoolProperty>(&ParamName); ParamName = ParamName ? 1 : 0; // translate the bitfield into a bool type for non-intel platforms
 #define P_GET_UBOOL_REF(ParamName)					PARAM_PASSED_BY_REF_ZEROED(ParamName, FBoolProperty, bool)
 
-#define P_GET_STRUCT(StructType,ParamName)			PARAM_PASSED_BY_VAL(ParamName, FStructProperty, PREPROCESSOR_COMMA_SEPARATED(StructType))
-#define P_GET_STRUCT_REF(StructType,ParamName)		PARAM_PASSED_BY_REF(ParamName, FStructProperty, PREPROCESSOR_COMMA_SEPARATED(StructType))
+#define P_GET_STRUCT(StructType,ParamName)			PARAM_PASSED_BY_VAL(ParamName, FStructProperty, StructType)
+#define P_GET_STRUCT_REF(StructType,ParamName)		PARAM_PASSED_BY_REF(ParamName, FStructProperty, StructType)
 
 #define P_GET_OBJECT(ObjectType,ParamName)			PARAM_PASSED_BY_VAL_ZEROED(ParamName, FObjectPropertyBase, ObjectType*)
 #define P_GET_OBJECT_REF(ObjectType,ParamName)		PARAM_PASSED_BY_REF_ZEROED(ParamName, FObjectPropertyBase, ObjectType*)
@@ -77,8 +77,8 @@ enum {MAX_VARIABLE_SIZE = 0x0FFF };
 #define P_GET_TARRAY(ElementType,ParamName)			PARAM_PASSED_BY_VAL(ParamName, FArrayProperty, TArray<ElementType>)
 #define P_GET_TARRAY_REF(ElementType,ParamName)		PARAM_PASSED_BY_REF(ParamName, FArrayProperty, TArray<ElementType>)
 
-#define P_GET_TMAP(KeyType,ValueType,ParamName)			PARAM_PASSED_BY_VAL(ParamName, FMapProperty, PREPROCESSOR_COMMA_SEPARATED(TMap<KeyType, ValueType>))
-#define P_GET_TMAP_REF(KeyType,ValueType,ParamName)		PARAM_PASSED_BY_REF(ParamName, FMapProperty, PREPROCESSOR_COMMA_SEPARATED(TMap<KeyType, ValueType>))
+#define P_GET_TMAP(KeyType,ValueType,ParamName)			PARAM_PASSED_BY_VAL(ParamName, FMapProperty, (TMap<KeyType, ValueType>))
+#define P_GET_TMAP_REF(KeyType,ValueType,ParamName)		PARAM_PASSED_BY_REF(ParamName, FMapProperty, (TMap<KeyType, ValueType>))
 
 #define P_GET_TSET(ElementType,ParamName)			PARAM_PASSED_BY_VAL(ParamName, FSetProperty, TSet<ElementType>)
 #define P_GET_TSET_REF(ElementType,ParamName)		PARAM_PASSED_BY_REF(ParamName, FSetProperty, TSet<ElementType>)
