Fixed TBaseStructure<const T> inheritance.

#rnx

diff --git a/Engine/Source/Runtime/CoreUObject/Public/UObject/Class.h b/Engine/Source/Runtime/CoreUObject/Public/UObject/Class.h
--- a/Engine/Source/Runtime/CoreUObject/Public/UObject/Class.h
+++ b/Engine/Source/Runtime/CoreUObject/Public/UObject/Class.h
@@ -5285,7 +5285,7 @@ struct TBaseStructure : TBaseStructureBase<T>
 };
 
 template <typename T>
-struct TBaseStructure<const T> : TBaseStructureBase<T>
+struct TBaseStructure<const T> : TBaseStructure<T>
 {
 };
 
