OFPE fixes to make sure to not dirty the map package when modifying external entities, including changing propertry override states.
The override manager shouldn't have to mark objects dirty, but the UI should.

diff --git a/Engine/Source/Runtime/CoreUObject/Private/UObject/OverriddenPropertySet.cpp b/Engine/Source/Runtime/CoreUObject/Private/UObject/OverriddenPropertySet.cpp
--- a/Engine/Source/Runtime/CoreUObject/Private/UObject/OverriddenPropertySet.cpp
+++ b/Engine/Source/Runtime/CoreUObject/Private/UObject/OverriddenPropertySet.cpp
@@ -482,31 +482,21 @@ bool FOverriddenPropertyNode::Serialize(FArchive& Ar)
 	return true;
 }
 
-void FOverriddenPropertyNode::SetOperation(EOverriddenPropertyOperation InOperation, bool* bWasModified /*= nullptr*/)
+void FOverriddenPropertyNode::SetOperation(EOverriddenPropertyOperation InOperation)
 {
 	if (Operation != InOperation)
 	{
 		Operation = InOperation;
-		if (bWasModified)
-		{
-			*bWasModified = true;
-		}
 	}
 }
 
-FOverriddenPropertyNode& FOverriddenPropertyNode::FindOrAddNode(const FOverriddenPropertyNodeID& InNodeID, bool* bWasModified /*= nullptr*/)
+FOverriddenPropertyNode& FOverriddenPropertyNode::FindOrAddNode(const FOverriddenPropertyNodeID& InNodeID)
 {
 	if (FOverriddenPropertyNode* SubNode = SubPropertyNodes.FindByKey(InNodeID))
 	{
 		return *SubNode;
 	}
 
-	if (bWasModified)
-	{
-		*bWasModified = true;
-	}
-
-
 	// We can safely assume that this node is at least modified from now on
 	if (Operation == EOverriddenPropertyOperation::None)
 	{
@@ -517,30 +507,22 @@ FOverriddenPropertyNode& FOverriddenPropertyNode::FindOrAddNode(const FOverridde
 	return SubPropertyNodes.Emplace_GetRef(InNodeID);
 }
 
-bool FOverriddenPropertyNode::RemoveSubPropertyNode(const FOverriddenPropertyNodeID& InNodeID, bool* bWasModified /*= nullptr*/)
+bool FOverriddenPropertyNode::RemoveSubPropertyNode(const FOverriddenPropertyNodeID& InNodeID)
 {
 	if (SubPropertyNodes.Remove(InNodeID))
 	{
-		if (bWasModified)
-		{
-			*bWasModified = true;
-		}
 		return true;
 	}
 	return false;
 }
 
-void FOverriddenPropertyNode::Reset(bool* bWasModified /*= nullptr*/)
+void FOverriddenPropertyNode::Reset()
 {
-	SetOperation(EOverriddenPropertyOperation::None, bWasModified);
+	SetOperation(EOverriddenPropertyOperation::None);
 
 	if (!SubPropertyNodes.IsEmpty())
 	{
 		SubPropertyNodes.Empty();
-		if (bWasModified)
-		{
-			*bWasModified = true;
-		}
 	}
 }
 
@@ -731,15 +713,6 @@ EOverriddenPropertyOperation FOverriddenPropertySet::GetOverriddenPropertyOperat
 
 bool FOverriddenPropertySet::ClearOverriddenProperty(FOverriddenPropertyNode& ParentPropertyNode, FPropertyVisitorPath::Iterator PropertyIterator, const void* Data)
 {
-	bool bWasModified = false;
-	ON_SCOPE_EXIT
-	{
-		if (bWasModified)
-		{
-			Owner->Modify();
-		}
-	};
-
 	FOverridableManager& OverridableManager = FOverridableManager::Get();
 	if (!PropertyIterator)
 	{
@@ -828,7 +801,7 @@ bool FOverriddenPropertySet::ClearOverriddenProperty(FOverriddenPropertyNode& Pa
 					if (CurrentOverriddenPropertyNode && SubObject)
 					{
 						const FOverriddenPropertyNodeID  SubObjectID(SubObject);
-						if (CurrentOverriddenPropertyNode->RemoveSubPropertyNode(SubObjectID, &bWasModified))
+						if (CurrentOverriddenPropertyNode->RemoveSubPropertyNode(SubObjectID))
 						{
 							OverridableManager.ClearInstancedSubObjectOverrides(Owner, SubObject);
 							return true;
@@ -899,7 +872,7 @@ bool FOverriddenPropertySet::ClearOverriddenProperty(FOverriddenPropertyNode& Pa
 				FOverriddenPropertyNodeID OverriddenKeyID = FOverriddenPropertyNodeID::FromMapKey(MapProperty->KeyProp, MapHelper.GetKeyPtr(InternalMapIndex));
 
 				FOverriddenPropertyNodeID CurrentPropKey;
-				if (CurrentOverriddenPropertyNode->RemoveSubPropertyNode(OverriddenKeyID, &bWasModified))
+				if (CurrentOverriddenPropertyNode->RemoveSubPropertyNode(OverriddenKeyID))
 				{
 					if (UObject* ValueSubObject = TryGetInstancedSubObjectValue(ValueObjectProperty, MapHelper.GetValuePtr(InternalMapIndex)))
 					{
@@ -916,7 +889,7 @@ bool FOverriddenPropertySet::ClearOverriddenProperty(FOverriddenPropertyNode& Pa
 		++PropertyIterator;
 	}
 
-	auto CleanupClearedNodes = [this, &TraversedNodes, &bWasModified]()
+	auto CleanupClearedNodes = [this, &TraversedNodes]()
 	{
 		// Go through each traversed property in reversed order to do cleanup
 		// We need to continue the cleanup until there is more overrides than just the one we are removing
@@ -928,13 +901,13 @@ bool FOverriddenPropertySet::ClearOverriddenProperty(FOverriddenPropertyNode& Pa
 				// In the case there are other overrides, just cleanup that node and stop.
 				if (CurrentNode->GetSubPropertyNodes().Num() > 1)
 				{
-					verifyf(CurrentNode->RemoveSubPropertyNode(LastCleanedNode->GetNodeID(), &bWasModified), TEXT("Expecting the node to always be removed"));
+					verifyf(CurrentNode->RemoveSubPropertyNode(LastCleanedNode->GetNodeID()), TEXT("Expecting the node to always be removed"));
 					LastCleanedNode = nullptr;
 					break;
 				}
 			}
 
-			CurrentNode->Reset(&bWasModified);
+			CurrentNode->Reset();
 			LastCleanedNode = CurrentNode;
 		}
 	};
@@ -960,15 +933,6 @@ bool FOverriddenPropertySet::ClearOverriddenProperty(FOverriddenPropertyNode& Pa
 
 void FOverriddenPropertySet::NotifyPropertyChange(FOverriddenPropertyNode* ParentPropertyNode, const EPropertyNotificationType Notification, FPropertyVisitorPath::Iterator PropertyIterator, const EPropertyChangeType::Type ChangeType, const void* Data, bool& bNeedsCleanup)
 {
-	bool bWasModified = false;
-	ON_SCOPE_EXIT
-	{
-		if (bWasModified)
-		{
-			Owner->Modify();
-		}
-	};
-
 	checkf(IsValid(Owner), TEXT("Expecting a valid overridable owner"));
 
 	if (ChangeType & EPropertyChangeType::ResetToDefault)
@@ -997,7 +961,7 @@ void FOverriddenPropertySet::NotifyPropertyChange(FOverriddenPropertyNode* Paren
 	FOverriddenPropertyNode* SubPropertyNode = nullptr;
 	if (ParentPropertyNode)
 	{
-		FOverriddenPropertyNode& SubPropertyNodeRef = ParentPropertyNode->FindOrAddNode(FOverriddenPropertyNodeID(Property), &bWasModified);
+		FOverriddenPropertyNode& SubPropertyNodeRef = ParentPropertyNode->FindOrAddNode(FOverriddenPropertyNodeID(Property));
 		SubPropertyNode = SubPropertyNodeRef.GetOperation() != EOverriddenPropertyOperation::Replace ? &SubPropertyNodeRef : nullptr;
 	}
 
@@ -1014,10 +978,10 @@ void FOverriddenPropertySet::NotifyPropertyChange(FOverriddenPropertyNode* Paren
 			 SubPropertyNode->GetOperation() == EOverriddenPropertyOperation::Modified))
 		{
 			FOverriddenPropertyNodeID RemoveNodeID(Property);
-			ParentPropertyNode->RemoveSubPropertyNode(RemoveNodeID, &bWasModified);
+			ParentPropertyNode->RemoveSubPropertyNode(RemoveNodeID);
 			if (ParentPropertyNode->GetOperation() == EOverriddenPropertyOperation::Modified && ParentPropertyNode->GetSubPropertyNodes().IsEmpty())
 			{
-				ParentPropertyNode->SetOperation(EOverriddenPropertyOperation::None, &bWasModified);
+				ParentPropertyNode->SetOperation(EOverriddenPropertyOperation::None);
 			}
 		}
 	};
@@ -1124,7 +1088,7 @@ void FOverriddenPropertySet::NotifyPropertyChange(FOverriddenPropertyNode* Paren
 					if (SubPropertyNode)
 					{
 						// Overriding all entry in the array
-						SubPropertyNode->SetOperation(EOverriddenPropertyOperation::Replace, &bWasModified);
+						SubPropertyNode->SetOperation(EOverriddenPropertyOperation::Replace);
 					}
 
 					// This is a case of the entire array is overridden
@@ -1146,8 +1110,8 @@ void FOverriddenPropertySet::NotifyPropertyChange(FOverriddenPropertyNode* Paren
 						if(SubPropertyNode)
 						{
 							const FOverriddenPropertyNodeID  AddedSubObjectID(AddedSubObject);
-							FOverriddenPropertyNode& AddedSubObjectNode = SubPropertyNode->FindOrAddNode(AddedSubObjectID, &bWasModified);
-							AddedSubObjectNode.SetOperation(EOverriddenPropertyOperation::Add, &bWasModified);
+							FOverriddenPropertyNode& AddedSubObjectNode = SubPropertyNode->FindOrAddNode(AddedSubObjectID);
+							AddedSubObjectNode.SetOperation(EOverriddenPropertyOperation::Add);
 
 							// Notify the subobject that it was added
 							if (FOverriddenPropertySet* AddedSubObjectOverriddenProperties = OverridableManager.GetOverriddenProperties(AddedSubObject))
@@ -1170,18 +1134,18 @@ void FOverriddenPropertySet::NotifyPropertyChange(FOverriddenPropertyNode* Paren
 							if (RemovedSubObjectArchetype && !RemovedSubObjectArchetype->HasAnyFlags(RF_ClassDefaultObject))
 							{
 								const FOverriddenPropertyNodeID RemovedSubObjectID (RemovedSubObjectArchetype);
-								FOverriddenPropertyNode& RemovedSubObjectNode = SubPropertyNode->FindOrAddNode(RemovedSubObjectID, &bWasModified);
+								FOverriddenPropertyNode& RemovedSubObjectNode = SubPropertyNode->FindOrAddNode(RemovedSubObjectID);
 								if (RemovedSubObjectNode.GetOperation() == EOverriddenPropertyOperation::Add)
 								{
 									// An add then a remove becomes no opt
-									if (SubPropertyNode->RemoveSubPropertyNode(RemovedSubObjectID, &bWasModified))
+									if (SubPropertyNode->RemoveSubPropertyNode(RemovedSubObjectID))
 									{
 										bNeedsCleanup = true;
 									}
 								}
 								else
 								{
-									RemovedSubObjectNode.SetOperation(EOverriddenPropertyOperation::Remove, &bWasModified);
+									RemovedSubObjectNode.SetOperation(EOverriddenPropertyOperation::Remove);
 								}
 							}
 							else
@@ -1196,7 +1160,7 @@ void FOverriddenPropertySet::NotifyPropertyChange(FOverriddenPropertyNode* Paren
 									}
 
 									// An add then a remove becomes no opt
-									if (SubPropertyNode->RemoveSubPropertyNode(RemovedSubObjectID, &bWasModified))
+									if (SubPropertyNode->RemoveSubPropertyNode(RemovedSubObjectID))
 									{
 										bNeedsCleanup = true;
 									}
@@ -1221,7 +1185,7 @@ void FOverriddenPropertySet::NotifyPropertyChange(FOverriddenPropertyNode* Paren
 					else if (SubPropertyNode)
 					{
 						// Overriding all entry in the array
-						SubPropertyNode->SetOperation(EOverriddenPropertyOperation::Replace, &bWasModified);
+						SubPropertyNode->SetOperation(EOverriddenPropertyOperation::Replace);
 					}
 					return;
 				}
@@ -1359,7 +1323,7 @@ void FOverriddenPropertySet::NotifyPropertyChange(FOverriddenPropertyNode* Paren
 				// Overriding a all entries in the map
 				if (SubPropertyNode)
 				{
-					SubPropertyNode->SetOperation(EOverriddenPropertyOperation::Replace, &bWasModified);
+					SubPropertyNode->SetOperation(EOverriddenPropertyOperation::Replace);
 				}
 
 				// This is a case of the entire array is overridden
@@ -1369,8 +1333,8 @@ void FOverriddenPropertySet::NotifyPropertyChange(FOverriddenPropertyNode* Paren
 					if(SubPropertyNode)
 					{
 						FOverriddenPropertyNodeID OverriddenKeyID = FOverriddenPropertyNodeID::FromMapKey(MapProperty->KeyProp, MapHelper.GetKeyPtr(It.GetInternalIndex()));
-						FOverriddenPropertyNode& OverriddenKeyNode = SubPropertyNode->FindOrAddNode(OverriddenKeyID, &bWasModified);
-						OverriddenKeyNode.SetOperation(EOverriddenPropertyOperation::Replace, &bWasModified);
+						FOverriddenPropertyNode& OverriddenKeyNode = SubPropertyNode->FindOrAddNode(OverriddenKeyID);
+						OverriddenKeyNode.SetOperation(EOverriddenPropertyOperation::Replace);
 					}
 
 					// @todo support instanced object as a key in maps
@@ -1393,8 +1357,8 @@ void FOverriddenPropertySet::NotifyPropertyChange(FOverriddenPropertyNode* Paren
 				if(SubPropertyNode)
 				{
 					FOverriddenPropertyNodeID AddedKeyID = FOverriddenPropertyNodeID::FromMapKey(MapProperty->KeyProp, MapHelper.GetKeyPtr(InternalMapIndex));
-					FOverriddenPropertyNode& AddedKeyNode = SubPropertyNode->FindOrAddNode(AddedKeyID, &bWasModified);
-					AddedKeyNode.SetOperation(EOverriddenPropertyOperation::Add, &bWasModified);
+					FOverriddenPropertyNode& AddedKeyNode = SubPropertyNode->FindOrAddNode(AddedKeyID);
+					AddedKeyNode.SetOperation(EOverriddenPropertyOperation::Add);
 				}
 			};
 
@@ -1405,18 +1369,18 @@ void FOverriddenPropertySet::NotifyPropertyChange(FOverriddenPropertyNode* Paren
 				if(SubPropertyNode)
 				{
 					FOverriddenPropertyNodeID RemovedKeyID = FOverriddenPropertyNodeID::FromMapKey(MapProperty->KeyProp, PreEditMapHelper.GetKeyPtr(InternalPreEditMapIndex));
-					FOverriddenPropertyNode& RemovedKeyNode = SubPropertyNode->FindOrAddNode(RemovedKeyID, &bWasModified);
+					FOverriddenPropertyNode& RemovedKeyNode = SubPropertyNode->FindOrAddNode(RemovedKeyID);
 					if (RemovedKeyNode.GetOperation() == EOverriddenPropertyOperation::Add)
 					{
 						// @Todo support remove/add/remove
-						if (SubPropertyNode->RemoveSubPropertyNode(RemovedKeyID, &bWasModified))
+						if (SubPropertyNode->RemoveSubPropertyNode(RemovedKeyID))
 						{
 							bNeedsCleanup = true;
 						}
 					}
 					else
 					{
-						RemovedKeyNode.SetOperation(EOverriddenPropertyOperation::Remove, &bWasModified);
+						RemovedKeyNode.SetOperation(EOverriddenPropertyOperation::Remove);
 					}
 				}
 			};
@@ -1432,7 +1396,7 @@ void FOverriddenPropertySet::NotifyPropertyChange(FOverriddenPropertyNode* Paren
 				else if(SubPropertyNode)
 				{
 					// Overriding all entry in the array
-					SubPropertyNode->SetOperation(EOverriddenPropertyOperation::Replace, &bWasModified);
+					SubPropertyNode->SetOperation(EOverriddenPropertyOperation::Replace);
 				}
 				return;
 			}
@@ -1527,7 +1491,7 @@ void FOverriddenPropertySet::NotifyPropertyChange(FOverriddenPropertyNode* Paren
 		{
 			if (Notification == EPropertyNotificationType::PostEdit && SubPropertyNode)
 			{
-				SubPropertyNode->SetOperation(EOverriddenPropertyOperation::Replace, &bWasModified);
+				SubPropertyNode->SetOperation(EOverriddenPropertyOperation::Replace);
 			}
 		}
 		else if (!NextPropertyIterator)
@@ -1563,7 +1527,7 @@ void FOverriddenPropertySet::NotifyPropertyChange(FOverriddenPropertyNode* Paren
 			{
 				if (SubPropertyNode)
 				{
-					SubPropertyNode->SetOperation(EOverriddenPropertyOperation::Replace, &bWasModified);
+					SubPropertyNode->SetOperation(EOverriddenPropertyOperation::Replace);
 				}
 
 				uint8* PreEditValuePtr = SavedPreEditContainers.Find(Owner, ObjectProperty);
@@ -1595,7 +1559,7 @@ void FOverriddenPropertySet::NotifyPropertyChange(FOverriddenPropertyNode* Paren
 		{
 			if (Notification == EPropertyNotificationType::PostEdit && SubPropertyNode)
 			{
-				SubPropertyNode->SetOperation(EOverriddenPropertyOperation::Replace, &bWasModified);
+				SubPropertyNode->SetOperation(EOverriddenPropertyOperation::Replace);
 			}
 		}
 		else if (OptionalProperty->IsSet(Data))
@@ -1611,7 +1575,7 @@ void FOverriddenPropertySet::NotifyPropertyChange(FOverriddenPropertyNode* Paren
 		if (SubPropertyNode)
 		{
 			// Replacing this entire property
-			SubPropertyNode->SetOperation(EOverriddenPropertyOperation::Replace, &bWasModified);
+			SubPropertyNode->SetOperation(EOverriddenPropertyOperation::Replace);
 		}
 	}
 }
@@ -1869,12 +1833,7 @@ bool FOverriddenPropertySet::IsCDOOwningProperty(const FProperty& Property) cons
 
 void FOverriddenPropertySet::Reset(bool bShouldDirtyObject)
 {
-	bool bWasModified = false;
-	RootNode.Reset(&bWasModified);
-	if (bWasModified && bShouldDirtyObject)
-	{
-		Owner->Modify();
-	}
+	RootNode.Reset();
 }
 
 void FOverriddenPropertySet::HandleObjectsReInstantiated(const TMap<UObject*, UObject*>& Map)
diff --git a/Engine/Source/Runtime/CoreUObject/Public/UObject/OverriddenPropertySet.h b/Engine/Source/Runtime/CoreUObject/Public/UObject/OverriddenPropertySet.h
--- a/Engine/Source/Runtime/CoreUObject/Public/UObject/OverriddenPropertySet.h
+++ b/Engine/Source/Runtime/CoreUObject/Public/UObject/OverriddenPropertySet.h
@@ -357,7 +357,7 @@ public:
 		return Operation;
 	}
 
-	void SetOperation(EOverriddenPropertyOperation InOperation, bool* bWasModified = nullptr);
+	void SetOperation(EOverriddenPropertyOperation InOperation);
 
 	TArrayView<FOverriddenPropertyNode> GetSubPropertyNodes()
 	{
@@ -369,16 +369,16 @@ public:
 		return SubPropertyNodes;
 	}
 
-	FOverriddenPropertyNode& FindOrAddNode(const FOverriddenPropertyNodeID& NodeID, bool* bWasModified = nullptr);
+	FOverriddenPropertyNode& FindOrAddNode(const FOverriddenPropertyNodeID& NodeID);
 
-	bool RemoveSubPropertyNode(const FOverriddenPropertyNodeID& NodeID, bool* bWasModified = nullptr);
+	bool RemoveSubPropertyNode(const FOverriddenPropertyNodeID& NodeID);
 
 	bool operator==(const FOverriddenPropertyNode& Other) const
 	{
 		return NodeID == Other.NodeID;
 	}
 
-	void Reset(bool* bWasModified = nullptr);
+	void Reset();
 
 	void HandleObjectsReInstantiated(const TMap<UObject*, UObject*>& Map);
 	void HandleDeadObjectReferences(const TSet<UObject*>& ActiveInstances, const TSet<UObject*>& TemplateInstances);
