PropertyBag: Expose CopyMatchingValuesByIDWithOverrides.

diff --git a/Engine/Source/Runtime/CoreUObject/Private/StructUtils/PropertyBag.cpp b/Engine/Source/Runtime/CoreUObject/Private/StructUtils/PropertyBag.cpp
--- a/Engine/Source/Runtime/CoreUObject/Private/StructUtils/PropertyBag.cpp
+++ b/Engine/Source/Runtime/CoreUObject/Private/StructUtils/PropertyBag.cpp
@@ -1927,6 +1927,11 @@ void FInstancedPropertyBag::CopyMatchingValuesByID(const FInstancedPropertyBag&
 	UE::StructUtils::Private::CopyMatchingValuesByID(Other.Value, Value);
 }
 
+void FInstancedPropertyBag::CopyMatchingValuesByIDWithOverrides(const FInstancedPropertyBag& NewDescs, TConstArrayView<FGuid> OverriddenPropertyIDs)
+{
+	UE::StructUtils::Private::CopyMatchingValuesByIDWithOverrides(NewDescs.Value, Value, OverriddenPropertyIDs);
+}
+
 void FInstancedPropertyBag::CopyMatchingValuesByName(const FInstancedPropertyBag& Other)
 {
 	UE::StructUtils::Private::CopyMatchingValues(
diff --git a/Engine/Source/Runtime/CoreUObject/Public/StructUtils/PropertyBag.h b/Engine/Source/Runtime/CoreUObject/Public/StructUtils/PropertyBag.h
--- a/Engine/Source/Runtime/CoreUObject/Public/StructUtils/PropertyBag.h
+++ b/Engine/Source/Runtime/CoreUObject/Public/StructUtils/PropertyBag.h
@@ -444,6 +444,14 @@ struct FInstancedPropertyBag
 	 */
 	UE_API void CopyMatchingValuesByID(const FInstancedPropertyBag& NewDescs);
 
+	/**
+	 * Copies matching property values from another bag of potentially mismatching layout.
+	 * The properties are matched between the bags based on the property ID.
+	 * @param NewDescs Reference to the bag to copy the values from
+	 * @param OverriddenPropertyIDs Only the matching properties, whose ID is found in this array are copied (if it is empty, nothing is copied)
+	 */
+	UE_API void CopyMatchingValuesByIDWithOverrides(const FInstancedPropertyBag& NewDescs, TConstArrayView<FGuid> OverriddenPropertyIDs);
+
 	/**
 	 * Copies matching property values from another bag of potentially mismatching layout.
 	 * The properties are matched between the bags based on the property name.
