[AsyncLoading2]

- Adds cvar "s.FailLoadOnMissingImport" to allow failing package loads when the package has Missing imports. (off by default). We don't want to allow packages to rely on missing imports but some content might be relying on that currently so we provide a toggle to allow for a smoother migration until all content is fixed.
- Restores cvar "s.FailLoadOnNotInstalledImport" to allow failing package loads when a package as NotInstalled imports (on by default)

#rnx

diff --git a/Engine/Source/Runtime/CoreUObject/Private/Serialization/AsyncLoading2.cpp b/Engine/Source/Runtime/CoreUObject/Private/Serialization/AsyncLoading2.cpp
--- a/Engine/Source/Runtime/CoreUObject/Private/Serialization/AsyncLoading2.cpp
+++ b/Engine/Source/Runtime/CoreUObject/Private/Serialization/AsyncLoading2.cpp
@@ -500,11 +500,19 @@ static FAutoConsoleVariableRef CVarGOnlyProcessRequiredPackagesWhenSyncLoading(
 	ECVF_Default
 );
 
-static bool GFailLoadOnFailedImport = false;
-static FAutoConsoleVariableRef CVarGFailLoadOnFailedImport(
-	TEXT("s.FailLoadOnFailedImport"),
-	GFailLoadOnFailedImport,
-	TEXT("Fail package load if an imported package cannot be loaded for any reason"),
+static bool GFailLoadOnNotInstalledImport = true;
+static FAutoConsoleVariableRef CVarGFailLoadOnNotInstalledImport(
+	TEXT("s.FailLoadOnNotInstalledImport"),
+	GFailLoadOnNotInstalledImport,
+	TEXT("Fail package load if an imported package is not installed"),
+	ECVF_Default
+);
+
+static bool GFailLoadOnMissingImport = false;
+static FAutoConsoleVariableRef CVarGFailLoadOnMissingImport(
+	TEXT("s.FailLoadOnMissingImport"),
+	GFailLoadOnMissingImport,
+	TEXT("Fail package load if an imported package is missing (cannot be found when it should be, i.e. this does not apply to packages that are NotInstalled)"),
 	ECVF_Default
 );
 
@@ -4713,12 +4721,14 @@ private:
 
 		if (PackageStatus != EPackageStoreEntryStatus::Missing)
 		{
-			if (PackageStatus == EPackageStoreEntryStatus::Ok && GFailLoadOnFailedImport)
+			if (PackageStatus == EPackageStoreEntryStatus::Ok && (GFailLoadOnMissingImport || GFailLoadOnNotInstalledImport))
 			{
 				// We need to confirm that this package can be fully imported so walk its transitive dependencies
 				// If we have a loaded package already that isn't failed we don't need to check it's dependencies
 				bool bFoundMissing = false;
 				bool bFoundNotInstalled = false;
+				const bool bLocalFailLoadOnMissingImport = GFailLoadOnMissingImport;
+				const bool bLocalFailLoadOnNotInstalledImport = GFailLoadOnNotInstalledImport;
 				FPackageId UnloadableImportPackageId;
 				TArray<FPackageId> ImportsToCheck(OutPackageStoreEntry.ImportedPackageIds);
 				TSet<FPackageId> ImportsSeen(OutPackageStoreEntry.ImportedPackageIds); // Out of paranoia, don't loop forever if there is somehow a circular import chain
@@ -4731,10 +4741,15 @@ private:
 						{
 							continue;
 						}
-						else if(PackageRef->IsNotInstalled())
+						else if (bLocalFailLoadOnMissingImport && PackageRef->IsMissing())
+						{
+							bFoundMissing |= PackageRef->IsMissing();
+							UnloadableImportPackageId = ImportId;
+							break;
+						}
+						else if(bLocalFailLoadOnNotInstalledImport && PackageRef->IsNotInstalled())
 						{
-							bFoundMissing = PackageRef->IsMissing();
-							bFoundNotInstalled = PackageRef->IsNotInstalled();
+							bFoundNotInstalled |= PackageRef->IsNotInstalled();
 							UnloadableImportPackageId = ImportId;
 							break;
 						}
@@ -4743,8 +4758,8 @@ private:
 					FPackageStoreEntry ImportEntry;
 					FName ImportName;
 					EPackageStoreEntryStatus ImportStatus = PackageStore.GetPackageStoreEntry(ImportId, ImportName, ImportEntry);
-					bFoundMissing |= ImportStatus == EPackageStoreEntryStatus::Missing;
-					bFoundNotInstalled |= ImportStatus == EPackageStoreEntryStatus::NotInstalled;
+					bFoundMissing |= bLocalFailLoadOnMissingImport && (ImportStatus == EPackageStoreEntryStatus::Missing);
+					bFoundNotInstalled |= bLocalFailLoadOnNotInstalledImport && (ImportStatus == EPackageStoreEntryStatus::NotInstalled);
 					if (bFoundMissing || bFoundNotInstalled)
 					{
 						UnloadableImportPackageId = ImportId;
