[InstancedStruct] Work with BP structs by checking authored name of the property.

diff --git a/Engine/Source/Runtime/CoreUObject/Private/StructUtils/InstancedStruct.cpp b/Engine/Source/Runtime/CoreUObject/Private/StructUtils/InstancedStruct.cpp
--- a/Engine/Source/Runtime/CoreUObject/Private/StructUtils/InstancedStruct.cpp
+++ b/Engine/Source/Runtime/CoreUObject/Private/StructUtils/InstancedStruct.cpp
@@ -571,7 +571,7 @@ bool FInstancedStruct::FindInnerPropertyInstance(FName PropertyName, const FProp
 	
 	for (const FProperty* Prop : TFieldRange<FProperty>(ScriptStruct))
 	{
-		if( Prop->GetFName() == PropertyName )
+		if( Prop->GetAuthoredName() == PropertyName )
 		{
 			OutProp = Prop;
 			OutData = StructMemory;
