ConstStructView: Allow non-const types to appear as template arguments in the getters.

diff --git a/Engine/Source/Runtime/CoreUObject/Public/StructUtils/StructView.h b/Engine/Source/Runtime/CoreUObject/Public/StructUtils/StructView.h
--- a/Engine/Source/Runtime/CoreUObject/Public/StructUtils/StructView.h
+++ b/Engine/Source/Runtime/CoreUObject/Public/StructUtils/StructView.h
@@ -251,16 +251,14 @@ public:
 
 	/** Returns const reference to the struct, this getter assumes that all data is valid. */
 	template<typename T>
-		requires (std::is_const_v<T>)
-	constexpr T& Get() const
+	constexpr const T& Get() const
 	{
 		return UE::StructUtils::GetStructRef<T>(ScriptStruct, StructMemory);
 	}
 
 	/** Returns const pointer to the struct, or nullptr if cast is not valid. */
 	template<typename T>
-		requires (std::is_const_v<T>)
-	constexpr T* GetPtr() const
+	constexpr const T* GetPtr() const
 	{
 		return UE::StructUtils::GetStructPtr<T>(ScriptStruct, StructMemory);
 	}
