Add option to be able to show serialized tagged members in CPU traces.
This can help track down the properties that are expensive to serialize.

#rnx

diff --git a/Engine/Source/Runtime/CoreUObject/Private/UObject/PropertyTag.cpp b/Engine/Source/Runtime/CoreUObject/Private/UObject/PropertyTag.cpp
--- a/Engine/Source/Runtime/CoreUObject/Private/UObject/PropertyTag.cpp
+++ b/Engine/Source/Runtime/CoreUObject/Private/UObject/PropertyTag.cpp
@@ -13,6 +13,14 @@
 #include "UObject/Package.h"
 #include "UObject/UnrealType.h"
 
+#ifndef UE_ENABLE_SERIALIZE_TAGGED_PROPERTY_TRACING
+#define UE_ENABLE_SERIALIZE_TAGGED_PROPERTY_TRACING 0
+#endif
+
+#if UE_ENABLE_SERIALIZE_TAGGED_PROPERTY_TRACING
+#include "Trace/Trace.inl"
+#endif
+
 /** Flags used for serialization of FPropertyTag. DO NOT EDIT THESE VALUES! */
 enum class EPropertyTagFlags : uint8
 {
@@ -550,8 +558,19 @@ void FPropertyTag::SerializeTaggedProperty(FArchive& Ar, FProperty* Property, ui
 	SerializeTaggedProperty(FStructuredArchiveFromArchive(Ar).GetSlot(), Property, Value, Defaults);
 }
 
+#if UE_ENABLE_SERIALIZE_TAGGED_PROPERTY_TRACING
+UE_TRACE_EVENT_BEGIN(Cpu, SerializeTaggedProperty, NoSync)
+	UE_TRACE_EVENT_FIELD(UE::Trace::WideString, PropertyName)
+UE_TRACE_EVENT_END()
+#endif
+
 void FPropertyTag::SerializeTaggedProperty(FStructuredArchive::FSlot Slot, FProperty* Property, uint8* Value, const uint8* Defaults) const
 {
+#if UE_ENABLE_SERIALIZE_TAGGED_PROPERTY_TRACING
+	UE_TRACE_LOG_SCOPED_T(Cpu, SerializeTaggedProperty, CpuChannel)
+		<< SerializeTaggedProperty.PropertyName(*WriteToString<256>(Property->GetFName()));
+#endif
+
 	FArchive& UnderlyingArchive = Slot.GetUnderlyingArchive();
 	const int64 StartOfProperty = UnderlyingArchive.Tell();
 
