Fixing Name initialization checks when Enums have no names.

diff --git a/Engine/Source/Runtime/CoreUObject/Public/UObject/Class.h b/Engine/Source/Runtime/CoreUObject/Public/UObject/Class.h
--- a/Engine/Source/Runtime/CoreUObject/Public/UObject/Class.h
+++ b/Engine/Source/Runtime/CoreUObject/Public/UObject/Class.h
@@ -3468,14 +3468,14 @@ protected:
 		// Get a mutable view over the names
 		[[nodiscard]] inline TArrayView<FName> GetMutableNames()
 		{
-			checkfSlow(!!(TaggedNames & TagDynamic), TEXT("This function should only be used after compiled-in enums have copied their static string names into FNames"));
+			checkfSlow(!TaggedNames || !!(TaggedNames & TagDynamic), TEXT("This function should only be used after compiled-in enums have copied their static string names into FNames"));
 			return MakeArrayView(reinterpret_cast<FName*>(TaggedNames & MaskPointer), NumValues);
 		}
 
 		// Get an immutable view over the names 
 		[[nodiscard]] inline TConstArrayView<FName> GetNames() const
 		{
-			checkfSlow(!!(TaggedNames & TagDynamic), TEXT("This function should only be used after compiled-in enums have copied their static string names into FNames"));
+			checkfSlow(!TaggedNames || !!(TaggedNames & TagDynamic), TEXT("This function should only be used after compiled-in enums have copied their static string names into FNames"));
 			return MakeConstArrayView(reinterpret_cast<const FName*>(TaggedNames & MaskPointer), NumValues);
 		}
 
