[Config] added ensure and recovery code when a config is created and 'config' specifier is missing on the UProperty

#FailedValidations Preflight

diff --git a/Engine/Source/Runtime/CoreUObject/Private/UObject/Obj.cpp b/Engine/Source/Runtime/CoreUObject/Private/UObject/Obj.cpp
--- a/Engine/Source/Runtime/CoreUObject/Private/UObject/Obj.cpp
+++ b/Engine/Source/Runtime/CoreUObject/Private/UObject/Obj.cpp
@@ -3930,8 +3930,15 @@ void UObject::UpdateSinglePropertyInConfigFile(const FProperty* InProperty, cons
 		TArray<FString> Keys;
 		NewFile.GetKeys(Keys);
 
-		const FString SectionName = Keys[0];
 		FString PropertyKey = InProperty->GetFName().ToString();
+		if (!ensureMsgf(!Keys.IsEmpty(), TEXT("Unable to find section for property '%s' in '%s'. Make sure the property has the 'config' specifier.")
+			, *PropertyKey
+			, *GetNameSafe(GetClass())))
+		{
+			return;
+		}
+
+		const FString SectionName = Keys[0];
 
 #if WITH_EDITOR
 		const FString& CVarName = InProperty->GetMetaData(GConsoleVariableFName);
