[AsyncLoading2]

Pass in the correct Header to use when fetching names of exports/imports in an AsyncPackage rather than assuming the cooked header is the one being referred to. When fetching names for optional exports, the wrong header was being used which leads to incorrect names and possible crashes due to out of bounds indexing.

#rnx

diff --git a/Engine/Source/Runtime/CoreUObject/Private/Serialization/AsyncLoading2.cpp b/Engine/Source/Runtime/CoreUObject/Private/Serialization/AsyncLoading2.cpp
--- a/Engine/Source/Runtime/CoreUObject/Private/Serialization/AsyncLoading2.cpp
+++ b/Engine/Source/Runtime/CoreUObject/Private/Serialization/AsyncLoading2.cpp
@@ -4025,7 +4025,7 @@ private:
 
 	void ProcessExportDependencies(const FAsyncPackageHeaderData& Header, int32 LocalExportIndex, FExportBundleEntry::EExportCommandType CommandType);
 	int32 GetPublicExportIndex(uint64 ExportHash, FAsyncPackageHeaderData*& OutHeader);
-	FString GetNameFromPackageObjectIndex(FPackageObjectIndex Index);
+	FString GetNameFromPackageObjectIndex(const FAsyncPackageHeaderData& Header, FPackageObjectIndex Index);
 	UObject* ConditionalCreateExport(const FAsyncPackageHeaderData& Header, int32 LocalExportIndex);
 	UObject* ConditionalSerializeExport(const FAsyncPackageHeaderData& Header, int32 LocalExportIndex);
 	UObject* ConditionalCreateImport(const FAsyncPackageHeaderData& Header, int32 LocalImportIndex);
@@ -8406,7 +8406,7 @@ Verse::VCell* FAsyncPackage2::ConditionalSerializeCellImport(const FAsyncPackage
 	return SourcePackage->ConditionalSerializeCellExport(*SourcePackageHeader, CellExportIndex);
 }
 
-FString FAsyncPackage2::GetNameFromPackageObjectIndex(FPackageObjectIndex Index)
+FString FAsyncPackage2::GetNameFromPackageObjectIndex(const FAsyncPackageHeaderData& Header, FPackageObjectIndex Index)
 {
 	FString Out;
 	if (Index.IsExport())
@@ -8418,13 +8418,13 @@ FString FAsyncPackage2::GetNameFromPackageObjectIndex(FPackageObjectIndex Index)
 		else
 		{
 			FName Name(NAME_None);
-			HeaderData.NameMap.TryGetName(HeaderData.ExportMap[Index.ToExport()].ObjectName, Name);
+			Header.NameMap.TryGetName(Header.ExportMap[Index.ToExport()].ObjectName, Name);
 			Out = Name.ToString();
 		}
 	}
 	else if (Index.IsImport())
 	{
-		Out = ImportStore.GetImportName(HeaderData, Index);
+		Out = ImportStore.GetImportName(Header, Index);
 	}
 
 	return Out;
@@ -8615,13 +8615,13 @@ void FAsyncPackage2::EventDrivenCreateExport(const FAsyncPackageHeaderData& Head
 
 	if (!LoadClass)
 	{
-		UE_ASYNC_PACKAGE_LOG(Error, Desc, TEXT("CreateExport"), TEXT("Could not find class object (0x%llX - '%s') for %s"), Export.ClassIndex.Value(), *GetNameFromPackageObjectIndex(Export.ClassIndex), *ObjectName.ToString());
+		UE_ASYNC_PACKAGE_LOG(Error, Desc, TEXT("CreateExport"), TEXT("Could not find class object (0x%llX - '%s') for %s"), Export.ClassIndex.Value(), *GetNameFromPackageObjectIndex(Header, Export.ClassIndex), *ObjectName.ToString());
 		ExportObject.bExportLoadFailed = true;
 		return;
 	}
 	if (!ThisParent)
 	{
-		UE_ASYNC_PACKAGE_LOG(Error, Desc, TEXT("CreateExport"), TEXT("Could not find outer (0x%llX - '%s') object for %s"), Export.OuterIndex.Value(), *GetNameFromPackageObjectIndex(Export.OuterIndex), *ObjectName.ToString());
+		UE_ASYNC_PACKAGE_LOG(Error, Desc, TEXT("CreateExport"), TEXT("Could not find outer (0x%llX - '%s') object for %s"), Export.OuterIndex.Value(), *GetNameFromPackageObjectIndex(Header, Export.OuterIndex), *ObjectName.ToString());
 		ExportObject.bExportLoadFailed = true;
 		return;
 	}
@@ -8631,7 +8631,7 @@ void FAsyncPackage2::EventDrivenCreateExport(const FAsyncPackageHeaderData& Head
 		ExportObject.SuperObject = EventDrivenIndexToObject(Header, Export.SuperIndex, false);
 		if (!ExportObject.SuperObject)
 		{
-			UE_ASYNC_PACKAGE_LOG(Error, Desc, TEXT("CreateExport"), TEXT("Could not find SuperStruct object (0x%llX - '%s') for %s"), Export.SuperIndex.Value(), *GetNameFromPackageObjectIndex(Export.SuperIndex), *ObjectName.ToString());
+			UE_ASYNC_PACKAGE_LOG(Error, Desc, TEXT("CreateExport"), TEXT("Could not find SuperStruct object (0x%llX - '%s') for %s"), Export.SuperIndex.Value(), *GetNameFromPackageObjectIndex(Header, Export.SuperIndex), *ObjectName.ToString());
 			ExportObject.bExportLoadFailed = true;
 			return;
 		}
@@ -8641,7 +8641,7 @@ void FAsyncPackage2::EventDrivenCreateExport(const FAsyncPackageHeaderData& Head
 	ExportObject.TemplateObject = EventDrivenIndexToObject(Header, Export.TemplateIndex, true);
 	if (!ExportObject.TemplateObject)
 	{
-		UE_ASYNC_PACKAGE_LOG(Error, Desc, TEXT("CreateExport"), TEXT("Could not find template object (0x%llX - '%s') for %s"), Export.TemplateIndex.Value(), *GetNameFromPackageObjectIndex(Export.TemplateIndex), *ObjectName.ToString());
+		UE_ASYNC_PACKAGE_LOG(Error, Desc, TEXT("CreateExport"), TEXT("Could not find template object (0x%llX - '%s') for %s"), Export.TemplateIndex.Value(), *GetNameFromPackageObjectIndex(Header, Export.TemplateIndex), *ObjectName.ToString());
 		ExportObject.bExportLoadFailed = true;
 		return;
 	}
@@ -11194,8 +11194,8 @@ void FAsyncPackage2::ClearConstructedObjects()
 		FExportObject& Export = Data.Exports[ExportIndex];
 		if (Export.bWasFoundInMemory)
 		{
-			checkf(Export.Object, TEXT("Export '%s' in package '%s' found in memory is missing:\n\t\tSuper: '%s', Template: '%s', LoadFailed: %s, Filtered: %s (%s), FoundInMemory: %s"),
-				*GetNameFromPackageObjectIndex(FPackageObjectIndex::FromExportIndex(ExportIndex)),
+			checkf(Export.Object, TEXT("Export %d in package '%s' found in memory is missing:\n\t\tSuper: '%s', Template: '%s', LoadFailed: %s, Filtered: %s (%s), FoundInMemory: %s"),
+				ExportIndex, // The header is cleared by this point so we only have the export index we can provide for backtracking purposes
 				*Desc.PackagePathToLoad.GetPackageFName().ToString(),
 				(Export.SuperObject ? *Export.SuperObject->GetPathName() : TEXT("null")),
 				(Export.TemplateObject ? *Export.TemplateObject->GetPathName() : TEXT("null")),
