PropertyBag: Serialize PropertyFlags.
It is used e.g. by EPropertyBagChildRowFeatures::AccessSpecifierButton.

#rnx

diff --git a/Engine/Source/Runtime/CoreUObject/Private/StructUtils/PropertyBag.cpp b/Engine/Source/Runtime/CoreUObject/Private/StructUtils/PropertyBag.cpp
--- a/Engine/Source/Runtime/CoreUObject/Private/StructUtils/PropertyBag.cpp
+++ b/Engine/Source/Runtime/CoreUObject/Private/StructUtils/PropertyBag.cpp
@@ -43,6 +43,7 @@ namespace FPropertyBagCustomVersion
 		ContainerTypes = 1,
 		NestedContainerTypes = 2,
 		MetaClass = 3,
+		PropertyFlags = 4,
 
 		// -----<new versions can be added above this line>-------------------------------------------------
 		VersionPlusOne,
@@ -1758,6 +1759,11 @@ FArchive& operator<<(FArchive& Ar, FPropertyBagPropertyDesc& Bag)
 		}
 #endif
 	}
+
+	if (Ar.CustomVer(FPropertyBagCustomVersion::GUID) >= FPropertyBagCustomVersion::PropertyFlags)
+	{
+		Ar << Bag.PropertyFlags;
+	}
 	
 	return Ar;
 }
