[Feature] Support for IDTs that operate relative to structs instead of classes

diff --git a/Engine/Source/Runtime/CoreUObject/Private/UObject/InstanceDataTransforms.cpp b/Engine/Source/Runtime/CoreUObject/Private/UObject/InstanceDataTransforms.cpp
--- a/Engine/Source/Runtime/CoreUObject/Private/UObject/InstanceDataTransforms.cpp
+++ b/Engine/Source/Runtime/CoreUObject/Private/UObject/InstanceDataTransforms.cpp
@@ -14,6 +14,7 @@
 #include "UObject/StrProperty.h"
 #include "UObject/TextProperty.h"
 #include "UObject/PropertyBagRepository.h"
+#include "UObject/PropertyHelper.h"
 #include "UObject/Package.h"
 #include "Dom/JsonValue.h"
 #include "Dom/JsonObject.h"
@@ -43,7 +44,7 @@ namespace UE::Private
 	};
 
 	typedef TSharedPtr<FJsonObject> FInstanceDataTransformParams;
-	typedef TFunction<EInstanceDataTransformResult(TObjectPtr<UObject>& IDO, TObjectPtr<UObject>& Owner, const FInstanceDataTransformParams&)> FInstanceDataTransformFunction;
+	typedef TFunction<EInstanceDataTransformResult(TObjectPtr<UObject>& IDO, TObjectPtr<UObject>& Owner, const FPropertyPathName& BasePath, const FInstanceDataTransformParams&)> FInstanceDataTransformFunction;
 
 	struct FInstanceDataTransform
 	{
@@ -85,7 +86,7 @@ namespace UE::Private
 						// check if the config file has been updated since the last time we read it
 						if (*LastReadTime < FileModificationTime)
 						{
-							ClassPathToIDTs.Empty(); // kill all cache
+							StructPathToIDTs.Empty(); // kill all cache
 							ReadConfigFile(ConfigFilepath);
 						}
 					}
@@ -95,12 +96,13 @@ namespace UE::Private
 					}
 				}
 			}
-			return ClassPathToIDTs.Find(ClassPath);
+			
+			return StructPathToIDTs.Find(ClassPath);
 		}
 
-		const TArray<TSharedPtr<FInstanceDataTransform>>* FindTransforms(const TNotNull<UClass*> Class)
+		const TArray<TSharedPtr<FInstanceDataTransform>>* FindTransforms(const TNotNull<UStruct*> Struct)
 		{
-			return FindTransforms(Class->GetClassPathName());
+			return FindTransforms(Struct->GetStructPathName());
 		}
 	private:
 
@@ -114,9 +116,9 @@ namespace UE::Private
 				if (FJsonSerializer::Deserialize(JsonReader, ConfigJson))
 				{
 					FString PackagePathString;
-					if (!ConfigJson->TryGetStringField(TEXT("ClassPackage"), PackagePathString))
+					if (!ConfigJson->TryGetStringField(TEXT("Package"), PackagePathString))
 					{
-						UE_LOG(LogInstanceDataTransforms, Warning, TEXT("'ClassPackage' feild was missing from '%s'"), *ConfigFilepath);
+						UE_LOG(LogInstanceDataTransforms, Warning, TEXT("'Package' feild was missing from '%s'"), *ConfigFilepath);
 						return;
 					}
 
@@ -143,14 +145,22 @@ namespace UE::Private
 							continue;
 						}
 
-						FString ClassNameString;
-						if (!(*TransformInfo)->TryGetStringField(TEXT("Class"), ClassNameString))
+						FTopLevelAssetPath StructPath;
+						if (FString ClassNameString; (*TransformInfo)->TryGetStringField(TEXT("Class"), ClassNameString))
 						{
-							UE_LOG(LogInstanceDataTransforms, Warning, TEXT("an element of the 'ClassTransforms' array in '%s' was missing a 'Class' field"), *ConfigFilepath);
+							FName ClassName = FName(ClassNameString);
+							StructPath = FTopLevelAssetPath(PackagePath.GetPackageFName(), ClassName);
+						}
+						else if (FString StructNameString; (*TransformInfo)->TryGetStringField(TEXT("Struct"), StructNameString))
+						{
+							FName StructName = FName(StructNameString);
+							StructPath = FTopLevelAssetPath(PackagePath.GetPackageFName(), StructName);
+						}
+						else
+						{
+							UE_LOG(LogInstanceDataTransforms, Warning, TEXT("an element of the 'ClassTransforms' array in '%s' was missing a 'Class' or 'Struct' field"), *ConfigFilepath);
 							continue;
 						}
-						FName ClassName = FName(ClassNameString);
-						FTopLevelAssetPath ClassPath(PackagePath.GetPackageFName(), ClassName);
 
 						FString OpCode;
 						if (!(*TransformInfo)->TryGetStringField(TEXT("OpCode"), OpCode))
@@ -169,7 +179,7 @@ namespace UE::Private
 						{
 							Params = *ParamsPtr;
 						}
-						TArray<TSharedPtr<FInstanceDataTransform>>& ClassTransforms = ClassPathToIDTs.FindOrAdd(ClassPath);
+						TArray<TSharedPtr<FInstanceDataTransform>>& ClassTransforms = StructPathToIDTs.FindOrAdd(StructPath);
 						ClassTransforms.Add(MakeShared<FInstanceDataTransform>(Params, *Func));
 					}
 				}
@@ -180,11 +190,11 @@ namespace UE::Private
 		friend struct FRegesterInstanceDataTransformHelper;
 		void RegisterOperation(FStringView Name, FInstanceDataTransformFunction Function)
 		{
-			RegisteredOperations.Add(Name.GetData(), Function);
+			RegisteredOperations.Add(FString(Name), Function);
 		}
 
 		TMap<FString, FDateTime> ConfigFileLastReadTime;
-		TMap<FTopLevelAssetPath, TArray<TSharedPtr<FInstanceDataTransform>>> ClassPathToIDTs;
+		TMap<FTopLevelAssetPath, TArray<TSharedPtr<FInstanceDataTransform>>> StructPathToIDTs;
 		TMap<FString, FInstanceDataTransformFunction> RegisteredOperations;
 	};
 
@@ -344,6 +354,83 @@ namespace UE::Private
 			}
 		}
 	}
+
+	void FlattenPropertyTree(const FPropertyPathNameTree& UnknownPropertyTree, TArray<FPropertyPathName>& OutPaths, FPropertyPathName& Path)
+	{
+		for (FPropertyPathNameTree::FConstIterator Itr = UnknownPropertyTree.CreateConstIterator(); Itr; ++Itr)
+		{
+			Path.Push(FPropertyPathNameSegment{ Itr.GetName(), Itr.GetType() });
+			OutPaths.Push(Path);
+			if (const FPropertyPathNameTree* SubTree = Itr.GetNode().GetSubTree())
+			{
+				FlattenPropertyTree(*SubTree, OutPaths, Path);
+			}
+			Path.Pop();
+		}
+	}
+
+	void PatchStructs(TObjectPtr<UObject>& IDO, TObjectPtr<UObject>& Owner)
+	{
+		// search the unknown property tree for structs. If there are any, look up IDTs associated with that struct type and run them
+		TSharedPtr<FPropertyPathNameTree> UnknownPropertyTree = FUnknownPropertyTree(Owner).FindOrCreate();
+
+		// IDTs can modify the UnknownPropertyTree and iterating while modifying a tree is a no-no so pre-iterate and cache the paths that we want to visit ahead of time
+		TArray<FPropertyPathName> Paths;
+		FPropertyPathName BasePath;
+		FlattenPropertyTree(*UnknownPropertyTree, Paths, BasePath);
+
+		// search the unknown property tree for structs. If there are any, look up IDTs associated with that struct type and run them
+		for (int32 i = 0; i < Paths.Num(); ++i)
+		{
+			FPropertyPathName& Path = Paths[i];
+			FPropertyPathNameSegment LeafSegment = Path.GetSegment(Path.GetSegmentCount() - 1);
+
+			auto PatchStruct = [&Path, &IDO, &Owner](const FPropertyTypeName& StructPath)
+				{
+					if (UStruct* Struct = UE::FindObjectByTypePath<UScriptStruct>(StructPath))
+					{
+
+						if (const TArray<TSharedPtr<Private::FInstanceDataTransform>>* IDTs = Private::FInstanceDataTransformRepository::Get().FindTransforms(Struct))
+						{
+							for (const TSharedPtr<Private::FInstanceDataTransform>& IDT : *IDTs)
+							{
+								Private::EInstanceDataTransformResult Result = IDT->Func(IDO, Owner, Path, IDT->Args);
+								if (Result == Private::EInstanceDataTransformResult::Break)
+								{
+									break;
+								}
+							}
+						}
+					}
+				};
+
+			if (LeafSegment.Type.GetName() == NAME_StructProperty)
+			{
+				PatchStruct(LeafSegment.Type.GetParameter(0));
+			}
+			else if (LeafSegment.Type.GetName() == NAME_ArrayProperty || LeafSegment.Type.GetName() == NAME_SetProperty)
+			{
+				FPropertyTypeName InnerType = LeafSegment.Type.GetParameter(0);
+				if (InnerType.GetName() == NAME_StructProperty)
+				{
+					PatchStruct(InnerType.GetParameter(0));
+				}
+			}
+			else if (LeafSegment.Type.GetName() == NAME_MapProperty)
+			{
+				FPropertyTypeName KeyType = LeafSegment.Type.GetParameter(0);
+				if (KeyType.GetName() == NAME_StructProperty)
+				{
+					PatchStruct(KeyType.GetParameter(0));
+				}
+				FPropertyTypeName ValType = LeafSegment.Type.GetParameter(1);
+				if (ValType.GetName() == NAME_StructProperty)
+				{
+					PatchStruct(ValType.GetParameter(0));
+				}
+			}
+		}
+	}
 }
 
 #if WITH_EDITOR
@@ -356,7 +443,7 @@ namespace UE::Private::IDT
 	 * @param Params - contains a map with the following parameters
 	 *		- ["src"] the source FPropertyPathName as a json object. If this doesn't point to an unknown property, the operation will be skipped
 	 */
-	EInstanceDataTransformResult RemoveProperty(TObjectPtr<UObject>& IDO, TObjectPtr<UObject>& Owner, const FInstanceDataTransformParams& Params)
+	EInstanceDataTransformResult RemoveProperty(TObjectPtr<UObject>& IDO, TObjectPtr<UObject>& Owner, const FPropertyPathName& BasePath, const FInstanceDataTransformParams& Params)
 	{
 		FUObjectSerializeContext* SerializeContext = FUObjectThreadContext::Get().GetSerializeContext();
 		TGuardValue<bool> ScopedImpersonateProperties(SerializeContext->bImpersonateProperties, true);
@@ -382,7 +469,7 @@ namespace UE::Private::IDT
 	};
 	REGISTER_IDT_OPERATION(RemoveProperty);
 
-	/**  
+	/**
 	 * Map a source FPropertyPathName into a destination FPropertyPathName relative to an object.
 	 * @param IDO - The instance data object that contains the source unknown property
 	 * @param Owner - the instance associated with the IDO
@@ -390,7 +477,7 @@ namespace UE::Private::IDT
 	 *		- ["src"] the source FPropertyPathName as a json object. If this doesn't point to an unknown property, the operation will be skipped
 	 *		- ["dst"] the destination FPropertyPathName as a json object. If this points to a non-existent property, the path will be added to the new IDO instance as an unknown property
 	 */
-	EInstanceDataTransformResult RedirectProperty(TObjectPtr<UObject>& IDO, TObjectPtr<UObject>& Owner, const FInstanceDataTransformParams& Params)
+	EInstanceDataTransformResult RedirectProperty(TObjectPtr<UObject>& IDO, TObjectPtr<UObject>& Owner, const FPropertyPathName& BasePath, const FInstanceDataTransformParams& Params)
 	{
 		FUObjectSerializeContext* SerializeContext = FUObjectThreadContext::Get().GetSerializeContext();
 		TGuardValue<bool> ScopedImpersonateProperties(SerializeContext->bImpersonateProperties, true);
@@ -398,7 +485,7 @@ namespace UE::Private::IDT
 		TSharedPtr<FPropertyPathNameTree> UnknownPropertyTree = FUnknownPropertyTree(Owner).FindOrCreate();
 		FUnknownEnumNames UnknownEnumNames = FUnknownEnumNames(Owner);
 
-		FPropertyPathName SrcPath;
+		FPropertyPathName SrcPath = BasePath;
 		if (!ensureAlwaysMsgf(Params && FPropertyPathName::FromJsonValue(Params->TryGetField(TEXT("src")), SrcPath),
 			TEXT("Failed to parse src parameter for a RedirectProperty IDT operation")))
 		{
@@ -406,7 +493,7 @@ namespace UE::Private::IDT
 		}
 		InferMissingPathInfo(SrcPath, Owner->GetClass());
 
-		FPropertyPathName DstPath;
+		FPropertyPathName DstPath = BasePath;
 		if (!ensureAlwaysMsgf(Params && FPropertyPathName::FromJsonValue(Params->TryGetField(TEXT("dst")), DstPath),
 			TEXT("Failed to parse dst parameter for a RedirectProperty IDT operation")))
 		{
@@ -417,7 +504,7 @@ namespace UE::Private::IDT
 
 		if (UnknownPropertyTree->Remove(SrcPath))
 		{
-			FPropertyValueInContainer DstInstanceValue = TryResolvePropertyPath(DstPath, Owner);
+			FPropertyPathNameResolver DstInstanceValue(DstPath, Owner);
 			if (!DstInstanceValue)
 			{
 				// DstPath is not in the Instance. This means we'll create a new unknown/loose property at DstPath
@@ -425,22 +512,38 @@ namespace UE::Private::IDT
 			}
 
 			// find value in old location (in the old IDO)
-			FPropertyValueInContainer SrcValue = TryResolvePropertyPath(SrcPath, IDO);
+			FPropertyPathNameResolver SrcValue(SrcPath, IDO);
 			check(SrcValue); // it should be impossible for SrcPath to be in the old unknown property tree and not be in the old IDO
 
 			// create a new IDO with the modified property tree
 			IDO = UE::FPropertyBagRepository::Get().ReinstanceWithCachedUnknownPropertyTree(Owner);
 
 			// find destination location in the new IDO
-			FEditPropertyChain IdoChangeChain;
-			TSharedPtr<FPropertyChangedChainEvent> IdoChangeEvent;
-			FPropertyValueInContainer DstIDOValue = TryResolvePropertyPathForChangeEvent(DstPath, IDO, EPropertyChangeType::ValueSet, IdoChangeEvent, IdoChangeChain);
+			FPropertyPathNameResolver DstIDOValue(DstPath, IDO);
 			check(DstIDOValue); // we should've just added this path to the IDO
 
-			// copy property from SrcValue to DstIDOValue (note that because we're using change events, the value will be changed in the instance as well)
-			IDO->PreEditChange(IdoChangeChain);
-			TryCopyOrConvert(SrcValue, DstIDOValue);
-			IDO->PostEditChangeChainProperty(*IdoChangeEvent);
+			while (SrcValue && DstIDOValue) // because of wildcards in the path, it's possible that there's several values to copy!
+			{
+				TSharedPtr<FPropertyChangedChainEvent> IdoChangeEvent;
+				DstIDOValue.BuildChangeEvent(IdoChangeEvent, EPropertyChangeType::ValueSet);
+
+				// copy property from SrcValue to DstIDOValue (note that because we're using change events, the value will be changed in the instance as well)
+				IDO->PreEditChange(IdoChangeEvent->PropertyChain);
+				TryCopyOrConvert(SrcValue.Value, DstIDOValue.Value);
+				IDO->PostEditChangeChainProperty(*IdoChangeEvent);
+
+				/*
+					because containers may contain more than one instance of a struct being redirected, we redirect them all by using wildcards in the path. 
+					FPropertyPathName wildcards are resolved by treating INDEX_NONE as "all indices of this container". FPropertyPathNameResolver will resove 
+					the 0th index first, then Next() will iterate to the next one
+				*/
+				SrcValue.Next();
+				DstIDOValue.Next();
+			}
+
+			// this should never fail unless a user has used wildcards in the path they manually wrote into the IDT
+			ensureMsgf(!SrcValue && !DstIDOValue, TEXT("An IDT ran that operates on every element of a container but the source and destination had a " 
+				"different number of elements! This may have caused corruption in your data. Save at your own risk."));
 		}
 
 		return EInstanceDataTransformResult::Continue;
@@ -456,11 +559,13 @@ namespace UE
 		check(IDO);
 		check(Owner);
 
+		// find and run class IDTs
 		if (const TArray<TSharedPtr<Private::FInstanceDataTransform>>* IDTs = Private::FInstanceDataTransformRepository::Get().FindTransforms(Owner->GetClass()))
 		{
 			for (const TSharedPtr<Private::FInstanceDataTransform>& IDT : *IDTs)
 			{
-				Private::EInstanceDataTransformResult Result = IDT->Func(IDO, Owner, IDT->Args);
+				FPropertyPathName RootPath; // class IDTs operate at the root of the class
+				Private::EInstanceDataTransformResult Result = IDT->Func(IDO, Owner, RootPath, IDT->Args);
 				if (Result == Private::EInstanceDataTransformResult::Break)
 				{
 					break;
@@ -468,6 +573,9 @@ namespace UE
 			}
 		}
 
+		// find and run struct IDTs
+		Private::PatchStructs(IDO, Owner);
+
 		return true;
 	}
 
diff --git a/Engine/Source/Runtime/CoreUObject/Private/UObject/PropertyPathFunctions.cpp b/Engine/Source/Runtime/CoreUObject/Private/UObject/PropertyPathFunctions.cpp
--- a/Engine/Source/Runtime/CoreUObject/Private/UObject/PropertyPathFunctions.cpp
+++ b/Engine/Source/Runtime/CoreUObject/Private/UObject/PropertyPathFunctions.cpp
@@ -35,27 +35,148 @@ inline static const UStruct* FindStructFromProperty(const FProperty* Property)
 	return nullptr;
 }
 
-// call the OnSegment method for every property resolved in the path. If a segment couldn't be resolved, break out early.
-static FPropertyValueInContainer TryIteratePropertyPath(const FPropertyPathName& Path, UObject* Object, TFunctionWithContext<void(const FPropertyPathNameSegment&, FPropertyValueInContainer&)> OnSegment)
+FPropertyValueInContainer TryResolvePropertyPath(const FPropertyPathName& Path, UObject* Object)
 {
-	const UStruct* NextStruct = Object ? Object->GetClass() : nullptr;
-	void* NextContainer = Object;
+	FPropertyPathNameResolver Resolver(Path, Object);
+	return Resolver.Value;
+}
+
+FPropertyPathNameResolver::FPropertyPathNameResolver(const FPropertyPathName& Path, UObject* Object)
+	: Path(Path)
+	, Object(Object)
+{
+	Resolve();
+}
 
-	FPropertyValueInContainer Value;
-	for (int32 Index = 0, Count = Path.GetSegmentCount(); Index < Count; ++Index)
+FPropertyPathNameResolver& FPropertyPathNameResolver::Next()
+{
+	while (!ActiveContainerIterators.IsEmpty())
+	{
+		FContainerIterator& Info = ActiveContainerIterators.Last();
+		++Info.CurIndex;
+		if (Info.Property->ArrayDim > 1) // static array
+		{
+			if (Info.CurIndex < Info.Property->ArrayDim)
+			{
+				return Resolve();
+			}
+		}
+		else
+		{
+			void* DataPtr = Info.Property->ContainerPtrToValuePtr<uint8>(Info.NextContainer);
+			if (const FArrayProperty* AsArrayProperty = CastField<FArrayProperty>(Info.Property))
+			{
+				FScriptArrayHelper Helper(AsArrayProperty, DataPtr);
+				if (Info.CurIndex < Helper.Num())
+				{
+					return Resolve();
+				}
+			}
+			else if (const FSetProperty* AsSetProperty = CastField<FSetProperty>(Info.Property))
+			{
+				FScriptSetHelper Helper(AsSetProperty, DataPtr);
+				if (Info.CurIndex < Helper.Num())
+				{
+					return Resolve();
+				}
+			}
+			else if (const FMapProperty* AsMapProperty = CastField<FMapProperty>(Info.Property))
+			{
+				FScriptMapHelper Helper(AsMapProperty, DataPtr);
+				if (Info.CurIndex < Helper.Num())
+				{
+					return Resolve();
+				}
+			}
+		}
+
+		ActiveContainerIterators.Pop();
+	}
+	Value = {};
+	return *this;
+}
+
+bool FPropertyPathNameResolver::BuildChangeEvent(TSharedPtr<FPropertyChangedChainEvent>& OutEvent, int32 ChangeType)
+{
+	if (Value)
+	{
+		EventArrayIndices.Empty();
+		EventArrayIndices.Emplace();
+		for (int32 PathIndex = 0, SegmentCount = Path.GetSegmentCount(); PathIndex < SegmentCount; ++PathIndex)
+		{
+			FPropertyPathNameSegment Segment = Path.GetSegment(PathIndex);
+			if (Segment.Index != INDEX_NONE)
+			{
+				EventArrayIndices[0].Add(Segment.Name.ToString(), Segment.Index);
+			}
+		}
+		for (const FContainerIterator& WildcardInfo : ActiveContainerIterators)
+		{
+			EventArrayIndices[0].Add(WildcardInfo.Property->GetName(), WildcardInfo.CurIndex);
+		}
+
+		FPropertyChangedEvent InnerEvent(const_cast<FProperty*>(Value.Property), ChangeType, { Object });
+		InnerEvent.ObjectIteratorIndex = 0;
+		OutEvent = MakeShared<FPropertyChangedChainEvent>(EventChain, InnerEvent);
+		OutEvent->SetArrayIndexPerObject(EventArrayIndices);
+		return true;
+	}
+	return false;
+}
+
+FPropertyPathNameResolver& FPropertyPathNameResolver::Resolve()
+{
+	const UStruct* NextStruct;
+	void* NextContainer;
+	int32 PathIndex;
+	if (ActiveContainerIterators.IsEmpty())
+	{
+		NextStruct = Object ? Object->GetClass() : nullptr;
+		NextContainer = Object;
+		PathIndex = 0;
+	}
+	else
+	{
+		// jump to the current container being entered
+		NextStruct = ActiveContainerIterators.Last().NextStruct;
+		NextContainer = ActiveContainerIterators.Last().NextContainer;
+		PathIndex = ActiveContainerIterators.Last().PathIndex;
+	}
+
+	for (const int32 SegmentCount = Path.GetSegmentCount(); PathIndex < SegmentCount; ++PathIndex)
 	{
 		// Fail if the previous segment failed to resolve the struct or container for this segment.
 		if (!NextStruct || !NextContainer)
 		{
-			return {};
+			Value = {};
+			return *this;
 		}
 
-		const FPropertyPathNameSegment Segment = Path.GetSegment(Index);
+		FPropertyPathNameSegment Segment = Path.GetSegment(PathIndex);
 		const FProperty* Property = FindPropertyByNameAndTypeName(NextStruct, Segment.Name, Segment.Type);
 
+		// Support for wildcard Indices
+		if (PathIndex < SegmentCount - 1) // don't use wildcards for the last element in the path
+		{
+			if (Segment.Index == INDEX_NONE)
+			{
+				if (!ActiveContainerIterators.IsEmpty() && ActiveContainerIterators.Last().PathIndex == PathIndex)
+				{
+					// wildcard element detected! use the current index!
+					Segment.Index = ActiveContainerIterators.Last().CurIndex;
+				}
+				else if (Property->ArrayDim > 1 || Property->IsA<FArrayProperty>() || Property->IsA<FSetProperty>() || Property->IsA<FMapProperty>())
+				{
+					Segment.Index = 0;
+					ActiveContainerIterators.Emplace(PathIndex, Property, NextContainer, NextStruct, Segment.Index);
+				}
+			}
+		}
+
 		if (!Property)
 		{
-			return {};
+			Value = {};
+			return *this;
 		}
 
 		Value.Property = Property;
@@ -68,7 +189,8 @@ static FPropertyValueInContainer TryIteratePropertyPath(const FPropertyPathName&
 		{
 			if (Segment.Index < 0 || Segment.Index >= Property->ArrayDim)
 			{
-				return {};
+				Value = {};
+				return *this;
 			}
 			Value.ArrayIndex = Segment.Index;
 		}
@@ -83,12 +205,13 @@ static FPropertyValueInContainer TryIteratePropertyPath(const FPropertyPathName&
 			FOptionalPropertyLayout OptionalHelper(OptionalProperty->GetValueProperty());
 			NextStruct = FindStructFromProperty(OptionalHelper.GetValueProperty());
 			NextContainer = OptionalHelper.GetValuePointerForReadOrReplaceIfSet(NextContainer);
+			if (!NextContainer && PathIndex < SegmentCount - 1)
+			{
+				return Next(); // this optional is unset. Try the next one in the list of wildcards
+			}
 		}
 
-		if (OnSegment)
-		{
-			OnSegment(Segment, Value);
-		}
+		EventChain.AddTail(const_cast<FProperty*>(Value.Property));
 
 		// Scalar values and static containers are finished resolving.
 		if (Property->ArrayDim > 1 || Segment.Index == INDEX_NONE)
@@ -104,7 +227,7 @@ static FPropertyValueInContainer TryIteratePropertyPath(const FPropertyPathName&
 			FScriptArrayHelper ArrayHelper(ArrayProperty, NextContainer);
 			if (!ArrayHelper.IsValidIndex(Segment.Index))
 			{
-				return {};
+				return Next(); // this index isn't in the array. Try the next one in the list of wildcards
 			}
 
 			NextStruct = FindStructFromProperty(ArrayProperty->Inner);
@@ -119,7 +242,7 @@ static FPropertyValueInContainer TryIteratePropertyPath(const FPropertyPathName&
 			NextContainer = SetHelper.FindNthElementPtr(Segment.Index);
 			if (!NextContainer)
 			{
-				return {};
+				return Next(); // this index isn't in the set. Try the next one in the list of wildcards
 			}
 			Value.Property = SetProperty->ElementProp;
 			Value.Container = NextContainer;
@@ -127,16 +250,18 @@ static FPropertyValueInContainer TryIteratePropertyPath(const FPropertyPathName&
 		else if (const FMapProperty* MapProperty = CastField<FMapProperty>(Property))
 		{
 			FScriptMapHelper MapHelper(MapProperty, NextContainer);
-			if (++Index == Count)
+			if (++PathIndex == SegmentCount)
 			{
-				return {};
+				Value = {};
+				return *this;
 			}
 
 			// A Key or Value segment with no type or index is required to distinguish which property to resolve.
-			const FPropertyPathNameSegment MapSegment = Path.GetSegment(Index);
+			const FPropertyPathNameSegment MapSegment = Path.GetSegment(PathIndex);
 			if (!MapSegment.Type.IsEmpty() || MapSegment.Index != INDEX_NONE)
 			{
-				return {};
+				Value = {};
+				return *this;
 			}
 
 			if (MapSegment.Name == NAME_Key)
@@ -145,7 +270,7 @@ static FPropertyValueInContainer TryIteratePropertyPath(const FPropertyPathName&
 				NextContainer = MapHelper.FindNthKeyPtr(Segment.Index);
 				if (!NextContainer)
 				{
-					return {};
+					return Next(); // this index isn't in the map. Try the next one in the list of wildcards
 				}
 				Value.Property = MapProperty->KeyProp;
 			}
@@ -155,13 +280,14 @@ static FPropertyValueInContainer TryIteratePropertyPath(const FPropertyPathName&
 				NextContainer = MapHelper.FindNthValuePtr(Segment.Index);
 				if (!NextContainer)
 				{
-					return {};
+					return Next(); // this index isn't in the map. Try the next one in the list of wildcards
 				}
 				Value.Property = MapProperty->ValueProp;
 			}
 			else
 			{
-				return {};
+				Value = {};
+				return *this;
 			}
 
 			// The key and value property both have an offset relative to the pair.
@@ -169,41 +295,13 @@ static FPropertyValueInContainer TryIteratePropertyPath(const FPropertyPathName&
 		}
 		else
 		{
-			return {};
-		}
-
-		if (OnSegment)
-		{
-			OnSegment(Segment, Value);
+			Value = {};
+			return *this;
 		}
 	}
-	return Value;
+	return *this;
 }
 
-FPropertyValueInContainer TryResolvePropertyPath(const FPropertyPathName& Path, UObject* Object)
-{
-	return TryIteratePropertyPath(Path, Object, nullptr);
-}
 
-FPropertyValueInContainer TryResolvePropertyPathForChangeEvent(const FPropertyPathName& Path, UObject* Object, uint32 ChangeType, TSharedPtr<FPropertyChangedChainEvent>& OutEvent, FEditPropertyChain& OutChain)
-{
-	TMap<FString, int32> ArrayIndeces;
-	FPropertyValueInContainer Result = TryIteratePropertyPath(Path, Object, [&OutChain, &ArrayIndeces](const FPropertyPathNameSegment& Segment, FPropertyValueInContainer& Value)
-		{
-			OutChain.AddTail(const_cast<FProperty*>(Value.Property));
-			if (Segment.Index != INDEX_NONE)
-			{
-				ArrayIndeces.Add(Value.Property->GetName(), Segment.Index);
-			}
-		});
-
-	if (Result)
-	{
-		FPropertyChangedEvent InnerEvent(const_cast<FProperty*>(Result.Property), ChangeType, {Object});
-		OutEvent = MakeShared<FPropertyChangedChainEvent>(OutChain, InnerEvent);
-		OutEvent->SetArrayIndexPerObject({ ArrayIndeces });
-	}
-	return Result;
-}
 
 } // UE
diff --git a/Engine/Source/Runtime/CoreUObject/Public/UObject/PropertyPathFunctions.h b/Engine/Source/Runtime/CoreUObject/Public/UObject/PropertyPathFunctions.h
--- a/Engine/Source/Runtime/CoreUObject/Public/UObject/PropertyPathFunctions.h
+++ b/Engine/Source/Runtime/CoreUObject/Public/UObject/PropertyPathFunctions.h
@@ -4,6 +4,9 @@
 
 #include "CoreTypes.h"
 #include "Templates/SharedPointerFwd.h"
+#include "Containers/SharedString.h"
+#include "Containers/Map.h"
+#include "UnrealType.h"
 
 #define UE_API COREUOBJECT_API
 
@@ -76,6 +79,44 @@ struct FPropertyValueInContainer
 	}
 };
 
+// FPropertyPathNames can have wildcarded container properties (signified by no index provided). 
+// In that case, Value will be set to the first element and next() will iterate to the next
+class FPropertyPathNameResolver
+{
+public:
+	UE_API explicit FPropertyPathNameResolver(const FPropertyPathName& Path, UObject* Object);
+
+	FPropertyValueInContainer Value;
+
+	// returns true if the path was resolvable. In the case of a wildcarded path, returns false when all indices have been visited
+	inline explicit operator bool() const
+	{
+		return !!Value;
+	}
+
+	// in the case of wildcarded paths, next will iterate to the next available index in the conatainer(s).
+	UE_API FPropertyPathNameResolver& Next();
+	UE_API bool BuildChangeEvent(TSharedPtr<FPropertyChangedChainEvent>& OutEvent, int32 ChangeType);
+private:
+	FPropertyPathNameResolver& Resolve();
+	const FPropertyPathName& Path;
+	UObject* Object;
+	FEditPropertyChain EventChain;
+	TArray<TMap<FString, int32>> EventArrayIndices;
+
+	struct FContainerIterator
+	{
+		int32 PathIndex;
+		const FProperty* Property;
+		void* NextContainer;
+		const UStruct* NextStruct;
+		int32 CurIndex;
+	};
+
+	// stack that holds the current iteration info of wildcard iterators
+	TArray<FContainerIterator> ActiveContainerIterators;
+};
+
 /**
  * Try to resolve the property path to a property value within the object.
  *
@@ -85,19 +126,6 @@ struct FPropertyValueInContainer
  */
 UE_API FPropertyValueInContainer TryResolvePropertyPath(const FPropertyPathName& Path, UObject* Object);
 
-/**
- * Try to resolve the property path to a property value within the object 
- * and construct an FPropertyChangeChainEvent that can be used to notify changes to the resolved data
- *
- * @param Path The non-empty property path to resolve.
- * @param Object The object representing the root of the property path to resolve.
- * @param ChangeType The change type to assign to the FPropertyChangeChainEvent
- * @param OutEvent The resulting change event
- * @param OutChain The property chain associated with the event
- * @return A valid property value reference if resolved, otherwise an invalid reference.
- */
-UE_API FPropertyValueInContainer TryResolvePropertyPathForChangeEvent(const FPropertyPathName& Path, UObject* Object, uint32 ChangeType, TSharedPtr<FPropertyChangedChainEvent>& OutEvent, FEditPropertyChain& OutChain);
-
 } // UE
 
 #undef UE_API
diff --git a/Engine/Source/Runtime/CoreUObject/Tests/UObject/PropertyPathFunctionsTest.cpp b/Engine/Source/Runtime/CoreUObject/Tests/UObject/PropertyPathFunctionsTest.cpp
--- a/Engine/Source/Runtime/CoreUObject/Tests/UObject/PropertyPathFunctionsTest.cpp
+++ b/Engine/Source/Runtime/CoreUObject/Tests/UObject/PropertyPathFunctionsTest.cpp
@@ -55,6 +55,7 @@ TEST_CASE_NAMED(FPropertyPathFunctionsTryResolveTest, "CoreUObject::PropertyPath
 	StructRef.Int32Set = {0x106, 0x107, 0x108, 0x109};
 	StructRef.Int32Map = {{0x10a, 0x10b}, {0x10c, 0x10d}, {0x10e, 0x10f}};
 	StructRef.Int32Optional = 0x110;
+	StructRef.StructOptional = FTestPropertyPathFunctionsStructKey{ 0x111 };
 
 	TestObject->StructArray.AddDefaulted(3);
 	TestObject->StructArray.Add(StructRef);
@@ -315,6 +316,167 @@ TEST_CASE_NAMED(FPropertyPathFunctionsTryResolveTest, "CoreUObject::PropertyPath
 
 	// Resolve unset StructOptional -> Int32
 	CHECK_FALSE((Value = TryResolvePropertyPath(MakePath({{TestClass, "StructOptional"}, {TestStruct, "Int32"}}), TestObject)));
+
+
+	// WILDCARDS
+	constexpr int32 StructStaticArraySize = UE_ARRAY_COUNT(UTestPropertyPathFunctionsClass::StructStaticArray);
+	constexpr int32 Int32StaticArraySize = UE_ARRAY_COUNT(FTestPropertyPathFunctionsStruct::Int32StaticArray);
+	
+	// Resolve StructStaticArray[*] -> Int32Array
+	{
+		FPropertyPathName Path = MakePath({ {TestClass, "StructStaticArray", INDEX_NONE}, {TestStruct, "Int32Array"} });
+		FPropertyPathNameResolver Resolver(Path, TestObject);
+		REQUIRE(Resolver);
+		int32 FoundCount = 0;
+		while (Resolver)
+		{
+			CHECK(Resolver.Value.ArrayIndex == 0);
+			CHECK(Resolver.Value.Struct == TestStruct);
+			CHECK(Resolver.Value.Property->GetFName() == "Int32Array");
+			CHECK(Resolver.Value.GetValuePtr<void>() == &TestObject->StructStaticArray[FoundCount].Int32Array);
+			++FoundCount;
+			Resolver.Next();
+		}
+		CHECK(FoundCount == StructStaticArraySize);
+	}
+
+	// Resolve StructStaticArray[*] -> Int32Map[*] -> Key
+	{
+		FPropertyPathName Path = MakePath({ {TestClass, "StructStaticArray", INDEX_NONE}, {TestStruct, "Int32Map", INDEX_NONE}, {nullptr, "Key"} });
+		FPropertyPathNameResolver Resolver(Path, TestObject);
+		REQUIRE(Resolver);
+		int32 FoundCount = 0;
+		TMap<int32, int32>::TConstIterator ExpectedDataIterator = StructRef.Int32Map.CreateConstIterator();
+		while (Resolver)
+		{
+			CHECK(Resolver.Value.ArrayIndex == 0);
+			CHECK(Resolver.Value.Struct == nullptr);
+			CHECK(Resolver.Value.Property->GetFName() == "Int32Map_Key");
+			CHECK(*Resolver.Value.GetValuePtr<int32>() == ExpectedDataIterator->Key);
+			++FoundCount;
+			++ExpectedDataIterator;
+			Resolver.Next();
+		}
+		CHECK(FoundCount == 1 * StructRef.Int32Map.Num()); // only StructStaticArray[3].Int32Map was given elements
+	}
+
+	// Resolve StructStaticArray[*] -> Int32Map[*] -> Value
+	{
+		FPropertyPathName Path = MakePath({ {TestClass, "StructStaticArray", INDEX_NONE}, {TestStruct, "Int32Map", INDEX_NONE}, {nullptr, "Value"} });
+		FPropertyPathNameResolver Resolver(Path, TestObject);
+		REQUIRE(Resolver);
+		int32 FoundCount = 0;
+		TMap<int32, int32>::TConstIterator ExpectedDataIterator = StructRef.Int32Map.CreateConstIterator();
+		while (Resolver)
+		{
+			CHECK(Resolver.Value.ArrayIndex == 0);
+			CHECK(Resolver.Value.Struct == nullptr);
+			CHECK(Resolver.Value.Property->GetFName() == "Int32Map");
+			CHECK(*Resolver.Value.GetValuePtr<int32>() == ExpectedDataIterator->Value);
+			++FoundCount;
+			++ExpectedDataIterator;
+			Resolver.Next();
+		}
+		CHECK(FoundCount == 1 * StructRef.Int32Map.Num()); // only StructStaticArray[3].Int32Map was given elements
+	}
+
+	// Resolve StructStaticArray[*] -> Int32Optional
+	{
+		FPropertyPathName Path = MakePath({ {TestClass, "StructStaticArray", INDEX_NONE}, {TestStruct, "Int32Optional"} });
+		FPropertyPathNameResolver Resolver(Path, TestObject);
+		REQUIRE(Resolver);
+		int32 FoundCount = 0;
+		while (Resolver)
+		{
+			CHECK(Resolver.Value.ArrayIndex == 0);
+			CHECK(Resolver.Value.Struct == TestStruct);
+			CHECK(Resolver.Value.Property->GetFName() == "Int32Optional");
+			CHECK(Resolver.Value.GetValuePtr<TOptional<int32>>() == &TestObject->StructStaticArray[FoundCount].Int32Optional);
+			++FoundCount;
+			Resolver.Next();
+		}
+		CHECK(FoundCount == StructStaticArraySize);
+	}
+
+	// Resolve StructStaticArray[*] -> StructOptional -> Key
+	{
+		FPropertyPathName Path = MakePath({ {TestClass, "StructStaticArray", INDEX_NONE}, {TestStruct, "StructOptional"}, {TestStructKey, "Key"} });
+		FPropertyPathNameResolver Resolver(Path, TestObject);
+		REQUIRE(Resolver);
+		int32 FoundCount = 0;
+		while (Resolver)
+		{
+			CHECK(Resolver.Value.ArrayIndex == 0);
+			CHECK(Resolver.Value.Struct == TestStructKey);
+			CHECK(Resolver.Value.Property->GetFName() == "Key");
+			CHECK(Resolver.Value.GetValuePtr<int32>() == &StructRef.StructOptional->Key);
+			++FoundCount;
+			Resolver.Next();
+		}
+		CHECK(FoundCount == 1); // only StructStaticArray[3].StructOptional was given a value
+	}
+
+	// Resolve StructSet[*] -> Key
+	{
+		FPropertyPathName Path = MakePath({ {TestClass, "StructSet", INDEX_NONE}, {TestStructKey, "Key"} });
+		FPropertyPathNameResolver Resolver(Path, TestObject);
+		REQUIRE(Resolver);
+		int32 FoundCount = 0;
+		TSet<FTestPropertyPathFunctionsStructKey>::TConstIterator ExpectedDataIterator = TestObject->StructSet.CreateConstIterator();
+		while (Resolver)
+		{
+			CHECK(Resolver.Value.ArrayIndex == 0);
+			CHECK(Resolver.Value.Struct == TestStructKey);
+			CHECK(Resolver.Value.Property->GetFName() == "Key");
+			CHECK(Resolver.Value.GetValuePtr<int32>() == &ExpectedDataIterator->Key);
+			CHECK(*Resolver.Value.GetValuePtr<int32>() == ExpectedDataIterator->Key);
+			++FoundCount;
+			++ExpectedDataIterator;
+			Resolver.Next();
+		}
+		CHECK(FoundCount == TestObject->StructSet.Num());
+	}
+
+	// Resolve StructMap[*] -> Key -> Key
+	{
+		FPropertyPathName Path = MakePath({ {TestClass, "StructMap", INDEX_NONE}, {nullptr, "Key"}, {TestStructKey, "Key"} });
+		FPropertyPathNameResolver Resolver(Path, TestObject);
+		REQUIRE(Resolver);
+		int32 FoundCount = 0;
+		TMap<FTestPropertyPathFunctionsStructKey, FTestPropertyPathFunctionsStruct>::TConstIterator ExpectedDataIterator = TestObject->StructMap.CreateConstIterator();
+		while (Resolver)
+		{
+			CHECK(Resolver.Value.ArrayIndex == 0);
+			CHECK(Resolver.Value.Struct == TestStructKey);
+			CHECK(Resolver.Value.Property->GetFName() == "Key");
+			CHECK(Resolver.Value.GetValuePtr<int32>() == &ExpectedDataIterator->Key.Key);
+			CHECK(*Resolver.Value.GetValuePtr<int32>() == ExpectedDataIterator->Key.Key);
+			++FoundCount;
+			++ExpectedDataIterator;
+			Resolver.Next();
+		}
+		CHECK(FoundCount == TestObject->StructMap.Num());
+	}
+
+	// Resolve StructMap[*] -> Value
+	{
+		FPropertyPathName Path = MakePath({ {TestClass, "StructMap", INDEX_NONE}, {nullptr, "Value"} });
+		FPropertyPathNameResolver Resolver(Path, TestObject);
+		REQUIRE(Resolver);
+		int32 FoundCount = 0;
+		TMap<FTestPropertyPathFunctionsStructKey, FTestPropertyPathFunctionsStruct>::TConstIterator ExpectedDataIterator = TestObject->StructMap.CreateConstIterator();
+		while (Resolver)
+		{
+			CHECK(Resolver.Value.ArrayIndex == 0);
+			CHECK(Resolver.Value.Struct == nullptr);
+			CHECK(Resolver.Value.Property->GetFName() == "StructMap");
+			CHECK(Resolver.Value.GetValuePtr<void>() == &ExpectedDataIterator->Value);
+			++FoundCount;
+			++ExpectedDataIterator;
+			Resolver.Next();
+		}
+		CHECK(FoundCount == TestObject->StructMap.Num());
+	}
 }
 
 } // UE
diff --git a/Engine/Source/Runtime/CoreUObject/Tests/UObject/PropertyPathFunctionsTest.h b/Engine/Source/Runtime/CoreUObject/Tests/UObject/PropertyPathFunctionsTest.h
--- a/Engine/Source/Runtime/CoreUObject/Tests/UObject/PropertyPathFunctionsTest.h
+++ b/Engine/Source/Runtime/CoreUObject/Tests/UObject/PropertyPathFunctionsTest.h
@@ -72,6 +72,9 @@ public:
 
 	UPROPERTY()
 	TOptional<int32> Int32Optional;
+
+	UPROPERTY()
+	TOptional<FTestPropertyPathFunctionsStructKey> StructOptional;
 };
 
 UCLASS()
