[BugFix] Couple bug fixes caught by BugHawk

diff --git a/Engine/Source/Runtime/CoreUObject/Public/UObject/PropertyPathFunctions.h b/Engine/Source/Runtime/CoreUObject/Public/UObject/PropertyPathFunctions.h
--- a/Engine/Source/Runtime/CoreUObject/Public/UObject/PropertyPathFunctions.h
+++ b/Engine/Source/Runtime/CoreUObject/Public/UObject/PropertyPathFunctions.h
@@ -80,7 +80,7 @@ struct FPropertyValueInContainer
 };
 
 // FPropertyPathNames can have wildcarded container properties (signified by no index provided). 
-// In that case, Value will be set to the first element and next() will iterate to the next
+// In that case, Value will be set to the first element and Next() will iterate to the next
 class FPropertyPathNameResolver
 {
 public:
@@ -94,7 +94,7 @@ public:
 		return !!Value;
 	}
 
-	// in the case of wildcarded paths, next will iterate to the next available index in the conatainer(s).
+	// in the case of wildcarded paths, Next() will iterate to the next available index in the conatainer(s).
 	UE_API FPropertyPathNameResolver& Next();
 	UE_API bool BuildChangeEvent(TSharedPtr<FPropertyChangedChainEvent>& OutEvent, int32 ChangeType);
 private:
