Converted TLazyObjectPtr's operations to UEOps.

diff --git a/Engine/Source/Runtime/CoreUObject/Public/UObject/LazyObjectPtr.h b/Engine/Source/Runtime/CoreUObject/Public/UObject/LazyObjectPtr.h
--- a/Engine/Source/Runtime/CoreUObject/Public/UObject/LazyObjectPtr.h
+++ b/Engine/Source/Runtime/CoreUObject/Public/UObject/LazyObjectPtr.h
@@ -6,11 +6,13 @@
 
 #pragma once
 
+#include "Concepts/PointerEqualityComparableWith.h"
 #include "Containers/UnrealString.h"
 #include "CoreGlobals.h"
 #include "HAL/Platform.h"
 #include "HAL/ThreadSafeCounter.h"
 #include "Misc/Guid.h"
+#include "Misc/UEOps.h"
 #include "Serialization/Archive.h"
 #include "Serialization/StructuredArchive.h"
 #include "Templates/Casts.h"
@@ -379,42 +381,20 @@ public:
 		Ar << static_cast<FLazyObjectPtr&>(*this);
 	}
 
-	/** Compare with another TLazyObjectPtr of related type */
-	template<typename U, typename = decltype((T*)nullptr == (U*)nullptr)>
-	UE_FORCEINLINE_HINT bool operator==(const TLazyObjectPtr<U>& Rhs) const
+	template <UE::CPointerEqualityComparableWith<T> U>
+	UE_FORCEINLINE_HINT bool UEOpEquals(const TLazyObjectPtr<U>& Rhs) const
 	{
 		return (const FLazyObjectPtr&)*this == (const FLazyObjectPtr&)Rhs;
 	}
-	template<typename U, typename = decltype((T*)nullptr != (U*)nullptr)>
-	UE_FORCEINLINE_HINT bool operator!=(const TLazyObjectPtr<U>& Rhs) const
-	{
-		return (const FLazyObjectPtr&)*this != (const FLazyObjectPtr&)Rhs;
-	}
-
-	/** Compare for equality with a raw pointer **/
-	template<typename U, typename = decltype((T*)nullptr == (U*)nullptr)>
-	UE_FORCEINLINE_HINT bool operator==(const U* Rhs) const
+	template <UE::CPointerEqualityComparableWith<T> U>
+	UE_FORCEINLINE_HINT bool UEOpEquals(const U* Rhs) const
 	{
 		return Get() == Rhs;
 	}
-
-	/** Compare to null */
-	UE_FORCEINLINE_HINT bool operator==(TYPE_OF_NULLPTR) const
+	UE_FORCEINLINE_HINT bool UEOpEquals(TYPE_OF_NULLPTR) const
 	{
 		return !IsValid();
 	}
-	/** Compare for inequality with a raw pointer	**/
-	template<typename U, typename = decltype((T*)nullptr != (U*)nullptr)>
-	UE_FORCEINLINE_HINT bool operator!=(const U* Rhs) const
-	{
-		return Get() != Rhs;
-	}
-
-	/** Compare for inequality with null **/
-	UE_FORCEINLINE_HINT bool operator!=(TYPE_OF_NULLPTR) const
-	{
-		return IsValid();
-	}
 };
 
 /** Hash function. */
@@ -431,33 +411,5 @@ FArchive& operator<<(FArchive& Ar, TLazyObjectPtr<T>& LazyObjectPtr)
 	return Ar;
 }
 
-/** Compare for equality with a raw pointer **/
-template<typename T, typename U, typename = decltype((T*)nullptr == (U*)nullptr)>
-UE_FORCEINLINE_HINT bool operator==(const U* Lhs, const TLazyObjectPtr<T>& Rhs)
-{
-	return Lhs == Rhs.Get();
-}
-
-/** Compare to null */
-template<typename T>
-UE_FORCEINLINE_HINT bool operator==(TYPE_OF_NULLPTR, const TLazyObjectPtr<T>& Rhs)
-{
-	return !Rhs.IsValid();
-}
-
-/** Compare for inequality with a raw pointer	**/
-template<typename T, typename U, typename = decltype((T*)nullptr != (U*)nullptr)>
-UE_FORCEINLINE_HINT bool operator!=(const U* Lhs, const TLazyObjectPtr<T>& Rhs)
-{
-	return Lhs != Rhs.Get();
-}
-
-/** Compare for inequality with null **/
-template<typename T>
-UE_FORCEINLINE_HINT bool operator!=(TYPE_OF_NULLPTR, const TLazyObjectPtr<T>& Rhs)
-{
-	return Rhs.IsValid();
-}
-
 template<class T> struct TIsPODType<TLazyObjectPtr<T> > { enum { Value = TIsPODType<FLazyObjectPtr>::Value }; };
 template<class T> struct TIsWeakPointerType<TLazyObjectPtr<T> > { enum { Value = TIsWeakPointerType<FLazyObjectPtr>::Value }; };
