RemoteObjects: Adding more information about where ExecuteWork is coming from.  Giving additonal context in certain cases about what is happening.

diff --git a/Engine/Source/Runtime/CoreUObject/Private/UObject/RemoteExecutor.cpp b/Engine/Source/Runtime/CoreUObject/Private/UObject/RemoteExecutor.cpp
--- a/Engine/Source/Runtime/CoreUObject/Private/UObject/RemoteExecutor.cpp
+++ b/Engine/Source/Runtime/CoreUObject/Private/UObject/RemoteExecutor.cpp
@@ -1029,10 +1029,7 @@ FRemoteWorkPriority CreateRootWorkPriority()
 
 void ExecutePendingWork()
 {
-	static const FName FuncName{ __func__ };
-	UE::RemoteExecutor::OnRegionBeginDelegate.ExecuteIfBound(FuncName);
 	GRemoteExecutor.ExecutePendingWork();
-	UE::RemoteExecutor::OnRegionEndDelegate.ExecuteIfBound(FString());
 }
 
 void ExecuteTransactionalWithExplicitPriority(FName WorkName, FRemoteWorkPriority WorkPriority, const TFunctionRef<void(void)>& Work)
diff --git a/Engine/Source/Runtime/CoreUObject/Private/UObject/RemoteObjectTransfer.cpp b/Engine/Source/Runtime/CoreUObject/Private/UObject/RemoteObjectTransfer.cpp
--- a/Engine/Source/Runtime/CoreUObject/Private/UObject/RemoteObjectTransfer.cpp
+++ b/Engine/Source/Runtime/CoreUObject/Private/UObject/RemoteObjectTransfer.cpp
@@ -268,7 +268,7 @@ public:
 						}
 					}
 
-					UE::RemoteExecutor::OnRegionEndDelegate.ExecuteIfBound(FString::Format(TEXT("{0} {1}"), {*RootObject->GetName(), *CanonicalRootObjectId.ToString()}));
+					UE::RemoteExecutor::OnRegionEndDelegate.ExecuteIfBound(FString::Format(TEXT("{0} {1} -> {2}"), {*RootObject->GetName(), *CanonicalRootObjectId.ToString(), *MigrationContext.RemoteServerId.ToString()}));
 
 					// all done, delete all requests for this object
 					PendingObjectMigrationRequests.RemoveAt(MigrationRequestsIndex, EAllowShrinking::No);
@@ -953,7 +953,7 @@ UObject* MigrateObjectFromObjectData(FRemoteObjectData& ObjectData, const FUObje
 		}
 	}
 
-	UE::RemoteExecutor::OnRegionEndDelegate.ExecuteIfBound(FString::Format(TEXT("{0} {1} ({2} bytes)"), {*ReceivedObjects[0]->GetName(), *RootObjectRemoteId.ToString(), ByteSize}));
+	UE::RemoteExecutor::OnRegionEndDelegate.ExecuteIfBound(FString::Format(TEXT("{0} {1} from {2} ({3} bytes)"), {*ReceivedObjects[0]->GetName(), *RootObjectRemoteId.ToString(), *MigrationContext.RemoteServerId.ToString(), ByteSize}));
 
 	return ReceivedObjects[RequestedObjectIndex];
 } 
diff --git a/Engine/Source/Runtime/Engine/Private/NetDriver.cpp b/Engine/Source/Runtime/Engine/Private/NetDriver.cpp
--- a/Engine/Source/Runtime/Engine/Private/NetDriver.cpp
+++ b/Engine/Source/Runtime/Engine/Private/NetDriver.cpp
@@ -2808,7 +2808,10 @@ void UNetDriver::PostTickDispatch()
 			}
 
 			// Actually do all of the work
+			static const FName RegionName { __func__ };
+			UE::RemoteExecutor::OnRegionBeginDelegate.ExecuteIfBound(RegionName);
 			UE::RemoteExecutor::ExecutePendingWork();
+			UE::RemoteExecutor::OnRegionEndDelegate.ExecuteIfBound(TEXT(""));
 
 			ensureMsgf(TransactingReplicators.IsEmpty(), TEXT("We expected all TransactingReplicators to be serviced by now!"));
 			ensureMsgf(ReplicatorsWithPendingTransactions.IsEmpty(), TEXT("Execution of RPCs or pending work caused more RPCs to be enqueued -- those will be delayed another frame"));
diff --git a/Engine/Source/Runtime/Engine/Private/TickTaskManager.cpp b/Engine/Source/Runtime/Engine/Private/TickTaskManager.cpp
--- a/Engine/Source/Runtime/Engine/Private/TickTaskManager.cpp
+++ b/Engine/Source/Runtime/Engine/Private/TickTaskManager.cpp
@@ -338,7 +338,9 @@ public:
 
 			if (Target->bRunTransactionally)
 			{
+				UE::RemoteExecutor::OnRegionBeginDelegate.ExecuteIfBound(UEnum::GetValueAsName(Context.TickGroup));
 				UE::RemoteExecutor::ExecuteTransactional(Target->CachedDiagnosticContext, ExecuteTickWork);
+				UE::RemoteExecutor::OnRegionEndDelegate.ExecuteIfBound(Target->CachedDiagnosticMessage);
 			}
 			else
 			{
@@ -444,7 +446,9 @@ public:
 		}
 
 #if UE_WITH_REMOTE_OBJECT_HANDLE
+		UE::RemoteExecutor::OnRegionBeginDelegate.ExecuteIfBound(UEnum::GetValueAsName(Context.TickGroup));
 		UE::RemoteExecutor::ExecutePendingWork();
+		UE::RemoteExecutor::OnRegionEndDelegate.ExecuteIfBound(TEXT("BatchTick"));
 #endif
 	}
 };
