Fix discovery of VNI packages during merged modules unloading

The previous implementation did not correctly find these paths.

#rnx

diff --git a/Engine/Source/Runtime/CoreUObject/Private/UObject/UObjectBase.cpp b/Engine/Source/Runtime/CoreUObject/Private/UObject/UObjectBase.cpp
--- a/Engine/Source/Runtime/CoreUObject/Private/UObject/UObjectBase.cpp
+++ b/Engine/Source/Runtime/CoreUObject/Private/UObject/UObjectBase.cpp
@@ -1128,12 +1128,15 @@ static TArray<UPackage*> GetScriptPackagesFromModuleNames(TConstArrayView<FName>
 
 	for (const FName ModuleName : InModuleNames)
 	{
+		// Identify regular script packages. The expected package name is '/Script/ModuleName'.
 		if (UPackage* Package = FindObjectFast<UPackage>(/* Outer */ nullptr, WriteToString<256>(TEXT("/Script/"), ModuleName).ToString()))
 		{
 			Packages.Add(Package);
 		}
 
-		if (UPackage* Package = FindObjectFast<UPackage>(/* Outer */ nullptr, WriteToString<256>(TEXT("/_Verse/VNI/"), ModuleName).ToString()))
+		// Identify Verse packages. The expected package name is '/PluginName/_Verse/VNI/ModuleName'.
+		FName PluginName = IPluginManager::Get().PackageNameFromModuleName(ModuleName);
+		if (UPackage* Package = FindObjectFast<UPackage>(/* Outer */ nullptr, WriteToString<256>(TEXT("/"), PluginName, TEXT("/_Verse/VNI/"), ModuleName).ToString()))
 		{
 			Packages.Add(Package);
 		}
diff --git a/Engine/Source/Runtime/Projects/Private/PluginManager.cpp b/Engine/Source/Runtime/Projects/Private/PluginManager.cpp
--- a/Engine/Source/Runtime/Projects/Private/PluginManager.cpp
+++ b/Engine/Source/Runtime/Projects/Private/PluginManager.cpp
@@ -3737,7 +3737,7 @@ FName FPluginManager::PackageNameFromModuleName(FName ModuleName)
 			const FModuleDescriptor& Descriptor = Modules[Idx];
 			if (Descriptor.Name == ModuleName)
 			{
-				UE_LOG(LogPluginManager, Log, TEXT("Module %s belongs to Plugin %s and we assume that is the name of the package with the UObjects is /Script/%s"), *ModuleName.ToString(), *Plugin->Name, *Plugin->Name);
+				UE_LOG(LogPluginManager, Verbose, TEXT("Module %s belongs to Plugin %s and we assume that is the name of the package with the UObjects is /Script/%s"), *ModuleName.ToString(), *Plugin->Name, *Plugin->Name);
 				return FName(*Plugin->Name);
 			}
 		}
