PR #13893: Fix typo in TBaseStructure<FRemoteObjectReference>::Get so it actually works

#13893

diff --git a/Engine/Source/Runtime/CoreUObject/Private/UObject/Class.cpp b/Engine/Source/Runtime/CoreUObject/Private/UObject/Class.cpp
--- a/Engine/Source/Runtime/CoreUObject/Private/UObject/Class.cpp
+++ b/Engine/Source/Runtime/CoreUObject/Private/UObject/Class.cpp
@@ -8108,7 +8108,7 @@ UE_IMPLEMENT_STRUCT("/Script/CoreUObject", RemoteServerId);
 
 UScriptStruct* TBaseStructure<FRemoteObjectReference>::Get()
 {
-	static UScriptStruct* ScriptStruct = StaticGetBaseStructureInternal(TEXT("RemoteObjecReference"));
+	static UScriptStruct* ScriptStruct = StaticGetBaseStructureInternal(TEXT("RemoteObjectReference"));
 	return ScriptStruct;
 }
 UE_IMPLEMENT_STRUCT("/Script/CoreUObject", RemoteObjectReference);
