Fixed serializion error with instanced structs and remote object handles
- Fixed raw object pointer being serialized which triggered assertions in net code
- Issue requires UE_WITH_REMOTE_OBJECT_HANDLE=1 to reproduce
- Reproduces only when not using Iris for beacon communication (e.g. Lyra sample currently)

diff --git a/Engine/Source/Runtime/CoreUObject/Private/StructUtils/InstancedStruct.cpp b/Engine/Source/Runtime/CoreUObject/Private/StructUtils/InstancedStruct.cpp
--- a/Engine/Source/Runtime/CoreUObject/Private/StructUtils/InstancedStruct.cpp
+++ b/Engine/Source/Runtime/CoreUObject/Private/StructUtils/InstancedStruct.cpp
@@ -513,7 +513,7 @@ bool FInstancedStruct::NetSerialize(FArchive& Ar, UPackageMap* Map, bool& bOutSu
 	{
 		if (Ar.IsLoading())
 		{
-			UScriptStruct* SerializedScriptStruct = nullptr;
+			TObjectPtr<UScriptStruct> SerializedScriptStruct = nullptr;
 
 			Ar << SerializedScriptStruct;
 
