Fixing crashes on engine exit related to other systems trying to resolve weak object pointers after the UObject system has been shut down

diff --git a/Engine/Source/Runtime/CoreUObject/Private/UObject/RemoteObject.cpp b/Engine/Source/Runtime/CoreUObject/Private/UObject/RemoteObject.cpp
--- a/Engine/Source/Runtime/CoreUObject/Private/UObject/RemoteObject.cpp
+++ b/Engine/Source/Runtime/CoreUObject/Private/UObject/RemoteObject.cpp
@@ -442,7 +442,7 @@ FRemoteObjectPathName* FindAssetPath(FRemoteObjectId RemoteId)
 
 UE::RemoteObject::Handle::FRemoteObjectStub* FindRemoteObjectStub(FRemoteObjectId ObjectId)
 {
-	return ObjectMaps->FindRemoteObjectStub(ObjectId);
+	return ObjectMaps ? ObjectMaps->FindRemoteObjectStub(ObjectId) : nullptr;
 }
 
 FName GetServerBaseNameForUniqueName(const UClass* Class)
@@ -672,8 +672,9 @@ FPhysicsIslandId GetPhysicsIslandId(const UObject* Object)
 	using namespace UE::RemoteObject::Private;
 #if UE_WITH_REMOTE_OBJECT_HANDLE
 	return GetPhysicsIslandId(FObjectHandleUtils::GetRemoteId(Object));
-#endif
+#else
 	return FPhysicsIslandId();
+#endif
 }
 
 void ChangePhysicsIslandId(const UObject* Object, FPhysicsIslandId NewPhysicsIslandId)
@@ -811,6 +812,11 @@ bool CanResolveObject(FRemoteObjectId ObjectId, ECanResolveObjectBehavior CanRes
 		return true;
 	}
 
+	if (!ObjectMaps)
+	{
+		return false;
+	}
+
 	if (FRemoteObjectStub* Stub = ObjectMaps->FindRemoteObjectStub(ObjectId))
 	{
 		// A stub exists so the object memory is not local but we can (attempt to) migrate it back if:
