Automatically disable garbage elimination when remote object support is enabled

diff --git a/Engine/Source/Runtime/CoreUObject/Private/UObject/RemoteObject.cpp b/Engine/Source/Runtime/CoreUObject/Private/UObject/RemoteObject.cpp
--- a/Engine/Source/Runtime/CoreUObject/Private/UObject/RemoteObject.cpp
+++ b/Engine/Source/Runtime/CoreUObject/Private/UObject/RemoteObject.cpp
@@ -345,8 +345,12 @@ void InitRemoteObjects()
 	}
 
 #if UE_WITH_REMOTE_OBJECT_HANDLE
-	// We rely on garbage elimination begin disabled because we don't allow it inside of borrowed objects
-	checkf(UObjectBaseUtility::IsGarbageEliminationEnabled() == false, TEXT("Remote object support requires garbage elimination to be disabled"));
+	// We rely on garbage elimination being disabled because we don't allow it inside of borrowed objects
+	if (UObjectBaseUtility::IsGarbageEliminationEnabled())
+	{
+		UE_LOG(LogRemoteObject, Warning, TEXT("Disabling garbage elimination support because remote object support is enabled and these features are mutually exclusive"));
+		UObjectBaseUtility::SetGarbageEliminationEnabled(false);
+	}	
 #endif
 }
 
