[ZenLoader]
- Fix race on excluded classes when multithreaded loader is active

diff --git a/Engine/Source/Runtime/CoreUObject/Private/Serialization/AsyncLoading2.cpp b/Engine/Source/Runtime/CoreUObject/Private/Serialization/AsyncLoading2.cpp
--- a/Engine/Source/Runtime/CoreUObject/Private/Serialization/AsyncLoading2.cpp
+++ b/Engine/Source/Runtime/CoreUObject/Private/Serialization/AsyncLoading2.cpp
@@ -8393,43 +8393,45 @@ class FExcludedClassSet
 	TSet<FString>   ExcludedClassNames;
 	TSet<UClass*>   ExcludedClasses;
 	TSet<UClass*>   ExcludedNativeClasses;
+	FRWLock         Lock;
 	uint64          ClassesVersionNumber = 0;
 	FDelegateHandle ConfigDelegate;
 
-	void EnsureClassesUpdated()
+	void UpdateClasses()
 	{
+		TRACE_CPUPROFILER_EVENT_SCOPE(FExcludedClassSet::UpdateClasses);
+
+		FWriteScopeLock ScopeLock(Lock);
+
+		uint64 CurrentVersion = GetRegisteredClassesVersionNumber();
+
 		if (ExcludedClassNames.Num())
 		{
-			uint64 CurrentVersion = GetRegisteredClassesVersionNumber();
-			if (ClassesVersionNumber != CurrentVersion)
-			{
-				TRACE_CPUPROFILER_EVENT_SCOPE(FExcludedClassSet::EnsureClassesUpdated);
+			ExcludedClasses.Reset();
 
-				ClassesVersionNumber = CurrentVersion;
-				ExcludedClasses.Reset();
-
-				for (auto It = ExcludedClassNames.CreateIterator(); It; ++It)
+			for (auto It = ExcludedClassNames.CreateIterator(); It; ++It)
+			{
+				if (UClass* Class = FindObject<UClass>(nullptr, *It))
 				{
-					if (UClass* Class = FindObject<UClass>(nullptr, *It))
+					if (Class->IsNative())
 					{
-						if (Class->IsNative())
-						{
-							ExcludedNativeClasses.Add(Class);
+						ExcludedNativeClasses.Add(Class);
 
-							// We don't need to look at native classes ever again
-							It.RemoveCurrent();
-						}
-						else
-						{
-							ExcludedClasses.Add(Class);
-						}
+						// We don't need to look at native classes ever again
+						It.RemoveCurrent();
+					}
+					else
+					{
+						ExcludedClasses.Add(Class);
 					}
 				}
 			}
 		}
+
+		ClassesVersionNumber = CurrentVersion;
 	}
 
-	void Reset()
+	void ResetNoLock()
 	{
 		ExcludedClassNames.Reset();
 		ExcludedClasses.Reset();
@@ -8441,7 +8443,9 @@ class FExcludedClassSet
 	{
 		TRACE_CPUPROFILER_EVENT_SCOPE(FExcludedClassSet::OnConfigSectionChanged);
 
-		Reset();
+		FWriteScopeLock ScopeLock(Lock);
+
+		ResetNoLock();
 
 		if (GConfig && GConfig->IsReadyForUse())
 		{
@@ -8480,25 +8484,39 @@ public:
 
 	bool IsExcluded(UClass* Class)
 	{
-		EnsureClassesUpdated();
-
-		for (UClass* ExcludedClass : ExcludedClasses)
+		while (true)
 		{
-			if (Class->IsChildOf(ExcludedClass))
+			UE::TReadScopeLock ScopeLock(Lock);
+
+			uint64 CurrentVersion = GetRegisteredClassesVersionNumber();
+			if (ClassesVersionNumber != CurrentVersion)
 			{
-				return true;
+				// In case classes need to be updated, we have to release the read lock and proceed
+				ScopeLock.ReadUnlock();
+
+				UpdateClasses();
+
+				continue;
 			}
-		}
 
-		for (UClass* ExcludedClass : ExcludedNativeClasses)
-		{
-			if (Class->IsChildOf(ExcludedClass))
+			for (UClass* ExcludedClass : ExcludedClasses)
 			{
-				return true;
+				if (Class->IsChildOf(ExcludedClass))
+				{
+					return true;
+				}
+			}
+
+			for (UClass* ExcludedClass : ExcludedNativeClasses)
+			{
+				if (Class->IsChildOf(ExcludedClass))
+				{
+					return true;
+				}
 			}
-		}
 		
-		return false;
+			return false;
+		}
 	}
 };
 
@@ -12749,4 +12767,4 @@ bool FPublicExportMapPreInsertPublicExportsTest::RunTest(const FString& Paramete
 	return true;
 }
 
-#endif // WITH_LOW_LEVEL_TESTS
+#endif // WITH_LOW_LEVEL_TESTS
\ No newline at end of file
