Compile with immutable CDOs only when remote object handles AND AutoRTFM is enabled

diff --git a/Engine/Source/Runtime/CoreUObject/Private/UObject/Class.cpp b/Engine/Source/Runtime/CoreUObject/Private/UObject/Class.cpp
--- a/Engine/Source/Runtime/CoreUObject/Private/UObject/Class.cpp
+++ b/Engine/Source/Runtime/CoreUObject/Private/UObject/Class.cpp
@@ -4949,7 +4949,7 @@ void UClass::AddReferencedObjects(UObject* InThis, FReferenceCollector& Collecto
 		This->CallAddReferencedObjects(This->ClassDefaultObject, Collector);
 	}
 
-#if UE_WITH_REMOTE_OBJECT_HANDLE
+#if UE_WITH_IMMUTABLE_CDO
 	Collector.AddStableReference(&This->ImmutableDefaultObject);
 #endif
 
@@ -5697,7 +5697,7 @@ void UClass::InternalCreateDefaultObjectWrapper() const
 	const_cast<UClass*>(this)->CreateDefaultObject();
 }
 
-#if UE_WITH_REMOTE_OBJECT_HANDLE
+#if UE_WITH_IMMUTABLE_CDO
 void UClass::InternalCreateImmutableDefaultObject() const
 {
 	UObject* DefaultObject = GetDefaultObject();
@@ -5738,7 +5738,7 @@ void UClass::InternalCreateImmutableDefaultObject() const
 void UClass::SetDefaultObject(UObject* InClassDefaultObject)
 {
 PRAGMA_DISABLE_DEPRECATION_WARNINGS
-#if UE_WITH_REMOTE_OBJECT_HANDLE
+#if UE_WITH_IMMUTABLE_CDO
 	if (InClassDefaultObject != ClassDefaultObject && ImmutableDefaultObject)
 	{
 		ImmutableDefaultObject->ClearFlags(RF_ClassDefaultObject | RF_ImmutableDefaultObject | RF_ArchetypeObject);
@@ -5754,7 +5754,7 @@ PRAGMA_ENABLE_DEPRECATION_WARNINGS
 void UClass::PostLoadDefaultObject(UObject* Object)
 {
 	Object->PostLoad();
-#if UE_WITH_REMOTE_OBJECT_HANDLE
+#if UE_WITH_IMMUTABLE_CDO
 	// Construct the immutable version of the CDO immediately after the CDO has been loaded
 	GetImmutableDefaultObject();
 #endif
diff --git a/Engine/Source/Runtime/CoreUObject/Private/UObject/UObjectArchetype.cpp b/Engine/Source/Runtime/CoreUObject/Private/UObject/UObjectArchetype.cpp
--- a/Engine/Source/Runtime/CoreUObject/Private/UObject/UObjectArchetype.cpp
+++ b/Engine/Source/Runtime/CoreUObject/Private/UObject/UObjectArchetype.cpp
@@ -171,7 +171,7 @@ FORCENOINLINE UObject* FindArchetypeFromRequiredInfoImpl(const UClass* Class, co
 		if (!Result)
 		{
 			// nothing found, I am not a CDO, so this is just the class CDO
-#if UE_WITH_REMOTE_OBJECT_HANDLE
+#if UE_WITH_IMMUTABLE_CDO
 			if (bUseImmutableArchetype)
 			{
 				Result = const_cast<UObject*>(Class->GetImmutableDefaultObject());
diff --git a/Engine/Source/Runtime/CoreUObject/Private/UObject/UObjectBase.cpp b/Engine/Source/Runtime/CoreUObject/Private/UObject/UObjectBase.cpp
--- a/Engine/Source/Runtime/CoreUObject/Private/UObject/UObjectBase.cpp
+++ b/Engine/Source/Runtime/CoreUObject/Private/UObject/UObjectBase.cpp
@@ -937,7 +937,7 @@ static void UObjectLoadAllCompiledInDefaultProperties(TArray<UClass*>& OutAllNew
 			{
 				UE_LOG(LogUObjectBootstrap, Verbose, TEXT("GetDefaultObject Begin %s %s"), *Class->GetOutermost()->GetName(), *Class->GetName());
 				Class->GetDefaultObject();
-#if UE_WITH_REMOTE_OBJECT_HANDLE
+#if UE_WITH_IMMUTABLE_CDO
 				Class->GetImmutableDefaultObject();
 #endif
 				UE_LOG(LogUObjectBootstrap, Verbose, TEXT("GetDefaultObject End %s %s"), *Class->GetOutermost()->GetName(), *Class->GetName());
@@ -949,7 +949,7 @@ static void UObjectLoadAllCompiledInDefaultProperties(TArray<UClass*>& OutAllNew
 			{
 				UE_LOG(LogUObjectBootstrap, Verbose, TEXT("GetDefaultObject Begin %s %s"), *Class->GetOutermost()->GetName(), *Class->GetName());
 				Class->GetDefaultObject();
-#if UE_WITH_REMOTE_OBJECT_HANDLE
+#if UE_WITH_IMMUTABLE_CDO
 				Class->GetImmutableDefaultObject();
 #endif
 				UE_LOG(LogUObjectBootstrap, Verbose, TEXT("GetDefaultObject End %s %s"), *Class->GetOutermost()->GetName(), *Class->GetName());
@@ -961,7 +961,7 @@ static void UObjectLoadAllCompiledInDefaultProperties(TArray<UClass*>& OutAllNew
 			{
 				UE_LOG(LogUObjectBootstrap, Verbose, TEXT("GetDefaultObject Begin %s %s"), *Class->GetOutermost()->GetName(), *Class->GetName());
 				Class->GetDefaultObject();
-#if UE_WITH_REMOTE_OBJECT_HANDLE
+#if UE_WITH_IMMUTABLE_CDO
 				Class->GetImmutableDefaultObject();
 #endif
 				UE_LOG(LogUObjectBootstrap, Verbose, TEXT("GetDefaultObject End %s %s"), *Class->GetOutermost()->GetName(), *Class->GetName());
diff --git a/Engine/Source/Runtime/CoreUObject/Public/UObject/Class.h b/Engine/Source/Runtime/CoreUObject/Public/UObject/Class.h
--- a/Engine/Source/Runtime/CoreUObject/Public/UObject/Class.h
+++ b/Engine/Source/Runtime/CoreUObject/Public/UObject/Class.h
@@ -148,6 +148,9 @@ COREUOBJECT_API DECLARE_LOG_CATEGORY_EXTERN(LogScriptSerialization, Log, All);
 	typedef FRWLock FUClassFuncLock;
 #endif
 
+#ifndef UE_WITH_IMMUTABLE_CDO
+	#define UE_WITH_IMMUTABLE_CDO (UE_WITH_REMOTE_OBJECT_HANDLE && UE_AUTORTFM)
+#endif
 
 /*-----------------------------------------------------------------------------
 	FRepRecord.
@@ -3924,7 +3927,7 @@ public:
 	TObjectPtr<UObject> ClassDefaultObject;
 
 protected:
-#if UE_WITH_REMOTE_OBJECT_HANDLE
+#if UE_WITH_IMMUTABLE_CDO
 	TObjectPtr<UObject> ImmutableDefaultObject;
 #endif
 
@@ -4383,7 +4386,7 @@ public:
 	 * @param bCreateIfNeeded if true (default) then the CDO is created if it is null
 	 * @return the immutable CDO for this class
 	 */
-#if UE_WITH_REMOTE_OBJECT_HANDLE
+#if UE_WITH_IMMUTABLE_CDO
 	const UObject* GetImmutableDefaultObject(bool bCreateIfNeeded = true) const
 	{
 		if (ImmutableDefaultObject == nullptr && bCreateIfNeeded)
@@ -4734,7 +4737,7 @@ private:
 	 */
 	COREUOBJECT_API void AssembleReferenceTokenStreamInternal(bool bForce = false);
 	COREUOBJECT_API void InternalCreateDefaultObjectWrapper() const;
-#if UE_WITH_REMOTE_OBJECT_HANDLE
+#if UE_WITH_IMMUTABLE_CDO
 	COREUOBJECT_API void InternalCreateImmutableDefaultObject() const;
 #endif
 protected:
