Allow compilation without late resolve feature

#rnx

diff --git a/Engine/Source/Runtime/CoreUObject/Private/UObject/PropertyObject.cpp b/Engine/Source/Runtime/CoreUObject/Private/UObject/PropertyObject.cpp
--- a/Engine/Source/Runtime/CoreUObject/Private/UObject/PropertyObject.cpp
+++ b/Engine/Source/Runtime/CoreUObject/Private/UObject/PropertyObject.cpp
@@ -86,7 +86,7 @@ EConvertFromTypeResult FObjectProperty::ConvertFromType(const FPropertyTag& Tag,
 				{
 					UE_LOG(LogClass, Error, TEXT("Failed to convert soft path %s to unloaded object as this is not supported during async loading with the currently active loader '%s'. Load and resave %s in the editor to fix!"), *PreviousValue.ToString(), LexToString(LoaderType), *Slot.GetUnderlyingArchive().GetArchiveName());
 				}
-#if WITH_EDITOR
+#if WITH_EDITOR && UE_WITH_OBJECT_HANDLE_LATE_RESOLVE
 				// Only for the parallel loading path first so that we can do some opt-in testing
 				else if (IsInParallelLoadingThread())
 				{
