Pass LOAD_Quiet when loading widget compiler rules when expecting possible failure

diff --git a/Engine/Source/Editor/UMGEditor/Private/Settings/WidgetEditingProjectSettings.cpp b/Engine/Source/Editor/UMGEditor/Private/Settings/WidgetEditingProjectSettings.cpp
--- a/Engine/Source/Editor/UMGEditor/Private/Settings/WidgetEditingProjectSettings.cpp
+++ b/Engine/Source/Editor/UMGEditor/Private/Settings/WidgetEditingProjectSettings.cpp
@@ -58,7 +58,7 @@ TArray<UWidgetCompilerRule*> UWidgetEditingProjectSettings::CompilerOption_Rules
 		{
 			// The compiling rule may not be loaded yet in early loading phases, we'll
 			// just have to skip the rules in those cases.
-			RuleClassPtr.LoadSynchronous();
+			RuleClassPtr.LoadSynchronous(LOAD_Quiet);
 			if (RuleClassPtr)
 			{
 				if (UWidgetCompilerRule* Rule = RuleClassPtr->GetDefaultObject<UWidgetCompilerRule>())
diff --git a/Engine/Source/Runtime/CoreUObject/Private/UObject/SoftObjectPath.cpp b/Engine/Source/Runtime/CoreUObject/Private/UObject/SoftObjectPath.cpp
--- a/Engine/Source/Runtime/CoreUObject/Private/UObject/SoftObjectPath.cpp
+++ b/Engine/Source/Runtime/CoreUObject/Private/UObject/SoftObjectPath.cpp
@@ -778,7 +778,7 @@ bool FSoftObjectPath::RemapPackage(FName OldPackageName, FName NewPackageName)
 	return AssetPath.RemapPackage(OldPackageName, NewPackageName);
 }
 
-UObject* FSoftObjectPath::TryLoad(FUObjectSerializeContext* InLoadContext) const
+UObject* FSoftObjectPath::TryLoad(FUObjectSerializeContext* InLoadContext, ELoadFlags LoadFlags /* = LOAD_None */) const
 {
 	UObject* LoadedObject = nullptr;
 
@@ -815,7 +815,7 @@ UObject* FSoftObjectPath::TryLoad(FUObjectSerializeContext* InLoadContext) const
 			}
 #endif
 
-			LoadedObject = StaticLoadObject(UObject::StaticClass(), nullptr, *PathString, nullptr, LOAD_None, nullptr, true);
+			LoadedObject = StaticLoadObject(UObject::StaticClass(), nullptr, *PathString, nullptr, LoadFlags, nullptr, true);
 
 			// Look at core redirects if we didn't find the object
 			if (!LoadedObject && SoftObjectPath::bResolveCoreRedirects)
@@ -823,7 +823,7 @@ UObject* FSoftObjectPath::TryLoad(FUObjectSerializeContext* InLoadContext) const
 				FSoftObjectPath FixupObjectPath = *this;
 				if (FixupObjectPath.FixupCoreRedirects())
 				{
-					LoadedObject = LoadObject<UObject>(nullptr, *FixupObjectPath.ToString());
+					LoadedObject = LoadObject<UObject>(nullptr, *FixupObjectPath.ToString(), {}, LoadFlags);
 				}
 			}
 
diff --git a/Engine/Source/Runtime/CoreUObject/Public/UObject/SoftObjectPath.h b/Engine/Source/Runtime/CoreUObject/Public/UObject/SoftObjectPath.h
--- a/Engine/Source/Runtime/CoreUObject/Public/UObject/SoftObjectPath.h
+++ b/Engine/Source/Runtime/CoreUObject/Public/UObject/SoftObjectPath.h
@@ -310,9 +310,10 @@ public:
 	/**
 	 * Attempts to load the asset, this will call LoadObject which can be very slow
 	 * @param InLoadContext Optional load context when called from nested load callstack
+	 * @params LoadFlags Flags to forward to LoadObject e.g. LOAD_Quiet
 	 * @return Loaded UObject, or nullptr if the reference is null or the asset fails to load
 	 */
-	COREUOBJECT_API UObject* TryLoad(FUObjectSerializeContext* InLoadContext = nullptr) const;
+	COREUOBJECT_API UObject* TryLoad(FUObjectSerializeContext* InLoadContext = nullptr, ELoadFlags LoadFlags = LOAD_None) const;
 
 	/**
 	 * Attempts to asynchronously load the object referenced by this path.
diff --git a/Engine/Source/Runtime/CoreUObject/Public/UObject/SoftObjectPtr.h b/Engine/Source/Runtime/CoreUObject/Public/UObject/SoftObjectPtr.h
--- a/Engine/Source/Runtime/CoreUObject/Public/UObject/SoftObjectPtr.h
+++ b/Engine/Source/Runtime/CoreUObject/Public/UObject/SoftObjectPtr.h
@@ -79,12 +79,12 @@ public:
 	}
 
 	/** Synchronously load (if necessary) and return the asset object represented by this asset ptr */
-	UObject* LoadSynchronous() const
+	UObject* LoadSynchronous(ELoadFlags LoadFlags = LOAD_None) const
 	{
 		UObject* Asset = Get();
 		if (Asset == nullptr && !IsNull())
 		{
-			ToSoftObjectPath().TryLoad();
+			ToSoftObjectPath().TryLoad(nullptr, LoadFlags);
 			
 			// TryLoad will have loaded this pointer if it is valid
 			Asset = Get();
@@ -972,9 +972,9 @@ public:
 	}
 
 	/** Synchronously load (if necessary) and return the asset object represented by this asset ptr */
-	UClass* LoadSynchronous() const
+	UClass* LoadSynchronous(ELoadFlags LoadFlags = LOAD_None) const
 	{
-		UObject* Asset = SoftObjectPtr.LoadSynchronous();
+		UObject* Asset = SoftObjectPtr.LoadSynchronous(LoadFlags);
 		UClass* Class = dynamic_cast<UClass*>(Asset);
 		if (!Class || !Class->IsChildOf(TClass::StaticClass()))
 		{
