[UObjectGlobals]
- Fix delegate being copied instead of moved when using MakeShared

diff --git a/Engine/Source/Runtime/CoreUObject/Public/UObject/UObjectGlobals.h b/Engine/Source/Runtime/CoreUObject/Public/UObject/UObjectGlobals.h
--- a/Engine/Source/Runtime/CoreUObject/Public/UObject/UObjectGlobals.h
+++ b/Engine/Source/Runtime/CoreUObject/Public/UObject/UObjectGlobals.h
@@ -755,7 +755,7 @@ public:
 	using DelegateType = TTSDelegate<void(FLoadPackageAsyncProgressParams& Params)>;
 
 	FLoadPackageAsyncProgressDelegate(DelegateType&& InDelegate, uint32 InMask = MAX_uint32)
-		: Delegate(InDelegate)
+		: Delegate(Forward<DelegateType>(InDelegate))
 		, Mask(InMask)
 	{}
 
