[UObject]
- Avoid contention when trying to reuse UObject from a pool if that UObject is not owned by the loader
- This also fix a potential deadlock if a crash occurs while overwriting a UObject and a delegate needs to be called on a UObject from another thread

diff --git a/Engine/Source/Runtime/CoreUObject/Private/UObject/UObjectGlobals.cpp b/Engine/Source/Runtime/CoreUObject/Private/UObject/UObjectGlobals.cpp
--- a/Engine/Source/Runtime/CoreUObject/Private/UObject/UObjectGlobals.cpp
+++ b/Engine/Source/Runtime/CoreUObject/Private/UObject/UObjectGlobals.cpp
@@ -3765,7 +3765,11 @@ UObject* StaticAllocateObject
 			// It's not safe to access the object from other threads while it's being destroyed and reconstucted so simulate GC behavior and lock
 			if (ensureMsgf(IsInGameThread(), TEXT("GC lock can only be acquired on the game thread. If you hit this ensure then an object is being reconstructed on a worker thread which is not thread-safe")) && GCGuard)
 			{
-				GCGuard->Lock();
+				// Take the GC lock if the object is still being async loaded to avoid any race condition, otherwise we don't care about the lock
+				if (Obj->HasAnyInternalFlags(EInternalObjectFlags_AsyncLoading))
+				{
+					GCGuard->Lock();
+				}
 			}
 
 			// Destroy the object.
