Fix build issue 1086243

#rnx

diff --git a/Engine/Source/Runtime/CoreUObject/Private/StructUtils/PropertyBag.cpp b/Engine/Source/Runtime/CoreUObject/Private/StructUtils/PropertyBag.cpp
--- a/Engine/Source/Runtime/CoreUObject/Private/StructUtils/PropertyBag.cpp
+++ b/Engine/Source/Runtime/CoreUObject/Private/StructUtils/PropertyBag.cpp
@@ -1571,33 +1571,33 @@ namespace UE::StructUtils::Private
 		TArray<FPropertyBagPropertyDesc> NewDescs;
 		NewDescs.Reserve(Descs.Num());
 
-		const int32 Min = std::min(SourcePropertyIndex, TargetPropertyIndex);
+		const int32 Min = FMath::Min(SourcePropertyIndex, TargetPropertyIndex);
 		const int32 SourceSpacer = SourcePropertyIndex < TargetPropertyIndex;
 
 		// Elements before the source and target properties.
 		int32 ReadIndex;
 		for (ReadIndex = 0; ReadIndex < Min; ++ReadIndex)
 		{
-			NewDescs.Emplace(std::move(Descs[ReadIndex]));
+			NewDescs.Emplace(MoveTemp(Descs[ReadIndex]));
 		}
 
 		// Between the source and target properties (skip source if needed)
 		for (ReadIndex += SourceSpacer; ReadIndex < TargetPropertyIndex; ++ReadIndex)
 		{
-			NewDescs.Emplace(std::move(Descs[ReadIndex]));
+			NewDescs.Emplace(MoveTemp(Descs[ReadIndex]));
 		}
 
 		// Source and target property
-		NewDescs.Emplace(std::move(Descs[bInsertBefore ? SourcePropertyIndex : TargetPropertyIndex]));
-		NewDescs.Emplace(std::move(Descs[bInsertBefore ? TargetPropertyIndex : SourcePropertyIndex]));
+		NewDescs.Emplace(MoveTemp(Descs[bInsertBefore ? SourcePropertyIndex : TargetPropertyIndex]));
+		NewDescs.Emplace(MoveTemp(Descs[bInsertBefore ? TargetPropertyIndex : SourcePropertyIndex]));
 
 		// The remaining elements
 		for (++ReadIndex; ReadIndex < Descs.Num(); ++ReadIndex)
 		{
-			NewDescs.Emplace(std::move(Descs[ReadIndex]));
+			NewDescs.Emplace(MoveTemp(Descs[ReadIndex]));
 		}
 
-		Descs = std::move(NewDescs);
+		Descs = MoveTemp(NewDescs);
 
 		return EPropertyBagAlterationResult::Success;
 	}
