Avoid crashing when applying an instance data transform to a type with a nested struct declaration, these are rare but do exist (e.g. for animation blueprints, which generate a nested struct to describe some constant data)

#rnx

diff --git a/Engine/Source/Runtime/CoreUObject/Private/UObject/InstanceDataTransforms.cpp b/Engine/Source/Runtime/CoreUObject/Private/UObject/InstanceDataTransforms.cpp
--- a/Engine/Source/Runtime/CoreUObject/Private/UObject/InstanceDataTransforms.cpp
+++ b/Engine/Source/Runtime/CoreUObject/Private/UObject/InstanceDataTransforms.cpp
@@ -102,6 +102,12 @@ namespace UE::Private
 
 		const TArray<TSharedPtr<FInstanceDataTransform>>* FindTransforms(const TNotNull<UStruct*> Struct)
 		{
+			if(Struct->GetOuter() != Struct->GetPackage())
+			{
+				// can't support transforms for nested structs like 
+				// anim blueprint's nested 'sparse class data':
+				return nullptr;
+			}
 			return FindTransforms(Struct->GetStructPathName());
 		}
 	private:
