Replaced a check with a handled ensure for multi-server commits with remote objects that have become unreachable

diff --git a/Engine/Source/Runtime/CoreUObject/Private/UObject/RemoteObjectTransfer.cpp b/Engine/Source/Runtime/CoreUObject/Private/UObject/RemoteObjectTransfer.cpp
--- a/Engine/Source/Runtime/CoreUObject/Private/UObject/RemoteObjectTransfer.cpp
+++ b/Engine/Source/Runtime/CoreUObject/Private/UObject/RemoteObjectTransfer.cpp
@@ -488,7 +488,10 @@ public:
 			check(!Handle::IsRemote(UsedObjectId));
 
 			UObject* UsedObject = StaticFindObjectFastInternal(UsedObjectId);
-			check(UsedObject);
+			if (!ensure(UsedObject))
+			{
+				continue;
+			}
 			FRemoteServerId OwnerServerId = Handle::GetOwnerServerId(UsedObject);
 
 			UObject* RootObject = FindCanonicalRootObjectForSerialization(UsedObject);
