Mark remote reference collector archives as 'Saving' to workaround a crash in the structured archive code in the editor builds

diff --git a/Engine/Source/Runtime/CoreUObject/Private/UObject/RemoteObjectSerialization.cpp b/Engine/Source/Runtime/CoreUObject/Private/UObject/RemoteObjectSerialization.cpp
--- a/Engine/Source/Runtime/CoreUObject/Private/UObject/RemoteObjectSerialization.cpp
+++ b/Engine/Source/Runtime/CoreUObject/Private/UObject/RemoteObjectSerialization.cpp
@@ -911,6 +911,9 @@ public:
 	FArchiveRemoteReferencePatcher(const TArray<UObject*>& DeserializedObjects)
 	{
 		ArIsObjectReferenceCollector = true;
+#if WITH_EDITORONLY_DATA
+		SetIsSaving(true); // workaround for UE-225585 (UStruct::SerializeBinEx() asserts in structured archive code in the editor)
+#endif
 		SetIsPersistent(false);
 		InitObjectMap(DeserializedObjects);
 		SetPortFlags(PPF_AvoidRemoteObjectMigration);
@@ -942,6 +945,9 @@ public:
 	{
 		ArchetypeToObjectMap.Add(FindArchetype(RootObject), RootObject);
 		ArIsObjectReferenceCollector = true;
+#if WITH_EDITORONLY_DATA
+		SetIsSaving(true); // workaround for UE-225585 (UStruct::SerializeBinEx() asserts in structured archive code in the editor)
+#endif
 		SetIsPersistent(false);
 		SetPortFlags(PPF_AvoidRemoteObjectMigration);
 	}
