[ZenLoader]
- Fix race in multithreaded loader when trying to update the priority while the package is still being created on another thread

diff --git a/Engine/Source/Runtime/CoreUObject/Private/Serialization/AsyncLoading2.cpp b/Engine/Source/Runtime/CoreUObject/Private/Serialization/AsyncLoading2.cpp
--- a/Engine/Source/Runtime/CoreUObject/Private/Serialization/AsyncLoading2.cpp
+++ b/Engine/Source/Runtime/CoreUObject/Private/Serialization/AsyncLoading2.cpp
@@ -6240,10 +6240,17 @@ void FAsyncLoadEventQueue2::UpdatePackagePriority(FAsyncPackage2* Package)
 	{
 		ReprioritizeNode(Node);
 	}
-	for (FEventLoadNode2& Node : Package->Data.ExportBundleNodes)
+
+	// Only try to reprioritize export bundle nodes if they have been initialized
+	// otherwise they can't even be scheduled yet and will be prioritized correctly later.
+	if (Package->AsyncPackageLoadingState > EAsyncPackageLoadingState2::ProcessPackageSummary)
 	{
-		ReprioritizeNode(Node);
+		for (FEventLoadNode2& Node : Package->Data.ExportBundleNodes)
+		{
+			ReprioritizeNode(Node);
+		}
 	}
+
 	UpdateExternalQueueState();
 }
 
