DSTM Speculation Mode (Scoped):  Continue executing when encountering a remote object (instead of aborting), then do a late-abort when exiting the scope.

#rnx

diff --git a/Engine/Source/Runtime/CoreUObject/Private/UObject/PropertyArray.cpp b/Engine/Source/Runtime/CoreUObject/Private/UObject/PropertyArray.cpp
--- a/Engine/Source/Runtime/CoreUObject/Private/UObject/PropertyArray.cpp
+++ b/Engine/Source/Runtime/CoreUObject/Private/UObject/PropertyArray.cpp
@@ -11,6 +11,10 @@
 #include "UObject/UnrealTypePrivate.h"
 #include "UObject/UObjectThreadContext.h"
 
+#if UE_WITH_REMOTE_OBJECT_HANDLE
+#include "UObject/RemoteExecutor.h"
+#endif
+
 /*-----------------------------------------------------------------------------
 	FArrayProperty.
 -----------------------------------------------------------------------------*/
@@ -1166,6 +1170,13 @@ void FArrayProperty::CopyValuesInternal(void* Dest, void const* Src, int32 Count
 			// as that would all catch classes that inherited from `FObjectProperty`!
 			if (AutoRTFM::IsClosed() && (Inner->GetClass() == FObjectProperty::StaticClass()))
 			{
+				#if UE_WITH_REMOTE_OBJECT_HANDLE
+				// CopyCompleteValue will resolve handles in the RawPtr <- TObjectPtr case.
+				// Since we can't abort while in an open scope, use a speculation scope to defer any
+				// migrations until this scope ends.
+				UE::RemoteExecutor::FSpeculationExecutionScope AvoidMigrationsInTheOpen;
+				#endif
+
 				// Since we've already written to each location that we are going to be copying over below, we know that
 				// all of the record-writes that the transactional system would have inserted are no-ops. Also, each copy
 				// of a `TObjectPtr` will go into the open to do mark the object as reachable, so by doing a single open
diff --git a/Engine/Source/Runtime/CoreUObject/Private/UObject/RemoteExecutor.cpp b/Engine/Source/Runtime/CoreUObject/Private/UObject/RemoteExecutor.cpp
--- a/Engine/Source/Runtime/CoreUObject/Private/UObject/RemoteExecutor.cpp
+++ b/Engine/Source/Runtime/CoreUObject/Private/UObject/RemoteExecutor.cpp
@@ -18,6 +18,12 @@ static FAutoConsoleVariableRef CVarRemoteExecutorNewNetPumpBehavior(
 	GRemoteExecutorNewNetPumpBehavior,
 	TEXT("Enable the 'new' net pump behavior"));
 
+bool GRemoteExecutorAllowSpeculationMode = true;
+static FAutoConsoleVariableRef CVarRemoteExecutorAllowSpeculationMode(
+	TEXT("remoteexecutor.AllowSpeculationMode"),
+	GRemoteExecutorAllowSpeculationMode,
+	TEXT("Allow speculation mode (late abort) when transacting"));
+
 namespace UE::RemoteExecutor
 {
 TDelegate<void()> TickNetworkDelegate;
@@ -39,13 +45,6 @@ TDelegate<void(FName /*Text*/)> OnRegionBeginDelegate;
 TDelegate<void(const FString& /*Text*/)> OnRegionEndDelegate;
 }
 
-enum class ERemoteExecutorAbortReason : uint8
-{
-	Unspecified,
-	RequiresDependencies,
-	AbandonWork
-};
-
 FArchive& operator<<(FArchive& Ar, FRemoteTransactionId& Id)
 {
 	Ar << Id.Id;
@@ -86,6 +85,9 @@ struct FRemoteExecutorWork
 	bool bRequestCreated = false;
 	
 	uint32 ExecutionAttempts = 0;
+
+	// If non-zero, then the current work will abort rather than commit
+	uint32 NumAbortsRequested = 0;
 	
 	bool bRequiresMultiServerCommit = false;
 	FString RequiresMultiServerCommitReason;
@@ -114,7 +116,10 @@ class FRemoteExecutor
 public:
 	FRemoteExecutorWork* ExecutingWork = nullptr;
 
-	ERemoteExecutorAbortReason AbortReason = ERemoteExecutorAbortReason::Unspecified;
+	// Pump the network next chance we get
+	bool bShouldPumpNetwork = false;
+
+	UE::RemoteExecutor::ERemoteExecutorAbortReason AbortReason = UE::RemoteExecutor::ERemoteExecutorAbortReason::Unspecified;
 	FString AbortReasonDescription;
 
 	// tracking data for servicing a remote multi-server commit
@@ -210,6 +215,30 @@ public:
 				}
 			};
 
+		auto ConditionallyPumpNetwork = [this](bool bForcePumpNetwork)->void
+			{
+				if (GRemoteExecutorNewNetPumpBehavior)
+				{
+					// Always pump if we're trying this transaction multiple times
+					bool bUnconditionallyPumpNetwork = bForcePumpNetwork || bShouldPumpNetwork;
+			
+					// if GRemoteExecutorForcedNetPumpIterationCount is nonzero, we may need to unconditionally pump the network
+					if (GRemoteExecutorForcedNetPumpIterationCount && !bUnconditionallyPumpNetwork)
+					{
+						IterationsSinceNetworkPump++;
+						bUnconditionallyPumpNetwork = (IterationsSinceNetworkPump % GRemoteExecutorForcedNetPumpIterationCount) == 0;
+					}
+
+					if (bUnconditionallyPumpNetwork)
+					{
+						// pump the network, reset our state variables
+						IterationsSinceNetworkPump = 0;
+						bShouldPumpNetwork = false;
+						UE::RemoteExecutor::TickNetworkDelegate.ExecuteIfBound();
+					}
+				}
+			};
+
 		double LastStallPrintTime = FPlatformTime::Seconds();
 
 		for (int32 LocalIterationNumber = 0; ; LocalIterationNumber++)
@@ -292,24 +321,10 @@ public:
 				break;
 			}
 
-			if (GRemoteExecutorNewNetPumpBehavior)
+			// Maybe pump the network
 			{
-				IterationsSinceNetworkPump++;
-
-				// if GRemoteExecutorForcedNetPumpIterationCount is nonzero, we may need to unconditionally pump the network
-				bool bUnconditionallyPumpNetwork = false;
-			
-				if (GRemoteExecutorForcedNetPumpIterationCount)
-				{
-					bUnconditionallyPumpNetwork = (IterationsSinceNetworkPump % GRemoteExecutorForcedNetPumpIterationCount) == 0;
-				}
-
-				if ((LocalIterationNumber > 0) || bUnconditionallyPumpNetwork)
-				{
-					// pump the network
-					IterationsSinceNetworkPump = 0;
-					UE::RemoteExecutor::TickNetworkDelegate.ExecuteIfBound();
-				}
+				const bool bForcePumpNetwork = (LocalIterationNumber > 0);
+				ConditionallyPumpNetwork(bForcePumpNetwork);
 			}
 			
 			for (int32 PendingWorkIndex = 0; PendingWorkIndex < PendingWorks.Num(); PendingWorkIndex++)
@@ -416,7 +431,7 @@ public:
 					}
 
 					// all of the subsystems are ready, try to perform the work
-					AbortReason = ERemoteExecutorAbortReason::Unspecified;
+					AbortReason = UE::RemoteExecutor::ERemoteExecutorAbortReason::Unspecified;
 					AbortReasonDescription = TEXT("");
 
 					ExecutingWork->ExecutionAttempts++;
@@ -443,8 +458,16 @@ public:
 									}
 								};
 
+								ExecutingWork->NumAbortsRequested = 0;
 								ExecutingWork->Work();
 
+								// Did we do something that should late-abort this transaction?
+								if (ExecutingWork->NumAbortsRequested > 0)
+								{
+									UE_LOG(LogRemoteExec, Warning, TEXT("We ended a transaction marked for abort. We expected to abort before this point (%d abort requests)."), ExecutingWork->NumAbortsRequested);
+									AutoRTFM::AbortTransaction();
+								}
+
 								if (ExecutingWork->bRequiresMultiServerCommit)
 								{
 									TRACE_CPUPROFILER_EVENT_SCOPE(MultiServerCommit)
@@ -581,8 +604,10 @@ public:
 													MultiServerCommitReadyServers.Reset();
 													bMultiServerCommitRequiresAbort = false;
 												};
+												
+												using namespace UE::RemoteExecutor;
+												AbortCurrentTransaction(EAbortTransactionBehavior::Immediately, ERemoteExecutorAbortReason::RequiresDependencies, TEXT("bMultiServerCommitRequiresAbort"));
 
-												UE::RemoteExecutor::AbortTransactionRequiresDependencies(TEXT("bMultiServerCommitRequiresAbort"));
 												return;
 											}
 											else
@@ -632,9 +657,11 @@ public:
 
 					if (WorkTransactResult == AutoRTFM::ETransactionResult::AbortedByRequest)
 					{
+						using namespace UE::RemoteExecutor;
+
 						bWorkTransactionAborted = true;
 
-						UE::RemoteExecutor::OnTransactionAbortedDelegate.Broadcast(RequestId, ExecutingWork->ExecutionAttempts, AbortReasonDescription);
+						OnTransactionAbortedDelegate.Broadcast(RequestId, ExecutingWork->ExecutionAttempts, AbortReasonDescription);
 
 						if (AbortReason == ERemoteExecutorAbortReason::RequiresDependencies) // -V547
 						{
@@ -670,7 +697,7 @@ public:
 					}
 					else
 					{
-						check(AbortReason == ERemoteExecutorAbortReason::Unspecified);
+						check(AbortReason == UE::RemoteExecutor::ERemoteExecutorAbortReason::Unspecified);
 
 						UE_LOG(LogRemoteExec, VeryVerbose, TEXT("ExecutePendingWork[%d,%s] : Work COMPLETED after %u attempts"),
 							LocalIterationNumber,
@@ -714,12 +741,16 @@ public:
 				}
 				else
 				{
+					UE::RemoteExecutor::OnRegionBeginDelegate.ExecuteIfBound(ExecutingWork->Name);
+
 					// execute non-transactional work
 					ExecutingWork->Work();
 
 					// remove this work from the list and do not increment PendingWorkIndex
 					EnqueueNextWorkFromList(ExecutingWork->WorkListIndex);
 					PendingWorks.RemoveAt(PendingWorkIndex, EAllowShrinking::No);
+
+					UE::RemoteExecutor::OnRegionEndDelegate.ExecuteIfBound(TEXT(""));
 				}
 
 				// it's possible this work ended with trying to perform a multi-server
@@ -730,6 +761,9 @@ public:
 				{
 					break;
 				}
+
+				// Finally, see if we have any requests waiting that need servicing
+				ConditionallyPumpNetwork(false);
 			}
 		}
 
@@ -753,22 +787,88 @@ void ExecuteTransactionalInternal(FName WorkName, FRemoteWorkPriority InWorkPrio
 
 namespace UE::RemoteExecutor
 {
+int32 SpeculationModeScopeCount = 0;
+
+FSpeculationExecutionScope::FSpeculationExecutionScope()
+{
+	if (GRemoteExecutorAllowSpeculationMode)
+	{
+		++SpeculationModeScopeCount;
+	}
+}
+
+FSpeculationExecutionScope::~FSpeculationExecutionScope()
+{
+	if (GRemoteExecutorAllowSpeculationMode)
+	{
+		--SpeculationModeScopeCount;
+		check(SpeculationModeScopeCount >= 0);
+	}
+	else
+	{
+		SpeculationModeScopeCount = 0;
+	}
+
+#if UE_WITH_REMOTE_OBJECT_HANDLE
+	if (SpeculationModeScopeCount == 0)
+	{
+		if (GRemoteExecutor.ExecutingWork && GRemoteExecutor.ExecutingWork->NumAbortsRequested)
+		{
+			AutoRTFM::AbortTransaction();
+		}
+	}
+#endif
+}
+
+bool IsSpeculationMode()
+{
+	check(SpeculationModeScopeCount >= 0);
+	return GRemoteExecutorAllowSpeculationMode && (SpeculationModeScopeCount > 0);
+}
+
 void RegisterRemoteSubsystem(FRemoteSubsystemBase* Subsystem)
 {
 	GRemoteExecutor.RegisterSubsystem(Subsystem);
 }
 
-void AbortTransactionRequiresDependencies(FStringView Description)
+void NotifyNetworkIsWaiting()
+{
+	GRemoteExecutor.bShouldPumpNetwork = true;
+}
+
+void AbortCurrentTransaction(EAbortTransactionBehavior AbortBehavior, ERemoteExecutorAbortReason AbortReason, FStringView AbortDescription)
 {
-	check(AutoRTFM::IsClosed());
+	check(AutoRTFM::IsTransactional());
 
 	UE_AUTORTFM_OPEN
 	{
-		GRemoteExecutor.AbortReason = ERemoteExecutorAbortReason::RequiresDependencies;
-		GRemoteExecutor.AbortReasonDescription = Description;
+		GRemoteExecutor.AbortReason = AbortReason;
+		GRemoteExecutor.AbortReasonDescription = AbortDescription;
+
+		if (FRemoteExecutorWork* CurrentWorkItem = GRemoteExecutor.ExecutingWork)
+		{
+			CurrentWorkItem->NumAbortsRequested++;
+		}
 	};
 
-	AutoRTFM::AbortTransaction();
+	if (AutoRTFM::IsClosed())
+	{
+		if (AbortBehavior == EAbortTransactionBehavior::Immediately)
+		{
+			AutoRTFM::AbortTransaction();
+		}
+	}
+}
+
+void AbortTransactionRequiresDependencies(FStringView Description)
+{
+	EAbortTransactionBehavior AbortBehavior = IsSpeculationMode() ? EAbortTransactionBehavior::BeforeCommit : EAbortTransactionBehavior::Immediately;
+	AbortCurrentTransaction(AbortBehavior, ERemoteExecutorAbortReason::RequiresDependencies, Description);
+}
+
+void AbortTransactionAndAbandonWork(FStringView Description)
+{
+	AbortCurrentTransaction(EAbortTransactionBehavior::Immediately, ERemoteExecutorAbortReason::AbandonWork, Description);
 }
 
 void RollbackTransactionRequiresDependencies(FStringView Description)
@@ -783,19 +883,6 @@ void RollbackTransactionRequiresDependencies(FStringView Description)
 #endif
 }
 
-void AbortTransactionAndAbandonWork(FStringView Description)
-{
-	check(AutoRTFM::IsClosed());
-
-	UE_AUTORTFM_OPEN
-	{
-		GRemoteExecutor.AbortReason = ERemoteExecutorAbortReason::AbandonWork;
-		GRemoteExecutor.AbortReasonDescription = Description;
-	};
-
-	AutoRTFM::AbortTransaction();
-}
-
 void TransactionRequiresMultiServerCommit(FStringView Description)
 {
 	if (GRemoteExecutor.ExecutingWork && !GRemoteExecutor.ExecutingWork->bRequiresMultiServerCommit)
diff --git a/Engine/Source/Runtime/CoreUObject/Private/UObject/RemoteObject.cpp b/Engine/Source/Runtime/CoreUObject/Private/UObject/RemoteObject.cpp
--- a/Engine/Source/Runtime/CoreUObject/Private/UObject/RemoteObject.cpp
+++ b/Engine/Source/Runtime/CoreUObject/Private/UObject/RemoteObject.cpp
@@ -730,6 +730,7 @@ void UpdateAllPhysicsLocalIslandId(const TMap<uint32, uint32>& PhysicsLocalIslan
 
 UObject* ResolveObject(const FRemoteObjectStub* Stub, ERemoteReferenceType RefType /*= ERemoteReferenceType::Strong*/)
 {
+	using namespace UE::RemoteExecutor;
 	using namespace UE::RemoteObject::Transfer;
 	using namespace UE::RemoteObject::Private;
 
@@ -759,16 +760,15 @@ UObject* ResolveObject(const FRemoteObjectStub* Stub, ERemoteReferenceType RefTy
 	{
 		checkf(!GIsGarbageCollecting, TEXT("Resolving remote objects while collecting garbage is not allowed (trying to resolve object %s (%s)"), *Stub->Id.ToString(), *Stub->Name.ToString());
 
-		UObject* Outer = Object ? Object->GetOuter() : StaticFindObjectFastInternal(Stub->OuterId);
+		MigrateObjectFromRemoteServer(Stub->Id, Stub->ResidentServerId);
 
-		MigrateObjectFromRemoteServer(Stub->Id, Stub->ResidentServerId, Outer);
-
-		// if running transactionally, we will have aborted and not reached here
+		// if running transactionally (and not speculating), we will have aborted and not reached here
+		// if running transactionally (and speculating), we will use a stale object to speculate and abort later.
 
 		// if running non-transactionally, object migrated immediately, so we just re-resolve
 		Object = StaticFindObjectFastInternal(Stub->Id);
 		bRemoteObject = !Object || IsRemote(Object);
-		checkf(!bRemoteObject, TEXT("Failed to resolve remote object %s, either this code is not running in a transaction and should be, or the transaction failed to abort"), *Stub->Id.ToString());
+		checkf(IsSpeculationMode() || !bRemoteObject, TEXT("Failed to resolve remote object %s, either this code is not running in a transaction and should be, or the transaction failed to abort"), *Stub->Id.ToString());
 	}
 
 	return Object;
@@ -830,8 +830,18 @@ bool CanResolveObject(FRemoteObjectId ObjectId, ECanResolveObjectBehavior CanRes
 			return false;
 		}
 
-		// 2. We're not explicitly saying it's unsafe to migrate objects (FUnsafeToMigrateScope)
-		// 3. During Garbage Collection. Note: GIsGarbageCollecting checks are not thread safe
+		// 2. We're not speculating (in that case Get() can return nullptr, so make sure this is false)
+		if (UE::RemoteExecutor::IsSpeculationMode())
+		{
+			// We also need to request the object here, or code that branches on IsValid would be wrong.
+			using namespace UE::RemoteObject::Transfer;
+			MigrateObjectFromRemoteServer(Stub->Id, Stub->ResidentServerId);
+
+			return false;
+		}
+
+		// 3. We're not explicitly saying it's unsafe to migrate objects (FUnsafeToMigrateScope)
+		// 4. During Garbage Collection. Note: GIsGarbageCollecting checks are not thread safe
 		return IsSafeToMigrateObjects();
 	}
 
diff --git a/Engine/Source/Runtime/CoreUObject/Private/UObject/RemoteObjectTransfer.cpp b/Engine/Source/Runtime/CoreUObject/Private/UObject/RemoteObjectTransfer.cpp
--- a/Engine/Source/Runtime/CoreUObject/Private/UObject/RemoteObjectTransfer.cpp
+++ b/Engine/Source/Runtime/CoreUObject/Private/UObject/RemoteObjectTransfer.cpp
@@ -736,6 +736,9 @@ public:
 				MigrateObjectToRemoteServer(ObjectId, DestinationServerId);
 			}
 		}
+
+		// Tell the executor that we have RPCs waiting to send
+		UE::RemoteExecutor::NotifyNetworkIsWaiting();
 	}
 
 	/**
@@ -961,7 +964,7 @@ UObject* MigrateObjectFromObjectData(FRemoteObjectData& ObjectData, const FUObje
 	return ReceivedObjects[RequestedObjectIndex];
 } 
 
-void MigrateObjectFromRemoteServer(FRemoteObjectId ObjectId, FRemoteServerId CurrentOwnerServerId, UObject* DestinationOuter)
+void MigrateObjectFromRemoteServer(FRemoteObjectId ObjectId, FRemoteServerId CurrentOwnerServerId)
 {
 	UE_LOG(LogRemoteObject, Verbose, TEXT("Aborting transaction, server %s needs obj id %s from server %s"),
 		*FRemoteServerId::GetLocalServerId().ToString(),
diff --git a/Engine/Source/Runtime/CoreUObject/Public/UObject/RemoteExecutor.h b/Engine/Source/Runtime/CoreUObject/Public/UObject/RemoteExecutor.h
--- a/Engine/Source/Runtime/CoreUObject/Public/UObject/RemoteExecutor.h
+++ b/Engine/Source/Runtime/CoreUObject/Public/UObject/RemoteExecutor.h
@@ -192,6 +192,38 @@ public:
 
 namespace UE::RemoteExecutor
 {
+	enum class EAbortTransactionBehavior : uint8
+	{
+		// We abort the transaction immediately (we do not finish the current callstack)
+		Immediately,
+
+		// We abort the transaction instead of committing it (we continue to execute speculatively)
+		BeforeCommit
+	};
+
+	enum class ERemoteExecutorAbortReason : uint8
+	{
+		Unspecified,
+		RequiresDependencies,
+		AbandonWork
+	};
+
+	/**
+	 * Denote that this scope runs in speculation mode.
+	 * Speculation mode allows us to return stale objects, or null when in a transaction & dereferencing a remote object.
+	 * Execution will continue (rather than aborting), but cannot commit in this mode (an attempt to commit will instead abort).
+	 * This allows us to continue to speculate about what a transaction might do, request those results, and abort only once.
+	 * 
+	 * Note: Due to the possibility of returning null, a scope must protect against null even if the object
+	 * actually exists on another server.
+	 */
+	struct FSpeculationExecutionScope
+	{
+		COREUOBJECT_API FSpeculationExecutionScope();
+		COREUOBJECT_API ~FSpeculationExecutionScope();
+	};
+
+
 	/** Delegate that pumps the network and executes RPC handlers */
 	extern COREUOBJECT_API TDelegate<void()> TickNetworkDelegate;
 	extern COREUOBJECT_API TDelegate<void(FRemoteTransactionId, FRemoteWorkPriority, const TArray<FRemoteServerId>&)> BeginMultiServerCommitDelegate;
@@ -222,8 +254,12 @@ namespace UE::RemoteExecutor
 	extern COREUOBJECT_API TDelegate<void(FName /*Text*/)> OnRegionBeginDelegate;
 	extern COREUOBJECT_API TDelegate<void(const FString& /*Text*/)> OnRegionEndDelegate;
 
+	COREUOBJECT_API bool IsSpeculationMode();
 	COREUOBJECT_API void RegisterRemoteSubsystem(FRemoteSubsystemBase* Subsystem);
 
+	/** Notify the executor that the network is waiting (we should tick it next chance we get) */
+	COREUOBJECT_API void NotifyNetworkIsWaiting();
+
 	COREUOBJECT_API void ExecuteTransactional(FName WorkName, const TFunctionRef<void(void)>& Work);
 	COREUOBJECT_API void ExecuteTransactionalWithExplicitPriority(FName WorkName, FRemoteWorkPriority WorkPriority, const TFunctionRef<void(void)>& Work);
 	COREUOBJECT_API void EnqueueWorkWithExplicitPriority(FName WorkName, FRemoteWorkPriority WorkPriority, bool bIsTransactional, const TFunction<void(void)>& Work);
@@ -232,6 +268,8 @@ namespace UE::RemoteExecutor
 	COREUOBJECT_API void EnqueueWorkList(FName WorkName, bool bIsTransactional, const TArray<TFunction<void(void)>>& WorkList);
 	COREUOBJECT_API void ExecutePendingWork();
 
+	COREUOBJECT_API void AbortCurrentTransaction(EAbortTransactionBehavior AbortBehavior, ERemoteExecutorAbortReason AbortReason, FStringView AbortDescription);
+
 	COREUOBJECT_API void AbortTransactionRequiresDependencies(FStringView Description);
 	COREUOBJECT_API void RollbackTransactionRequiresDependencies(FStringView Description);
 	COREUOBJECT_API void AbortTransactionAndAbandonWork(FStringView Description);
diff --git a/Engine/Source/Runtime/CoreUObject/Public/UObject/RemoteObjectTransfer.h b/Engine/Source/Runtime/CoreUObject/Public/UObject/RemoteObjectTransfer.h
--- a/Engine/Source/Runtime/CoreUObject/Public/UObject/RemoteObjectTransfer.h
+++ b/Engine/Source/Runtime/CoreUObject/Public/UObject/RemoteObjectTransfer.h
@@ -89,10 +89,9 @@ namespace UE::RemoteObject::Transfer
 	* Migrates an object from a remote server (temp function)
 	* @param ObjectId Remote object id
 	* @param CurrentOwnerServerId Id of a remote server that currently owns the object
-	* @param DestinationOuter Optional outer the migrated object should be reparented to
 	* @return Migrated object
 	*/
-	COREUOBJECT_API void MigrateObjectFromRemoteServer(FRemoteObjectId ObjectId, FRemoteServerId CurrentOwnerServerId, UObject* DestinationOuter = nullptr);
+	COREUOBJECT_API void MigrateObjectFromRemoteServer(FRemoteObjectId ObjectId, FRemoteServerId CurrentOwnerServerId);
 
 	/**
 	* Reports code that touches a resident object
@@ -170,7 +169,7 @@ namespace UE::RemoteObject::Transfer
 	/** Delegate that restores object data from a database. */
 	extern COREUOBJECT_API TDelegate<void(const FUObjectMigrationContext&)> RestoreRemoteObjectDataDelegate;
 
-} // namespace UE::RemoveObject::Transfer
+} // namespace UE::RemoteObject::Transfer
 
 struct FRemoteObjectReference
 {
diff --git a/Engine/Source/Runtime/Engine/Private/DataReplication.cpp b/Engine/Source/Runtime/Engine/Private/DataReplication.cpp
--- a/Engine/Source/Runtime/Engine/Private/DataReplication.cpp
+++ b/Engine/Source/Runtime/Engine/Private/DataReplication.cpp
@@ -89,6 +89,10 @@ static TAutoConsoleVariable<int32> CVarQueueRpcsForTransactions(
 	1,
 	TEXT("When set to 1, any transactional connection will queue RPCs and execute them in transactions allowing for rescheduling"));
 
+static TAutoConsoleVariable<bool> CVarEnqueueTransactionalRpcs(
+	TEXT("net.EnqueueTransactionalRpcs"),
+	true,
+	TEXT("When set to true (and net.QueueRpcsForTransactions to 1), each individual RPC will be enqueued rather than immediately executed (allows further unblocking of transactions)"));
 
 extern int32 GNumSkippedObjectEmptyUpdates;
 
@@ -2678,15 +2682,18 @@ void FObjectReplicator::EnqueuePendingTransactionalRPCs()
 						UE::RemoteExecutor::OnRegionEndDelegate.ExecuteIfBound(TEXT("Abort"));
 					});
 
-					FNetBitReader Reader(Connection->PackageMap, Pending.Buffer.GetData(), Pending.NumBits);
-					Connection->SetNetVersionsOnArchive(Reader);
+					{
+						UE::RemoteExecutor::FSpeculationExecutionScope OptimizationScope;
+						FNetBitReader Reader(Connection->PackageMap, Pending.Buffer.GetData(), Pending.NumBits);
+						Connection->SetNetVersionsOnArchive(Reader);
 
-					TSet<FNetworkGUID> UnmappedGuids;
-					bool bSkipped = false;
-					bSuccess = ReceivedRPC(Reader, UnmappedGuids, bSkipped, Pending.RepFlags, FieldCache, ESkipRpcBehavior::NeverSkip);
+						TSet<FNetworkGUID> UnmappedGuids;
+						bool bSkipped = false;
+						bSuccess = ReceivedRPC(Reader, UnmappedGuids, bSkipped, Pending.RepFlags, FieldCache, ESkipRpcBehavior::NeverSkip);
 
-					ensureMsgf(UnmappedGuids.IsEmpty(), TEXT("%hs: Servers are expected to never have UnmappedGuids"), __func__);
-					ensureMsgf(!bSkipped, TEXT("%hs: We should never skip an RPC when requesting NeverSkip"), __func__);
+						ensureMsgf(UnmappedGuids.IsEmpty(), TEXT("%hs: Servers are expected to never have UnmappedGuids"), __func__);
+						ensureMsgf(!bSkipped, TEXT("%hs: We should never skip an RPC when requesting NeverSkip"), __func__);
+					}
 
 					UE::RemoteExecutor::OnRegionEndDelegate.ExecuteIfBound(TEXT(""));
 					AutoRTFM::PopOnAbortHandler(&Pending);
@@ -2713,7 +2720,17 @@ void FObjectReplicator::EnqueuePendingTransactionalRPCs()
 			});
 	}
 
-	UE::RemoteExecutor::EnqueueWorkList(Object->GetFName(), true, TransactionsToEnqueue);
+	if (CVarEnqueueTransactionalRpcs.GetValueOnAnyThread(true))
+	{
+		UE::RemoteExecutor::EnqueueWorkList(Object->GetFName(), true, TransactionsToEnqueue);
+	}
+	else
+	{
+		for (const TFunction<void(void)>& RpcFunc : TransactionsToEnqueue)
+		{
+			RpcFunc();
+		}
+	}
 }
 
 void FObjectReplicator::QueuePropertyRepNotify(
