[UObject]
- Only query the outer and FName when necessary to avoid race conditions with UObject rename when serializing with a reference collector

diff --git a/Engine/Source/Runtime/CoreUObject/Private/UObject/Obj.cpp b/Engine/Source/Runtime/CoreUObject/Private/UObject/Obj.cpp
--- a/Engine/Source/Runtime/CoreUObject/Private/UObject/Obj.cpp
+++ b/Engine/Source/Runtime/CoreUObject/Private/UObject/Obj.cpp
@@ -1665,13 +1665,11 @@ void UObject::Serialize(FStructuredArchive::FRecord Record)
 
 		FArchive& UnderlyingArchive = Record.GetUnderlyingArchive();
 
-		// These three items are very special items from a serialization standpoint. They aren't actually serialized.
+		// This item is very special from a serialization standpoint. It is never actually serialized.
 		UClass *ObjClass = GetClass();
-		UObject* LoadOuter = GetOuter();
-		FName LoadName = GetFName();
 
 		// Make sure this object's class's data is loaded.
-		if(ObjClass->HasAnyFlags(RF_NeedLoad) )
+		if (ObjClass->HasAnyFlags(RF_NeedLoad) )
 		{
 			UnderlyingArchive.Preload(ObjClass);
 
@@ -1703,9 +1701,15 @@ void UObject::Serialize(FStructuredArchive::FRecord Record)
 		// Special info.
 		if ((!UnderlyingArchive.IsLoading() && !UnderlyingArchive.IsSaving() && !UnderlyingArchive.IsObjectReferenceCollector()))
 		{
+			// This item is very special from a serialization standpoint. It is never actually serialized.
+			// Only query the outer and FName when necessary to avoid race conditions with UObject rename when serializing for reference collection
+			FName LoadName = GetFName();
 			Record << SA_VALUE(TEXT("LoadName"), LoadName);
 			if (!UnderlyingArchive.IsIgnoringOuterRef())
 			{
+				// This item is very special from a serialization standpoint. It is never actually serialized.
+				// Only query the outer and FName when necessary to avoid race conditions with UObject rename when serializing for reference collection
+				UObject* LoadOuter = GetOuter();
 				Record << SA_VALUE(TEXT("LoadOuter"), LoadOuter);
 			}
 			if (!UnderlyingArchive.IsIgnoringClassRef())
@@ -1720,6 +1724,11 @@ void UObject::Serialize(FStructuredArchive::FRecord Record)
 			{
 				UPackage* LoadPackage = GetExternalPackage();
 
+				// These items are very special from a serialization standpoint. They are never actually serialized.
+				// Only query the outer and FName when necessary to avoid race conditions with UObject rename when serializing for reference collection
+				FName LoadName = GetFName();
+				UObject* LoadOuter = GetOuter();
+
 				if (UnderlyingArchive.IsLoading())
 				{
 					Record << SA_VALUE(TEXT("LoadName"), LoadName) << SA_VALUE(TEXT("LoadOuter"), LoadOuter) << SA_VALUE(TEXT("LoadPackage"), LoadPackage);
