[UE] Mitigate a potential race condition when preloading untyped property bag objects during editor load.

#FailedValidations Preflight

diff --git a/Engine/Source/Runtime/CoreUObject/Private/UObject/LinkerLoad.cpp b/Engine/Source/Runtime/CoreUObject/Private/UObject/LinkerLoad.cpp
--- a/Engine/Source/Runtime/CoreUObject/Private/UObject/LinkerLoad.cpp
+++ b/Engine/Source/Runtime/CoreUObject/Private/UObject/LinkerLoad.cpp
@@ -1181,7 +1181,6 @@ FLinkerLoad::FLinkerLoad(UPackage* InParent, const FPackagePath& InPackagePath,
 #if WITH_EDITOR
 , bExportsDuplicatesFixed(false)
 , bIsPackageRelocated(false)
-, bIsLoadingToPropertyBagObject(false)
 , bIsSerializingScriptProperties(false)
 #endif // WITH_EDITOR
 , IsTimeLimitExceededCallCount(0)
@@ -4789,7 +4788,8 @@ void FLinkerLoad::Preload( UObject* Object )
 				// for placeholder objects that have no explicit type, we only want to serialize the TPS stream
 				bool bSerializeOnlyScriptProperties = false;
 				const bool bDoesSavedClassMatchActualClass = DoesSavedClassMatchActualClass(ExportIndex);
-				FGuardValue_Bitfield(bIsLoadingToPropertyBagObject, UE::FPropertyBagRepository::IsPropertyBagPlaceholderObject(Object));
+				bool& bIsLoadingToPropertyBagObject = FUObjectThreadContext::Get().GetSerializeContext()->bIsLoadingToPropertyBagObject;
+				TGuardValue<bool> ScopeGuard(bIsLoadingToPropertyBagObject, UE::FPropertyBagRepository::IsPropertyBagPlaceholderObject(Object));
 				if (UEVer() >= EUnrealEngineObjectUE5Version::SCRIPT_SERIALIZATION_OFFSET)
 				{
 					if (bIsLoadingToPropertyBagObject || !bDoesSavedClassMatchActualClass)
@@ -7292,6 +7292,11 @@ bool FLinkerLoad::TryGetPreloadedLoader(const FPackagePath& InPackagePath, FOpen
 	return IPackageResourceManager::TryTakePreloadableArchive(InPackagePath, OutResult);
 }
 
+bool FLinkerLoad::IsLoadingToPropertyBagObject() const
+{
+	return FUObjectThreadContext::Get().GetSerializeContext()->bIsLoadingToPropertyBagObject;
+}
+
 #endif
 
 bool FLinkerLoad::SerializeBulkData(FBulkData& BulkData, const FBulkDataSerializationParams& Params)
diff --git a/Engine/Source/Runtime/CoreUObject/Private/UObject/UObjectThreadContext.cpp b/Engine/Source/Runtime/CoreUObject/Private/UObject/UObjectThreadContext.cpp
--- a/Engine/Source/Runtime/CoreUObject/Private/UObject/UObjectThreadContext.cpp
+++ b/Engine/Source/Runtime/CoreUObject/Private/UObject/UObjectThreadContext.cpp
@@ -66,6 +66,7 @@ FUObjectSerializeContext::FUObjectSerializeContext()
 	, bTrackUnknownProperties(false)
 	, bTrackUnknownEnumNames(false)
 	, bImpersonateProperties(false)
+	, bIsLoadingToPropertyBagObject(false)
 #endif
 {
 }
diff --git a/Engine/Source/Runtime/CoreUObject/Public/UObject/LinkerLoad.h b/Engine/Source/Runtime/CoreUObject/Public/UObject/LinkerLoad.h
--- a/Engine/Source/Runtime/CoreUObject/Public/UObject/LinkerLoad.h
+++ b/Engine/Source/Runtime/CoreUObject/Public/UObject/LinkerLoad.h
@@ -479,9 +479,6 @@ private:
 	/** Cache if the package is relocated or not */
 	bool bIsPackageRelocated : 1;
 
-	/** Set when the loader is serializing to a property bag placeholder object.											*/
-	bool bIsLoadingToPropertyBagObject : 1;
-
 	/** TRUE when the loader is actively serializing an object's script properties data.									*/
 	bool bIsSerializingScriptProperties : 1;
 #endif // WITH_EDITOR
@@ -1082,7 +1079,7 @@ private:
 #if WITH_EDITOR
 		// This handles the case where we're only serializing to a property bag and we don't immediately know where the script
 		// property data starts/ends in the object's data stream. Allows for loading from packages saved with an older version.
-		if (UNLIKELY(bIsLoadingToPropertyBagObject && !bIsSerializingScriptProperties))
+		if (UNLIKELY(IsLoadingToPropertyBagObject() && !bIsSerializingScriptProperties))
 		{
 			Loader->Seek(Tell() + Length);
 		}
@@ -1404,6 +1401,9 @@ private:
 	* @return True if the import error should be suppressed
 	*/
 	bool IsSuppressableBlueprintImportError(int32 ImportIndex) const;
+
+	/** @return True when the loader is serializing to a property bag placeholder object. */
+	bool IsLoadingToPropertyBagObject() const;
 #endif // WITH_EDITOR
 
 	/**
diff --git a/Engine/Source/Runtime/CoreUObject/Public/UObject/UObjectThreadContext.h b/Engine/Source/Runtime/CoreUObject/Public/UObject/UObjectThreadContext.h
--- a/Engine/Source/Runtime/CoreUObject/Public/UObject/UObjectThreadContext.h
+++ b/Engine/Source/Runtime/CoreUObject/Public/UObject/UObjectThreadContext.h
@@ -195,6 +195,8 @@ public:
 	UE_INTERNAL bool bTrackUnknownEnumNames;
 	/** True when the SerializedObject properties are being impersonated. */
 	UE_INTERNAL bool bImpersonateProperties;
+	/** True when serializing to a property bag placeholder object while loading. */
+	UE_INTERNAL bool bIsLoadingToPropertyBagObject;
 #endif
 
 	/** Adds a new loaded object */
