[UObjectArray]
- Fix TSAN benign race on relaxed read

diff --git a/Engine/Source/Runtime/CoreUObject/Private/UObject/UObjectArray.cpp b/Engine/Source/Runtime/CoreUObject/Private/UObject/UObjectArray.cpp
--- a/Engine/Source/Runtime/CoreUObject/Private/UObject/UObjectArray.cpp
+++ b/Engine/Source/Runtime/CoreUObject/Private/UObject/UObjectArray.cpp
@@ -518,14 +518,15 @@ int32 FUObjectArray::AllocateSerialNumber(int32 Index)
 	checkSlow(ObjectItem);
 
 	volatile int32 *SerialNumberPtr = &ObjectItem->SerialNumber;
-	int32 SerialNumber = *SerialNumberPtr;
-	if (!SerialNumber)
-	{
-		// Open around PrimarySerialNumber. If we fail/abort a transaction, we don't need to undo this; we simply allow it to grow for the next use.
-		// Disable validation for this open as we're performing an explicitly
-		// recorded write to SerialNumber which the memory validator will treat
-		// as a false-positive mixed open / closed write.
-		UE_AUTORTFM_OPEN_NO_VALIDATION
+	// Open around PrimarySerialNumber. If we fail/abort a transaction, we don't need to undo this; we simply allow it to grow for the next use.
+	// Disable validation for this open as we're performing an explicitly
+	// recorded write to SerialNumber which the memory validator will treat
+	// as a false-positive mixed open / closed write.
+	int32 SerialNumber;
+	UE_AUTORTFM_OPEN_NO_VALIDATION
+	{
+		SerialNumber = FPlatformAtomics::AtomicRead_Relaxed(SerialNumberPtr);
+		if (!SerialNumber)
 		{
 			SerialNumber = PrimarySerialNumber.Increment();
 			UE_CLOG(SerialNumber <= START_SERIAL_NUMBER, LogUObjectArray, Fatal, TEXT("UObject serial numbers overflowed (trying to allocate serial number %d)."), SerialNumber);
@@ -536,9 +537,10 @@ int32 FUObjectArray::AllocateSerialNumber(int32 Index)
 				// Someone else got it first; use their value.
 				SerialNumber = ValueWas;
 			}
-		};
-	}
+		}
+	};
 	checkSlow(SerialNumber > START_SERIAL_NUMBER);
+	
 	return SerialNumber;
 }
 
