Call ConditionalSerializeCellExport correctly

#rnx
#FailedValidations Preflight

diff --git a/Engine/Source/Runtime/CoreUObject/Private/Serialization/AsyncLoading2.cpp b/Engine/Source/Runtime/CoreUObject/Private/Serialization/AsyncLoading2.cpp
--- a/Engine/Source/Runtime/CoreUObject/Private/Serialization/AsyncLoading2.cpp
+++ b/Engine/Source/Runtime/CoreUObject/Private/Serialization/AsyncLoading2.cpp
@@ -8156,7 +8156,7 @@ void FAsyncPackage2::ProcessExportDependencies(const FAsyncPackageHeaderData& He
 			}
 			else
 			{
-				ConditionalCreateCellExport(Header, Dep.LocalImportOrExportIndex.ToExport() - Header.ExportsView.Num());
+				ConditionalSerializeCellExport(Header, Dep.LocalImportOrExportIndex.ToExport() - Header.ExportsView.Num());
 			}
 		}
 		else
@@ -8167,7 +8167,7 @@ void FAsyncPackage2::ProcessExportDependencies(const FAsyncPackageHeaderData& He
 			}
 			else
 			{
-				ConditionalCreateCellImport(Header, Dep.LocalImportOrExportIndex.ToImport() - Header.ImportMap.Num());
+				ConditionalSerializeCellImport(Header, Dep.LocalImportOrExportIndex.ToImport() - Header.ImportMap.Num());
 			}
 		}
 	}
