SavePackage: Rename Save2 back to Save, now that the legacy version of Save has been deleted for a few releases. Intention is to reduce confusion about what the 2 in Save2 refers to.
#rnx

diff --git a/Engine/Source/Runtime/Core/Public/Serialization/Archive.h b/Engine/Source/Runtime/Core/Public/Serialization/Archive.h
--- a/Engine/Source/Runtime/Core/Public/Serialization/Archive.h
+++ b/Engine/Source/Runtime/Core/Public/Serialization/Archive.h
@@ -1076,8 +1076,8 @@ public:
 // These will be private in FArchive
 protected:
 	/**
-	 * Provider of data and API specific to UPackage::Save2. Required by archives used from UPackage::Save2, and required for cook saves of packages.
-	 * nullptr means the archive is not one from UPackage::Save2 and is not being called from a cook save.
+	 * Provider of data and API specific to UPackage::Save. Required by archives used from UPackage::Save, and required for cook saves of packages.
+	 * nullptr means the archive is not one from UPackage::Save and is not being called from a cook save.
 	 */
 	FArchiveSavePackageData* SavePackageData = nullptr;
 	
diff --git a/Engine/Source/Runtime/CoreUObject/Private/UObject/SavePackage.cpp b/Engine/Source/Runtime/CoreUObject/Private/UObject/SavePackage.cpp
deleted file mode 100644
diff --git a/Engine/Source/Runtime/CoreUObject/Private/UObject/SavePackage/SavePackageUtilities.h b/Engine/Source/Runtime/CoreUObject/Private/UObject/SavePackage/SavePackageUtilities.h
--- a/Engine/Source/Runtime/CoreUObject/Private/UObject/SavePackage/SavePackageUtilities.h
+++ b/Engine/Source/Runtime/CoreUObject/Private/UObject/SavePackage/SavePackageUtilities.h
@@ -19,7 +19,7 @@
 #include "UObject/SavePackage.h"
 #include "UObject/UObjectGlobals.h"
 
-// This file contains private utilities shared by UPackage::Save and UPackage::Save2 
+// This file contains private utilities used by UPackage::Save
 
 class FCbFieldView;
 class FCbWriter;
@@ -120,7 +120,7 @@ using FSavePackageOutputFileArray = TArray<FSavePackageOutputFile, TInlineAlloca
 
  /**
   * Helper structure to encapsulate sorting a linker's import table alphabetically
-  * @note Save2 should not have to use this sorting long term
+  * @note We should be able to remove this and use less expensive sort order.
   */
 struct FObjectImportSortHelper
 {
@@ -134,7 +134,7 @@ struct FObjectImportSortHelper
 
 /**
  * Helper structure to encapsulate sorting a linker's export table alphabetically
- * @note Save2 should not have to use this sorting long term
+  * @note We should be able to remove this and use less expensive sort order.
  */
 struct FObjectExportSortHelper
 {
@@ -146,7 +146,7 @@ struct FObjectExportSortHelper
 	static void SortExports(FLinkerSave* Linker);
 };
 
-// Utility functions used by both UPackage::Save and/or UPackage::Save2
+// Utility functions used by UPackage::Save
 namespace UE::SavePackageUtilities
 {
 
diff --git a/Engine/Source/Runtime/CoreUObject/Private/UObject/SavePackage2.cpp b/Engine/Source/Runtime/CoreUObject/Private/UObject/SavePackage2.cpp
--- a/Engine/Source/Runtime/CoreUObject/Private/UObject/SavePackage2.cpp
+++ b/Engine/Source/Runtime/CoreUObject/Private/UObject/SavePackage2.cpp
@@ -1,6 +1,5 @@
 // Copyright Epic Games, Inc. All Rights Reserved.
 
-
 #include "UObject/SavePackage.h"
 #include "UObject/LinkerLoad.h"
 
@@ -19,6 +18,7 @@
 #include "Misc/ConfigCacheIni.h"
 #include "Misc/FeedbackContext.h"
 #include "Misc/MessageDialog.h"
+#include "Misc/Optional.h"
 #include "Misc/PackageAccessTracking.h"
 #include "Misc/PackageName.h"
 #include "Misc/PathViews.h"
@@ -36,6 +36,7 @@
 #include "Serialization/PropertyLocalizationDataGathering.h"
 #include "Serialization/UnversionedPropertySerialization.h"
 #include "Serialization/ImportTypeHierarchy.h"
+#include "UObject/ArchiveCookContext.h"
 #include "UObject/DebugSerializationFlags.h"
 #include "UObject/EditorObjectVersion.h"
 #include "UObject/InstanceDataObjectUtils.h"
@@ -67,6 +68,17 @@
 #include "ProfilingDebugging/ScopedTimers.h"
 #endif
 
+bool UPackage::SavePackage(UPackage* InOuter, UObject* InAsset, const TCHAR* Filename, const FSavePackageArgs& SaveArgs)
+{
+	const FSavePackageResultStruct Result = Save(InOuter, InAsset, Filename, SaveArgs);
+	return Result == ESavePackageResult::Success;
+}
+
+FSavePackageContext::~FSavePackageContext()
+{
+	delete PackageWriter;
+}
+
 // defined in UObjectGlobals.cpp
 extern const FName NAME_UniqueObjectNameForCooking;
 COREUOBJECT_API extern bool GOutputCookingWarnings;
@@ -86,7 +98,7 @@ static FAutoConsoleVariableRef CVar_MemoryMappedBulkDataAlignment(
 	TEXT("cook.MemoryMappedBulkDataAlignment"),
 	GMemoryMappedBulkDataAlignment,
 	TEXT("0 - Target platform memory mapping alignment is respected, might result in big zero-padding")
-	TEXT("1 - No special alignment required, 1-byte alignment used (no padding, default)")															  
+	TEXT("1 - No special alignment required, 1-byte alignment used (no padding, default)")
 	TEXT("Any other non-zero value - overrides the memory mapped data alignment, implies padding")
 );
 
@@ -171,12 +183,12 @@ ESavePackageResult ValidatePackage(FSaveContext& SaveContext)
 				FText ErrorText;
 				if (!Asset->GetExternalPackage() && EnumHasAnyFlags(TopLevelFlags, RF_Standalone))
 				{
-					ErrorText = FText::Format(NSLOCTEXT("SavePackage2", "AssetSaveMissingStandaloneFlag", "The Asset {Name} being saved does not have any of the provided object flags (0x{Flags}); saving the package would cause data loss. Run with -dpcvars=save.FixupStandaloneFlags=1 to add the RF_Standalone flag."), 
+					ErrorText = FText::Format(NSLOCTEXT("SavePackage2", "AssetSaveMissingStandaloneFlag", "The Asset {Name} being saved does not have any of the provided object flags (0x{Flags}); saving the package would cause data loss. Run with -dpcvars=save.FixupStandaloneFlags=1 to add the RF_Standalone flag."),
 						Arguments);
 				}
 				else
 				{
-					ErrorText = FText::Format(NSLOCTEXT("SavePackage2", "AssetSaveMissingTopLevelFlags", "The Asset {Name} being saved does not have any of the provided object flags (0x{Flags}); saving the package would cause data loss."),	
+					ErrorText = FText::Format(NSLOCTEXT("SavePackage2", "AssetSaveMissingTopLevelFlags", "The Asset {Name} being saved does not have any of the provided object flags (0x{Flags}); saving the package would cause data loss."),
 						Arguments);
 				}
 				if (SaveContext.IsGenerateSaveError() && SaveContext.GetError())
@@ -203,8 +215,8 @@ ESavePackageResult ValidatePackage(FSaveContext& SaveContext)
 			{
 				FFormatNamedArguments Arguments;
 				Arguments.Add(TEXT("Name"), FText::FromString(FilenameStr));
-				FText FormatText = SaveContext.GetPackage()->ContainsMap() 
-					? NSLOCTEXT("SavePackage2", "MapSaveNotAllowed", "Map '{Name}' is not allowed to save (see log for reason)") 
+				FText FormatText = SaveContext.GetPackage()->ContainsMap()
+					? NSLOCTEXT("SavePackage2", "MapSaveNotAllowed", "Map '{Name}' is not allowed to save (see log for reason)")
 					: NSLOCTEXT("SavePackage2", "AssetSaveNotAllowed", "Asset '{Name}' is not allowed to save (see log for reason)");
 				FText ErrorText = FText::Format(FormatText, Arguments);
 				SaveContext.GetError()->Logf(ELogVerbosity::Warning, TEXT("%s"), *ErrorText.ToString());
@@ -750,13 +762,13 @@ ESavePackageResult ValidateExports(FSaveContext& SaveContext)
 		TSet<UObject*> Exports;
 		Algo::Transform(SaveContext.GetExports(), Exports, [](const FTaggedExport& InExport) { return InExport.Obj; });
 
-		FExportsValidationContext::EFlags Flags = SaveContext.IsCooking() 
+		FExportsValidationContext::EFlags Flags = SaveContext.IsCooking()
 			? FExportsValidationContext::EFlags::IsCooking
 			: FExportsValidationContext::EFlags::None;
 		FOutputDevice* OutputDevice = SaveContext.IsGenerateSaveError() ? SaveContext.GetError() : nullptr;
 		for (const TFunction<FSavePackageSettings::ExternalExportValidationFunc>& ValidateExport : SaveContext.GetExternalExportValidations())
 		{
-			SaveContext.Result = ValidateExport({ SaveContext.GetPackage(), Exports, SaveContext.GetObjectSaveContext().SaveOverrides, Flags, OutputDevice});
+			SaveContext.Result = ValidateExport({ SaveContext.GetPackage(), Exports, SaveContext.GetObjectSaveContext().SaveOverrides, Flags, OutputDevice });
 			if (SaveContext.Result != ESavePackageResult::Success)
 			{
 				return SaveContext.Result;
@@ -844,7 +856,7 @@ ESavePackageResult ValidateIllegalReferences(FSaveContext& SaveContext, TArray<U
 		}
 		else
 		{
-			ReferenceDescription = FString::Printf(TEXT("'%s' referenced by '%s' (at '%s') in its '%s' property"), 
+			ReferenceDescription = FString::Printf(TEXT("'%s' referenced by '%s' (at '%s') in its '%s' property"),
 				*CulpritStr, *ReferencerStr, *ReferencerOuterStr, *ReferencerPropertyStr);
 		}
 
@@ -990,8 +1002,8 @@ ESavePackageResult ValidateImports(FSaveContext& SaveContext)
 		}
 
 		// if an import outer is an export and that import doesn't have a specific package set then, there's an error
-		const bool bWrongImport = Import->GetOuter() 
-			&& Import->GetOuter()->IsInPackage(SaveContext.GetPackage()) 
+		const bool bWrongImport = Import->GetOuter()
+			&& Import->GetOuter()->IsInPackage(SaveContext.GetPackage())
 			&& Import->GetExternalPackage() == nullptr
 			// The optional context will have import that are actually in the same package, similar to external package
 			&& SaveContext.GetCurrentHarvestingRealm() != ESaveRealm::Optional;
@@ -1776,7 +1788,7 @@ void GetCellsWithOuter(Verse::FAllocationContext Context, FSaveContext& SaveCont
 			}
 			return *this;
 		}
-	} Visitor{SaveContext, Results, Stack};
+	} Visitor{ SaveContext, Results, Stack };
 	Outer->Serialize(Visitor);
 	while (Stack.Num() > 0)
 	{
@@ -2308,10 +2320,10 @@ int64 WriteObjectDataResources(TArray<FObjectDataResource>& DataResources, FStru
 	// The data resource table is only saved for cooked output
 	if (Linker.IsCooking() == false || DataResources.IsEmpty())
 	{
-		Linker.Summary.DataResourceOffset = -1; 
+		Linker.Summary.DataResourceOffset = -1;
 		return 0;
 	}
-	
+
 	Linker.Summary.DataResourceOffset = (int32)Linker.Tell();
 	FObjectDataResource::Serialize(StructuredArchiveRoot.EnterField(TEXT("DataResources")), DataResources);
 
@@ -2324,7 +2336,7 @@ int64 WriteObjectDataResources(TArray<FObjectDataResource>& DataResources, FStru
  * FPackageFileSummary uses int32 for a lot of offsets but a lot of our package writing code is
  * capable of handling files that exceed MAX_int32 so the final size is calculated as a int64.
  * If this value is truncated when storing in FPackageFileSummary then the package will not read
- * in correctly and most likely cause a crash. Rather than allowing the user to save bad data we 
+ * in correctly and most likely cause a crash. Rather than allowing the user to save bad data we
  * can use this utility to catch the error and log it so that the user can take action.
  */
 ESavePackageResult SetSummaryTotalHeaderSize(FSaveContext& SaveContext, int64 TotalHeaderSize)
@@ -2509,7 +2521,7 @@ ESavePackageResult WritePackageHeader(FStructuredArchive::FRecord& StructuredArc
 	if (!SaveContext.IsCooking() && UE::Serialization::Private::FImportTypeHierarchy::IsSerializationEnabled())
 	{
 		SCOPED_SAVETIMER(UPackage_Save_SaveImportTypeHierarchies);
-		
+
 		Linker->Summary.ImportTypeHierarchiesOffset = (int32)Linker->Tell();
 		Linker->Summary.ImportTypeHierarchiesCount = Linker->ImportTypeHierarchies.Num();
 
@@ -2554,7 +2566,7 @@ ESavePackageResult WritePackageHeader(FStructuredArchive::FRecord& StructuredArc
 		SCOPED_SAVETIMER(UPackage_Save_PreloadDependencies);
 		SavePreloadDependencies(StructuredArchiveRoot, SaveContext);
 	}
-	
+
 	// Rather than check if an offset is truncated every time we assign one, we can just check the final TotalHeaderSize to see if it is truncated.
 	// Checking every time an offset is assigned would let us fail quicker but a) relies that new code follows the convention b) bloats the code a fair bit.
 	const ESavePackageResult Result = SetSummaryTotalHeaderSize(SaveContext, Linker->Tell());
@@ -2627,12 +2639,23 @@ ESavePackageResult WritePackageTextHeader(FStructuredArchive::FRecord& Structure
 	{
 		FExportProxyArchive(FArchive& Inner)
 			: FArchiveProxy(Inner)
-			, Offset(Inner.Tell()) { }
+			, Offset(Inner.Tell())
+		{
+		}
+
+		virtual void Seek(int64 Pos)
+		{
+			InnerArchive.Seek(Offset + Pos);
+		}
+		virtual int64 Tell()
+		{
+			return InnerArchive.Tell() - Offset;
+		}
+		virtual int64 TotalSize()
+		{
+			return InnerArchive.TotalSize() - Offset;
+		}
 
-		virtual void Seek(int64 Pos) { InnerArchive.Seek(Offset + Pos); }
-		virtual int64 Tell() { return InnerArchive.Tell() - Offset; }
-		virtual int64 TotalSize() { return InnerArchive.TotalSize() - Offset; }
-		
 		int64 Offset;
 	};
 
@@ -2657,7 +2680,7 @@ ESavePackageResult WritePackageTextHeader(FStructuredArchive::FRecord& Structure
 		{
 			continue;
 		}
-			
+
 		SCOPED_SAVETIMER(UPackage_Save_SaveExport);
 		SCOPED_SAVETIMER_TEXT(*WriteToString<128>(GetClassTraceScope(Export.Object), TEXT("_SaveSerialize")));
 
@@ -2845,7 +2868,7 @@ ESavePackageResult WriteExports(FStructuredArchive::FRecord& StructuredArchiveRo
 	FLinkerSave& Linker = *SaveContext.GetLinker();
 
 	Linker.Summary.BulkDataStartOffset = Linker.Tell();
-	
+
 	if (Linker.IsCooking() == false)
 	{
 		check(!Linker.HasCookedIndexBulkData());
@@ -2875,7 +2898,7 @@ ESavePackageResult WriteExports(FStructuredArchive::FRecord& StructuredArchiveRo
 		return ESavePackageResult::Success;
 	}
 
-	// Saving non-inline bulk data to seperate file(s) (Cooking/EditorDomain)
+	// Saving non-inline bulk data to separate file(s) (Cooking/EditorDomain)
 
 	if (Linker.IsCooking() == false)
 	{
@@ -2897,16 +2920,16 @@ ESavePackageResult WriteExports(FStructuredArchive::FRecord& StructuredArchiveRo
 
 	auto GetFilePath = [&SaveContext, bIsOptionalRealm](EPackageExtension Ext, FBulkDataCookedIndex CookedIndex = FBulkDataCookedIndex::Default) -> FString
 	{
-			if (CookedIndex.IsDefault())
-			{
-				FString FileExt = bIsOptionalRealm ? FString(".o") + LexToString(Ext) : LexToString(Ext);
-				return FPathViews::ChangeExtension(SaveContext.GetFilename(), FileExt);
-			}
-			else
-			{
-				FString FileExt = CookedIndex.GetAsExtension() + LexToString(Ext);
-				return FPathViews::ChangeExtension(SaveContext.GetFilename(), FileExt);
-			}
+		if (CookedIndex.IsDefault())
+		{
+			FString FileExt = bIsOptionalRealm ? FString(".o") + LexToString(Ext) : LexToString(Ext);
+			return FPathViews::ChangeExtension(SaveContext.GetFilename(), FileExt);
+		}
+		else
+		{
+			FString FileExt = CookedIndex.GetAsExtension() + LexToString(Ext);
+			return FPathViews::ChangeExtension(SaveContext.GetFilename(), FileExt);
+		}
 	};
 
 	auto WriteToPackageWriter = [&SaveContext, PackageWriter, bIsOptionalRealm](FFileRegionMemoryWriter& Ar, IPackageWriter::FBulkDataInfo Info) -> int64
@@ -2985,8 +3008,8 @@ ESavePackageResult WriteExports(FStructuredArchive::FRecord& StructuredArchiveRo
 		checkf((SaveContext.GetSaveArgs().SaveFlags & SAVE_BulkDataByReference) != 0 || SaveContext.GetLinker()->PackageTrailerBuilder->GetNumReferencedPayloads() == 0,
 			TEXT("Attempting to build a package trailer with referenced payloads but the SAVE_BulkDataByReference flag is not set. '%s'"), *SaveContext.GetPackage()->GetName());
 
-		checkf(	(SaveContext.GetSaveArgs().SaveFlags & SAVE_BulkDataByReference) == 0 || SaveContext.GetLinker()->PackageTrailerBuilder->GetNumLocalPayloads() == 0,
-				TEXT("Attempting to build a package trailer with local payloads but the SAVE_BulkDataByReference flag is set. '%s'"), *SaveContext.GetPackage()->GetName());
+		checkf((SaveContext.GetSaveArgs().SaveFlags & SAVE_BulkDataByReference) == 0 || SaveContext.GetLinker()->PackageTrailerBuilder->GetNumLocalPayloads() == 0,
+			TEXT("Attempting to build a package trailer with local payloads but the SAVE_BulkDataByReference flag is set. '%s'"), *SaveContext.GetPackage()->GetName());
 
 		checkf(SaveContext.IsTextFormat() == false, TEXT("Attempting to build a package trailer for text based asset '%s', this is not supported!"), *SaveContext.GetPackage()->GetName());
 
@@ -3023,7 +3046,7 @@ ESavePackageResult WriteAdditionalExportFiles(FSaveContext& SaveContext)
 
 	if (SaveContext.IsCooking() && SaveContext.AdditionalFilesFromExports.Num() > 0)
 	{
-		checkf(SaveContext.GetCurrentHarvestingRealm() != ESaveRealm::Optional , TEXT("Addtional export files is currently unsupported with optional package multi output, Package %s"), *SaveContext.GetPackage()->GetName());
+		checkf(SaveContext.GetCurrentHarvestingRealm() != ESaveRealm::Optional, TEXT("Addtional export files is currently unsupported with optional package multi output, Package %s"), *SaveContext.GetPackage()->GetName());
 		IPackageWriter* PackageWriter = SavePackageContext ? SavePackageContext->PackageWriter : nullptr;
 		checkf(PackageWriter, TEXT("Cooking requires a PackageWriter"));
 		for (FLargeMemoryWriter& Writer : SaveContext.AdditionalFilesFromExports)
@@ -3229,7 +3252,7 @@ static bool TryCalculateSavedHashAndWriteItBackToMemoryWriter(FLargeMemoryWriter
 	Summary.SetSavedHash(OutSavedHash);
 
 	int64 SavedHashOffset = Summary.GetSavedHashRelativeOffset();
-	if (SavedHashOffset < 0 
+	if (SavedHashOffset < 0
 		|| SavedHashOffset + static_cast<int64>(sizeof(FIoHash)) > Writer->TotalSize())
 	{
 		UE_LOG(LogSavePackage, Error,
@@ -3513,7 +3536,7 @@ ESavePackageResult BeginCachePlatformCookedData(FSaveContext& SaveContext)
 #endif
 }
 
-/** 
+/**
  * Archive used during post save serialization.
  * It ignores serialization of primitive types. It is just used to provide a callback hook for UStruct::Serialize and UObject::Serialize to execute custom code.
  */
@@ -3713,7 +3736,7 @@ ESavePackageResult SaveHarvestedRealms(FSaveContext& SaveContext, ESaveRealm Har
 
 				{
 					// Disables writing stack trace data when appending the exports data
-					FArchiveStackTraceDisabledScope _; 
+					FArchiveStackTraceDisabledScope _;
 					Linker.Serialize(ExportsArchive->GetData(), ExportsArchive->TotalSize());
 				}
 
@@ -3760,7 +3783,7 @@ ESavePackageResult SaveHarvestedRealms(FSaveContext& SaveContext, ESaveRealm Har
 	{
 		SlowTask.EnterProgressFrame();
 
-		SaveContext.Result = WriteBulkData(SaveContext, VirtualExportsFileOffset); 
+		SaveContext.Result = WriteBulkData(SaveContext, VirtualExportsFileOffset);
 		if (SaveContext.Result != ESavePackageResult::Success)
 		{
 			return SaveContext.Result;
@@ -3797,7 +3820,7 @@ ESavePackageResult SaveHarvestedRealms(FSaveContext& SaveContext, ESaveRealm Har
 	if (SaveContext.Result != ESavePackageResult::Success)
 	{
 		return SaveContext.Result;
-	}	
+	}
 	if (PackageWriter)
 	{
 		checkf(SaveContext.GetLinker()->Tell() == EndOfExportsOffset,
@@ -3863,7 +3886,7 @@ ESavePackageResult InnerSave(FSaveContext& SaveContext)
 
 	// Harvest Package
 	SlowTask.EnterProgressFrame();
-	SaveContext.GetObjectSaveContext().ObjectSaveContextPhase = EObjectSaveContextPhase ::Harvest;
+	SaveContext.GetObjectSaveContext().ObjectSaveContextPhase = EObjectSaveContextPhase::Harvest;
 	SaveContext.Result = HarvestPackage(SaveContext);
 	if (SaveContext.Result != ESavePackageResult::Success)
 	{
@@ -3897,7 +3920,7 @@ ESavePackageResult InnerSave(FSaveContext& SaveContext)
 	//										-> .o.ubulk
 	//										-> etc
 	SlowTask.EnterProgressFrame();
-	SaveContext.GetObjectSaveContext().ObjectSaveContextPhase = EObjectSaveContextPhase ::Write;
+	SaveContext.GetObjectSaveContext().ObjectSaveContextPhase = EObjectSaveContextPhase::Write;
 	for (ESaveRealm HarvestingContext : SaveContext.GetHarvestedRealmsToSave())
 	{
 		SaveContext.Result = SaveHarvestedRealms(SaveContext, HarvestingContext);
@@ -3919,7 +3942,7 @@ ESavePackageResult InnerSave(FSaveContext& SaveContext)
 /** WriteAdditionalPayloads writes BulkData, PayloadSidecarFiles, and other data that is separate from exports.
  * Depending on settings, these additional files may be appended to the LinkerSave after the exports, or they
  * may be written into separate sidecar files.
- * 
+ *
  * @param SaveContext The context for the overall save, including data about the additional payloads
  * @param LinkerSize If the Linker has finished writing, this is the size of the Linker's archive: Linker->Tell(). Otherwise it is -1.
  */
@@ -3962,12 +3985,13 @@ FText GetSlowTaskStatusMessage(const FSaveContext& SaveContext)
 
 } // end namespace
 
-FSavePackageResultStruct UPackage::Save2(UPackage* InPackage, UObject* InAsset, const TCHAR* InFilename,
+FSavePackageResultStruct UPackage::Save(UPackage* InPackage, UObject* InAsset, const TCHAR* InFilename,
 	const FSavePackageArgs& SaveArgs)
 {
+	TRACE_CPUPROFILER_EVENT_SCOPE(UPackage::Save);
+
 	COOK_STAT(FScopedDurationTimer FuncSaveTimer(FSavePackageStats::SavePackageTimeSec));
 	COOK_STAT(FSavePackageStats::NumPackagesSaved++);
-	SCOPED_SAVETIMER(UPackage_Save2);
 	UE_SCOPED_COOK_STAT(InPackage->GetFName(), EPackageEventStatType::SavePackage);
 
 	FSaveContext SaveContext(InPackage, InAsset, InFilename, SaveArgs);
@@ -4152,7 +4176,7 @@ ESavePackageResult UPackage::SaveConcurrent(TArrayView<FPackageSaveInfo> InPacka
 			OutResults.Add(SaveContext.GetFinalResult());
 		}
 	}
-	
+
 	return ESavePackageResult::Success;
 }
 
diff --git a/Engine/Source/Runtime/CoreUObject/Public/UObject/LinkerDiff.h b/Engine/Source/Runtime/CoreUObject/Public/UObject/LinkerDiff.h
--- a/Engine/Source/Runtime/CoreUObject/Public/UObject/LinkerDiff.h
+++ b/Engine/Source/Runtime/CoreUObject/Public/UObject/LinkerDiff.h
@@ -14,7 +14,7 @@ class FOutputDevice;
 struct FLinkerExportObject;
 
 class
-UE_DEPRECATED(5.6, "Functionality is no longer needed, now that legacy Save is removed and Save2 is standard. Contact Epic if you need this functionality.")
+UE_DEPRECATED(5.6, "Functionality is no longer needed, now that legacy Save is removed. Contact Epic if you need this functionality.")
 FLinkerDiff
 {
 public:
diff --git a/Engine/Source/Runtime/CoreUObject/Public/UObject/Package.h b/Engine/Source/Runtime/CoreUObject/Public/UObject/Package.h
--- a/Engine/Source/Runtime/CoreUObject/Public/UObject/Package.h
+++ b/Engine/Source/Runtime/CoreUObject/Public/UObject/Package.h
@@ -1169,7 +1169,7 @@ public:
 		const FSavePackageArgs& SaveArgs);
 
 	/**
-	 * Save a list of packages concurrently using Save2 mechanism
+	 * Save a list of packages concurrently.
 	 * SaveConcurrent is currently experimental and shouldn't be used until it can safely replace Save.
 	 */
 	COREUOBJECT_API static ESavePackageResult SaveConcurrent(TArrayView<FPackageSaveInfo> InPackages,
@@ -1203,8 +1203,5 @@ public:
 	* @return true if Package contains no more assets.
 	*/
 	COREUOBJECT_API static bool IsEmptyPackage(UPackage* Package, const UObject* LastReferencer = NULL);
-
-private:
-	static FSavePackageResultStruct Save2(UPackage* InPackage, UObject* InAsset, const TCHAR* InFilename, const FSavePackageArgs& SaveArgs);
 };
 PRAGMA_ENABLE_DEPRECATION_WARNINGS
