Call GetPublicCellExportIndex from ConditionalSerializeCellImport

#rnx
#FailedValidations Preflight

diff --git a/Engine/Source/Runtime/CoreUObject/Private/Serialization/AsyncLoading2.cpp b/Engine/Source/Runtime/CoreUObject/Private/Serialization/AsyncLoading2.cpp
--- a/Engine/Source/Runtime/CoreUObject/Private/Serialization/AsyncLoading2.cpp
+++ b/Engine/Source/Runtime/CoreUObject/Private/Serialization/AsyncLoading2.cpp
@@ -8359,7 +8359,7 @@ Verse::VCell* FAsyncPackage2::ConditionalSerializeCellImport(const FAsyncPackage
 	}
 	uint64 CellExportHash = Header.ImportedPublicExportHashes[PackageImportRef.GetImportedPublicExportHashIndex()];
 	FAsyncPackageHeaderData* SourcePackageHeader = nullptr;
-	int32 CellExportIndex = SourcePackage->GetPublicExportIndex(CellExportHash, SourcePackageHeader);
+	int32 CellExportIndex = SourcePackage->GetPublicCellExportIndex(CellExportHash, SourcePackageHeader);
 	if (CellExportIndex < 0)
 	{
 		return nullptr;
