Obj->GetClass()->IsA may hit a null pointer if used during CoreUObject initialization

#rnx

diff --git a/Engine/Source/Runtime/CoreUObject/Private/UObject/UObjectGlobals.cpp b/Engine/Source/Runtime/CoreUObject/Private/UObject/UObjectGlobals.cpp
--- a/Engine/Source/Runtime/CoreUObject/Private/UObject/UObjectGlobals.cpp
+++ b/Engine/Source/Runtime/CoreUObject/Private/UObject/UObjectGlobals.cpp
@@ -4098,9 +4098,19 @@ void FObjectInitializer::Construct_Internal()
 		ThreadContext.PopInitializer();
 	});
 
-	if (Obj && (GetAllowNativeComponentClassOverrides() || Obj->GetClass()->UObject::IsA(UVerseClass::StaticClass())))
+	if (Obj)
 	{
-		Obj->GetClass()->SetupObjectInitializer(*this);
+		bool bSetupObjectInitializer = GetAllowNativeComponentClassOverrides();
+		if (!bSetupObjectInitializer)
+		{
+			// The Class of the Class may not be available during very early initialization of CoreUObject.
+			UClass* ClassClass = Obj->GetClass()->GetClass();
+			bSetupObjectInitializer = ClassClass && ClassClass->IsChildOf<UVerseClass>();
+		}
+		if (bSetupObjectInitializer)
+		{
+			Obj->GetClass()->SetupObjectInitializer(*this);
+		}
 	}
 
 #if WITH_EDITORONLY_DATA
