When remote object handles are enabled (UE_WITH_REMOTE_OBJECT_HANDLE=1), queue all mutable references into a special RemoteHandleBatcher instead of ImmutableBatcher because any mutable reference may be converted to a remote object stub when performing garbage collection

diff --git a/Engine/Source/Runtime/CoreUObject/Private/UObject/GarbageCollection.cpp b/Engine/Source/Runtime/CoreUObject/Private/UObject/GarbageCollection.cpp
--- a/Engine/Source/Runtime/CoreUObject/Private/UObject/GarbageCollection.cpp
+++ b/Engine/Source/Runtime/CoreUObject/Private/UObject/GarbageCollection.cpp
@@ -1204,6 +1204,12 @@ struct FResolvedMutableReference : FImmutableReference
 	UObject** Mutable;
 };
 
+struct FRemoteReference { UObject** Object; };
+struct FResolvedRemoteReference : FImmutableReference
+{
+	UObject** Mutable;
+};
+
 FORCEINLINE_DEBUGGABLE static bool ValidateReference(UObject* Object, FPermanentObjectPoolExtents PermanentPool, FReferenceToken Referencer, FMemberId MemberId)
 {
 	bool bOk = (!PermanentPool.Contains(Object)) & (!!Object) & IsObjectHandleResolved_ForGC(reinterpret_cast<FObjectHandle&>(Object)); //-V792
@@ -1465,17 +1471,22 @@ protected:
 
 FORCEINLINE UObject*						GetObject(FImmutableReference In) { return In.Object; }
 FORCEINLINE UObject*						GetObject(FMutableReference In) { return *In.Object; }
+FORCEINLINE UObject*						GetObject(FRemoteReference In) { return *In.Object; }
 FORCEINLINE void							PrefetchObjectPointer(FImmutableReference) {}
 FORCEINLINE void							PrefetchObjectPointer(FMutableReference In) { FPlatformMisc::Prefetch(In.Object); }
+FORCEINLINE void							PrefetchObjectPointer(FRemoteReference In) { FPlatformMisc::Prefetch(In.Object); }
 FORCEINLINE FImmutableReference				ToImmutableReference(FMutableReference In) { return {*In.Object}; }
+FORCEINLINE FImmutableReference				ToImmutableReference(FRemoteReference In) { return { *In.Object }; }
 FORCEINLINE FImmutableReference				ToResolvedReference(FImmutableReference In) { return In; }
 FORCEINLINE FResolvedMutableReference		ToResolvedReference(FMutableReference In) { return { ToImmutableReference(In), In.Object }; }
+FORCEINLINE FResolvedRemoteReference		ToResolvedReference(FRemoteReference In) { return { ToImmutableReference(In), In.Object }; }
 template<class RefType> RefType				ToReference(FMutableReference In);
 template<> FORCEINLINE FImmutableReference	ToReference(FMutableReference In) { return ToImmutableReference(In); }
 template<> FORCEINLINE FMutableReference	ToReference(FMutableReference In) { return In; }
 template<class RefType> RefType				MakeReference(UObject*& Object);
 template<> FORCEINLINE FImmutableReference	MakeReference(UObject*& Object) { return {Object}; }
 template<> FORCEINLINE FMutableReference	MakeReference(UObject*& Object) { return {&Object}; }
+template<> FORCEINLINE FRemoteReference		MakeReference(UObject*& Object) { return { &Object }; }
 
 using FStridedReferenceArray = Private::FStridedReferenceArray;
 using FStridedReferenceView = Private::FStridedReferenceView;
@@ -1783,6 +1794,29 @@ private:
 };
 
 
+///////////////////////// TDummyReferenceBatcher ///////////////////////
+
+/**
+ * Dummy reference batcher used to avoid allocating memory for batchers that will not be used in the current build config.
+ * Note that the dummy batcher is only to satisfy the compiler and the template function definitions should not be needed because the linker should never require them.
+ */
+class FDummyReferenceBatcher
+{
+public:
+	FDummyReferenceBatcher(FWorkerContext& InContext)
+	{
+	}
+	template<class ArrayType>
+	void PushArray(ArrayType& Array);
+	template<class ViewType>
+	void PushReferenceView(ViewType View);
+	template<class UnvalidatedReferenceType>
+	void PushReference(UnvalidatedReferenceType UnvalidatedReference);
+	void FlushQueues()
+	{
+	}
+};
+
 ///////////////////////// FStructBatcher + helpers ///////////////////////
 
 // Fixed block of array-of-structs waiting to be processed
@@ -2898,12 +2932,12 @@ int32 GetNumSlowAROs()
 // Helps identify code locations that kill references
 FORCEINLINE_DEBUGGABLE void KillReference(UObject*& Object) { Object = nullptr; }
 
-#if UE_WITH_REMOTE_OBJECT_HANDLE
-FORCEINLINE_DEBUGGABLE void MakeRemoteReference(UObject*& Object) 
+FORCEINLINE_DEBUGGABLE void MakeRemoteHandleReference(UObject*& Object) 
 { 
+#if UE_WITH_REMOTE_OBJECT_HANDLE
 	Object = reinterpret_cast<UObject*>(FObjectHandle::ConvertToRemoteHandle(Object).PointerOrHandle);
+#endif
 }
-#endif // UE_WITH_REMOTE_OBJECT_HANDLE
 
 template <EGCOptions Options>
 constexpr FORCEINLINE EKillable MayKill(EOrigin Origin, bool bAllowKill)
@@ -2976,20 +3010,19 @@ public:
 			FImmutableReference Reference = {Object};	
 			FReferenceMetadata Metadata(ObjectIndex);
 			bool bKillable = Killable == EKillable::Yes;
-			if (Metadata.Has(KillFlag) & bKillable) //-V792
+			if (FRemoteObjectId::RemoteObjectSupportCompiledIn && Metadata.Has(EInternalObjectFlags::Remote))
 			{
-				check(ReferencingObject || IsEliminatingGarbage());
-				checkSlow(Metadata.ObjectItem->GetOwnerIndex() <= 0);
-				KillReference(Object);
+				MakeRemoteHandleReference(Object);
 				return;
 			}
-#if UE_WITH_REMOTE_OBJECT_HANDLE
-			else if (Metadata.Has(EInternalObjectFlags::Remote))
+			else if (Metadata.Has(KillFlag) & bKillable) //-V792
 			{
-				MakeRemoteReference(Object);
+				check(ReferencingObject || IsEliminatingGarbage());
+				checkSlow(Metadata.ObjectItem->GetOwnerIndex() <= 0);
+				KillReference(Object);
 				return;
 			}
-#endif
+
 			DetectGarbageReference(Context, Metadata);
 			HandleValidReference(Context, /*ReferencingObject, */Reference, Metadata);
 		}
@@ -2998,19 +3031,31 @@ public:
 	FORCEINLINE static void HandleBatchedReference(FWorkerContext& Context, FResolvedMutableReference Reference, FReferenceMetadata Metadata)
 	{
 		UE::GC::GDetailedStats.IncreaseObjectRefStats(GetObject(Reference));
-		if (Metadata.Has(KillFlag))
+		if (FRemoteObjectId::RemoteObjectSupportCompiledIn && Metadata.Has(EInternalObjectFlags::Remote))
+		{
+			MakeRemoteHandleReference(*Reference.Mutable);
+		}
+		else if (Metadata.Has(KillFlag))
 		{
 			checkSlow(Metadata.ObjectItem->GetOwnerIndex() <= 0);
 			KillReference(*Reference.Mutable);
 		}
-#if UE_WITH_REMOTE_OBJECT_HANDLE
-		else if (Metadata.Has(EInternalObjectFlags::Remote))
+		else
 		{
-			MakeRemoteReference(*Reference.Mutable);
+			HandleValidReference(Context, Reference, Metadata);
+		}
+	}
+
+	FORCEINLINE static void HandleBatchedReference(FWorkerContext& Context, FResolvedRemoteReference Reference, FReferenceMetadata Metadata)
+	{
+		UE::GC::GDetailedStats.IncreaseObjectRefStats(GetObject(Reference));		
+		if (FRemoteObjectId::RemoteObjectSupportCompiledIn && Metadata.Has(EInternalObjectFlags::Remote))
+		{
+			MakeRemoteHandleReference(*Reference.Mutable);
 		}
-#endif
 		else
 		{
+			DetectGarbageReference(Context, Metadata);
 			HandleValidReference(Context, Reference, Metadata);
 		}
 	}
@@ -3083,6 +3128,12 @@ struct TBatchDispatcher
 	FReferenceCollector& Collector;
 	TReferenceBatcher<FMutableReference, FResolvedMutableReference, ProcessorType> KillableBatcher;
 	TReferenceBatcher<FImmutableReference, FImmutableReference, ProcessorType> ImmutableBatcher;
+#if UE_WITH_REMOTE_OBJECT_HANDLE
+	TReferenceBatcher<FRemoteReference, FResolvedRemoteReference, ProcessorType> RemoteHandleBatcher;
+#else
+	FDummyReferenceBatcher RemoteHandleBatcher;
+#endif
+
 #if WITH_VERSE_VM || defined(__INTELLISENSE__)
 	Verse::FMarkStack VerseGCMarkStack;
 #endif
@@ -3094,6 +3145,7 @@ struct TBatchDispatcher
 	, Collector(InCollector)
 	, KillableBatcher(InContext)
 	, ImmutableBatcher(InContext)
+	, RemoteHandleBatcher(InContext)
 	, StructBatcher(InContext)
 	{}
 
@@ -3182,6 +3234,7 @@ struct TBatchDispatcher
 	{
 		KillableBatcher.FlushQueues();
 		ImmutableBatcher.FlushQueues();
+		RemoteHandleBatcher.FlushQueues();
 	}
 
 	FORCENOINLINE bool FlushToStructBlocks()
@@ -3214,6 +3267,14 @@ struct TBatchDispatcher
 			FPlatformMisc::Prefetch(&Object);
 			KillableBatcher.PushReference(FMutableReference{&Object});
 		}
+		else if constexpr (FRemoteObjectId::RemoteObjectSupportCompiledIn)
+		{
+			// Sanity check so that we never compile this path with the dummy batcher
+			static_assert(!std::is_same_v<decltype(RemoteHandleBatcher), FDummyReferenceBatcher>);
+
+			FPlatformMisc::Prefetch(&Object);
+			RemoteHandleBatcher.PushReference(FRemoteReference{ &Object });
+		}
 		else
 		{
 			ImmutableBatcher.PushReference(FImmutableReference{Object});
@@ -3226,6 +3287,10 @@ struct TBatchDispatcher
 		{
 			KillableBatcher.PushReferenceView(References);
 		}
+		else if constexpr (FRemoteObjectId::RemoteObjectSupportCompiledIn)
+		{
+			RemoteHandleBatcher.PushReferenceView(References);
+		}
 		else
 		{
 			ImmutableBatcher.PushReferenceView(References);
@@ -3239,6 +3304,10 @@ struct TBatchDispatcher
 		{
 			KillableBatcher.PushArray((ArrayType&&)Array);
 		}
+		else if constexpr (FRemoteObjectId::RemoteObjectSupportCompiledIn)
+		{
+			RemoteHandleBatcher.PushArray((ArrayType&&)Array);
+		}
 		else
 		{
 			ImmutableBatcher.PushArray((ArrayType&&)Array);
@@ -3247,11 +3316,11 @@ struct TBatchDispatcher
 
 	FORCEINLINE_DEBUGGABLE void QueueSet(const UObject* ReferencingObject, TSet<UObject*>& Objects, FMemberId MemberId, EKillable Killable)
 	{
-		if (Killable == EKillable::Yes)
+		if (FRemoteObjectId::RemoteObjectSupportCompiledIn || Killable == EKillable::Yes)
 		{
 			for (UObject*& Object : Objects)
 			{
-				QueueReference(ReferencingObject, Object, MemberId, EKillable::Yes);
+				QueueReference(ReferencingObject, Object, MemberId, Killable);
 			}
 		}
 		else
@@ -3534,7 +3603,12 @@ public:
 		if (ValidateReference(Object, PermanentPool, FReferenceToken(ReferencingObject), MemberId))
 		{
 			FReferenceMetadata Metadata(GUObjectArray.ObjectToIndex(Object));
-			if (bAllowReferenceElimination & Metadata.Has(TReachabilityProcessor<Options>::KillFlag)) //-V792
+			if (FRemoteObjectId::RemoteObjectSupportCompiledIn && Metadata.Has(EInternalObjectFlags::Remote))
+			{
+				MakeRemoteHandleReference(Object);
+				return;
+			}
+			else if (bAllowReferenceElimination & Metadata.Has(TReachabilityProcessor<Options>::KillFlag)) //-V792
 			{
 				if (MayKill<Options>(Origin, bAllowReferenceElimination) == EKillable::Yes)
 				{
@@ -3543,23 +3617,13 @@ public:
 					KillReference(Object);
 					return;
 				}
-				else if (!IsEliminatingGarbage(Options) && bTrackGarbage
-#if UE_WITH_REMOTE_OBJECT_HANDLE
-					&& !Metadata.Has(EInternalObjectFlags::Remote)
-#endif
-					)
+				else if (!IsEliminatingGarbage(Options) && bTrackGarbage)
 				{
 					check(Origin == EOrigin::Other);
 					HandleGarbageReference(Context, ReferencingObject, Object, MemberId);
 				}
 			}
-#if UE_WITH_REMOTE_OBJECT_HANDLE
-			if (Metadata.Has(EInternalObjectFlags::Remote))
-			{
-				MakeRemoteReference(Object);
-				return;
-			}
-#endif			
+
 			bool bReachedFirst = TReachabilityProcessor<Options>::HandleValidReference(Context, FImmutableReference{Object}, Metadata);
 
 			if (bReachedFirst && bTrackHistory)
