[CoreUObject]
* Added PLATFORM_COMPILER_IWYU around deprecated code in SoftObjectPtr to prevent IWYU from using paths we don't want to use

diff --git a/Engine/Source/Runtime/CoreUObject/Public/UObject/SoftObjectPtr.h b/Engine/Source/Runtime/CoreUObject/Public/UObject/SoftObjectPtr.h
--- a/Engine/Source/Runtime/CoreUObject/Public/UObject/SoftObjectPtr.h
+++ b/Engine/Source/Runtime/CoreUObject/Public/UObject/SoftObjectPtr.h
@@ -194,6 +194,8 @@ public:
 		: SoftObjectPtr(Other.SoftObjectPtr)
 	{
 	}
+
+#if !defined(PLATFORM_COMPILER_IWYU)
 	template <
 		class U
 		UE_REQUIRES(!std::is_convertible_v<U*, T*>)
@@ -203,6 +205,7 @@ public:
 		: SoftObjectPtr(Other.SoftObjectPtr)
 	{
 	}
+#endif
 
 	/** Construct from a moveable soft pointer */
 	template <
@@ -213,6 +216,8 @@ public:
 		: SoftObjectPtr(MoveTemp(Other.SoftObjectPtr))
 	{
 	}
+
+#if !defined(PLATFORM_COMPILER_IWYU)
 	template <
 		class U
 		UE_REQUIRES(!std::is_convertible_v<U*, T*>)
@@ -222,6 +227,7 @@ public:
 		: SoftObjectPtr(MoveTemp(Other.SoftObjectPtr))
 	{
 	}
+#endif
 
 	/** Construct from an object already in memory */
 	template <
@@ -232,6 +238,8 @@ public:
 		: SoftObjectPtr(Object)
 	{
 	}
+
+#if !defined(PLATFORM_COMPILER_IWYU)
 	template <
 		typename U
 		UE_REQUIRES(!std::is_convertible_v<U*, T*>)
@@ -241,7 +249,7 @@ public:
 		: SoftObjectPtr(Object)
 	{
 	}
-
+#endif
 #if UE_ENABLE_NOTNULL_WRAPPER
 	/** Construct from an object already in memory. Convenience function for implicit initialization from TNotNull pointers. */
 	template <
@@ -272,6 +280,8 @@ public:
 		: SoftObjectPtr(Object.Get())
 	{
 	}
+
+#if !defined(PLATFORM_COMPILER_IWYU)
 	template <
 		typename U
 		UE_REQUIRES(!std::is_convertible_v<U*, T*>)
@@ -281,6 +291,7 @@ public:
 		: SoftObjectPtr(Object.Get())
 	{
 	}
+#endif
 
 	/** Construct from a nullptr */
 	[[nodiscard]] UE_FORCEINLINE_HINT TSoftObjectPtr(TYPE_OF_NULLPTR)
@@ -303,11 +314,14 @@ public:
 		//
 		// The real fix is to make the FSoftObjectPath(const UObject*) constructor explicit, but that will involve substantially more changes.
 	}
+
+#if !defined(PLATFORM_COMPILER_IWYU)
 	UE_SOFTOBJECTPTR_CONVERSION_DEPRECATED(5.5, "Constructing TSoftObjectPtr from an FString has been deprecated - instead, explicitly construct a FSoftObjectPath.")
 	[[nodiscard]] explicit TSoftObjectPtr(const FString& Path)
 		: SoftObjectPtr(FSoftObjectPath(Path))
 	{
 	}
+#endif
 
 	/** Reset the soft pointer back to the null state */
 	UE_FORCEINLINE_HINT void Reset()
@@ -331,6 +345,8 @@ public:
 		SoftObjectPtr = Object;
 		return *this;
 	}
+
+#if !defined(PLATFORM_COMPILER_IWYU)
 	template <
 		typename U
 		UE_REQUIRES(!std::is_convertible_v<U*, T*>)
@@ -341,6 +357,7 @@ public:
 		SoftObjectPtr = Object;
 		return *this;
 	}
+#endif
 
 #if UE_ENABLE_NOTNULL_WRAPPER
 	/** Copy from an object already in memory. Convenience function for implicit assignment from TNotNull pointers. */
@@ -375,6 +392,8 @@ public:
 		SoftObjectPtr = Object.Get();
 		return *this;
 	}
+
+#if !defined(PLATFORM_COMPILER_IWYU)
 	template <
 		typename U
 		UE_REQUIRES(!std::is_convertible_v<U*, T*>)
@@ -385,6 +404,7 @@ public:
 		SoftObjectPtr = Object.Get();
 		return *this;
 	}
+#endif
 
 	/** Assign from a nullptr */
 	inline TSoftObjectPtr& operator=(TYPE_OF_NULLPTR)
@@ -410,6 +430,8 @@ public:
 		SoftObjectPtr = Other;
 		return *this;
 	}
+
+#if !defined(PLATFORM_COMPILER_IWYU)
 	template <
 		class U
 		UE_REQUIRES(!std::is_convertible_v<U*, T*>)
@@ -420,6 +442,7 @@ public:
 		SoftObjectPtr = Other;
 		return *this;
 	}
+#endif
 
 	/** Copy from another soft pointer */
 	template <
@@ -431,6 +454,8 @@ public:
 		SoftObjectPtr = MoveTemp(Other.SoftObjectPtr);
 		return *this;
 	}
+
+#if !defined(PLATFORM_COMPILER_IWYU)
 	template <
 		class U
 		UE_REQUIRES(!std::is_convertible_v<U*, T*>)
@@ -441,7 +466,7 @@ public:
 		SoftObjectPtr = MoveTemp(Other.SoftObjectPtr);
 		return *this;
 	}
-
+#endif
 	/**
 	 * Compare soft pointers for equality
 	 * Caution: Two soft pointers might not be equal to each other, but they both might return nullptr
@@ -460,6 +485,8 @@ public:
 	{
 		return SoftObjectPtr == Rhs.SoftObjectPtr;
 	}
+
+#if !defined(PLATFORM_COMPILER_IWYU)
 	template <
 		typename U
 		UE_REQUIRES(!TModels_V<CEqualityComparableWith, T*, U*>)
@@ -469,6 +496,8 @@ public:
 	{
 		return SoftObjectPtr == Rhs.SoftObjectPtr;
 	}
+#endif
+
 	template <
 		typename U
 		UE_REQUIRES(TModels_V<CEqualityComparableWith, T*, U*>)
@@ -477,6 +506,8 @@ public:
 	{
 		return SoftObjectPtr == TSoftObjectPtr<U>(Rhs).SoftObjectPtr;
 	}
+
+#if !defined(PLATFORM_COMPILER_IWYU)
 	template <
 		typename U
 		UE_REQUIRES(!TModels_V<CEqualityComparableWith, T*, U*>)
@@ -486,6 +517,8 @@ public:
 	{
 		return SoftObjectPtr == TSoftObjectPtr<U>(Rhs).SoftObjectPtr;
 	}
+#endif
+
 	[[nodiscard]] UE_FORCEINLINE_HINT bool operator==(TYPE_OF_NULLPTR) const
 	{
 		return SoftObjectPtr == nullptr;
