Removed LinkerInstancingContext duplication in FWorldPartitionLevelHelper::LoadActorsInternal, causes huge memory usage when loading a lot of actors in a cell and should not be necessary.
Issue stemmed from the Loader modifying the user provided LinkerInstancingContext and setting it in a state when further loads of the same instances could be incorrect.

Added functionality to LinkerInstancingContext to link it on creation to an immutable LinkerInstancingContext and have them function as one.
Loader code now takes the user provided LinkerInstancingContext and uses it as immutable context, this avoids modifications in the user LinkerInstancingContext

#tests cook + test/fix of a local repro

diff --git a/Engine/Source/Runtime/CoreUObject/Private/Serialization/AsyncLoading2.cpp b/Engine/Source/Runtime/CoreUObject/Private/Serialization/AsyncLoading2.cpp
--- a/Engine/Source/Runtime/CoreUObject/Private/Serialization/AsyncLoading2.cpp
+++ b/Engine/Source/Runtime/CoreUObject/Private/Serialization/AsyncLoading2.cpp
@@ -730,7 +730,7 @@ struct FPackageRequest
 			LoadFlags,
 			PIEInstanceID,
 #endif
-			InstancingContext ? *InstancingContext : FLinkerInstancingContext(),
+			InstancingContext ? FLinkerInstancingContext(*InstancingContext, FLinkerInstancingContext::EContextSharingMode::Link) : FLinkerInstancingContext(),
 			CustomName,
 			PackagePath,
 			MoveTemp(PackageLoadedDelegate),
diff --git a/Engine/Source/Runtime/CoreUObject/Private/Serialization/AsyncPackage.h b/Engine/Source/Runtime/CoreUObject/Private/Serialization/AsyncPackage.h
--- a/Engine/Source/Runtime/CoreUObject/Private/Serialization/AsyncPackage.h
+++ b/Engine/Source/Runtime/CoreUObject/Private/Serialization/AsyncPackage.h
@@ -66,7 +66,7 @@ struct FAsyncPackageDesc
 		, PackageFlags(OldPackage.PackageFlags)
 		, Priority(OldPackage.Priority)
 		, PIEInstanceID(OldPackage.PIEInstanceID)
-		, InstancingContext(OldPackage.InstancingContext)
+		, InstancingContext(OldPackage.InstancingContext, FLinkerInstancingContext::EContextSharingMode::Link)
 #if UE_WITH_PACKAGE_ACCESS_TRACKING
 		, ReferencerPackageName(OldPackage.ReferencerPackageName)
 		, ReferencerPackageOp(OldPackage.ReferencerPackageOp)
diff --git a/Engine/Source/Runtime/CoreUObject/Private/UObject/LinkerInstancingContext.cpp b/Engine/Source/Runtime/CoreUObject/Private/UObject/LinkerInstancingContext.cpp
--- a/Engine/Source/Runtime/CoreUObject/Private/UObject/LinkerInstancingContext.cpp
+++ b/Engine/Source/Runtime/CoreUObject/Private/UObject/LinkerInstancingContext.cpp
@@ -104,6 +104,19 @@ public:
 		return Path;
 	}
 
+	bool RemapPath(const FSoftObjectPath& Path, FSoftObjectPath& OutRemappedPath) const
+	{
+		FLinkerInstancingContextReadLock ScopeLock(Lock);
+		if (const FTopLevelAssetPath* Remapped = PathMapping.Find(Path.GetAssetPath()))
+		{
+			OutRemappedPath = FSoftObjectPath(*Remapped, Path.GetSubPathString());
+			return true;
+		}
+
+		OutRemappedPath = Path;
+		return false;
+	}
+
 	bool IsInstanced() const
 	{
 		FLinkerInstancingContextReadLock ScopeLock(Lock);
@@ -123,6 +136,14 @@ public:
 		InstancedPackageMap.BuildPackageMapping(Original, Instanced, bInSoftObjectPathRemappingEnabled);
 	}
 
+	bool ValidatePackageMappingAgainst(FSharedLinkerInstancingContextData& Other)
+	{
+		FLinkerInstancingContextReadLock ScopeLock(Lock);
+		FLinkerInstancingContextReadLock OtherScopeLock(Other.Lock);
+	
+		return InstancedPackageMap.ValidateAgainst(Other.InstancedPackageMap);
+	}
+
 	void AddTag(FName NewTag)
 	{
 		FLinkerInstancingContextWriteLock ScopeLock(Lock);
@@ -213,37 +234,106 @@ FLinkerInstancingContext::FLinkerInstancingContext(bool bInSoftObjectPathRemappi
 	SharedData = MakeShared<FSharedLinkerInstancingContextData>(bInSoftObjectPathRemappingEnabled);
 }
 
+FLinkerInstancingContext::FLinkerInstancingContext(const FLinkerInstancingContext& InOther, EContextSharingMode Mode)
+{
+	LLM_SCOPE_BYTAG(Loading_LinkerInstancingContext);
+
+	switch (Mode)
+	{
+		case EContextSharingMode::Share:
+		{	
+			// Share all the data
+			ImmutableData = InOther.ImmutableData;
+			SharedData = InOther.SharedData;
+		}
+		break;
+		case EContextSharingMode::Link:
+		{
+			if(InOther.ImmutableData.IsValid())
+			{
+				check(InOther.SharedData.IsValid());
+				
+				// other context already has shared immutable data, so we share all the data
+				ImmutableData = InOther.ImmutableData;
+				SharedData = InOther.SharedData;
+			}
+			else
+			{
+				// promote source shared data to immutable data for this context
+				ImmutableData = InOther.SharedData;
+				// create a new shared data object to store mappings added in this context
+				SharedData = MakeShared<FSharedLinkerInstancingContextData>();
+				// copy relevant settings from LinkedImmutableData
+				SharedData->SetSoftObjectPathRemappingEnabled(InOther.GetSoftObjectPathRemappingEnabled());
+			}
+		}
+		break;
+		case EContextSharingMode::Copy:
+		{
+			// copy the shared data
+			SharedData = MakeShared<FSharedLinkerInstancingContextData>(*InOther.SharedData.Get());
+			// keep sharing the immtuable data
+			ImmutableData = InOther.ImmutableData;
+		} 
+		break;
+	}	
+}
+
+FLinkerInstancingContext::FLinkerInstancingContext(const FLinkerInstancingContext& Other)	
+	: FLinkerInstancingContext(Other, EContextSharingMode::Share)
+{
+	
+}
+
 FLinkerInstancingContext FLinkerInstancingContext::DuplicateContext(const FLinkerInstancingContext& InLinkerInstancingContext)
 {
-	FLinkerInstancingContext OutLinkerInstancingContext = InLinkerInstancingContext;
-	OutLinkerInstancingContext.SharedData = MakeShared<FSharedLinkerInstancingContextData>(*OutLinkerInstancingContext.SharedData.Get());
-	return OutLinkerInstancingContext;
+	return FLinkerInstancingContext(InLinkerInstancingContext, EContextSharingMode::Copy);
 }
 
 void FLinkerInstancingContext::EnableAutomationTest() 
 { 
+	check(!ImmutableData.IsValid());
 	SharedData->EnableAutomationTest();
 }
 
 void FLinkerInstancingContext::BuildPackageMapping(FName Original, FName Instanced)
 {
 	SharedData->BuildPackageMapping(Original, Instanced, GetSoftObjectPathRemappingEnabled());
+
+	if (ImmutableData.IsValid())
+	{
+		// check we didn't already built a package mapping on it 
+		ensureAlwaysMsgf(SharedData->ValidatePackageMappingAgainst(*ImmutableData), TEXT("Validate against ImmutableData failed after building PackageMapping for Original: %s, Instanced: %s"), *Original.ToString(), *Instanced.ToString());
+	}
 }
 
 bool FLinkerInstancingContext::FindPackageMapping(FName Original, FName& Instanced) const
 {
-	return SharedData->FindPackageMapping(Original, Instanced);
+	bool bFound = SharedData->FindPackageMapping(Original, Instanced);
+
+	if (!bFound && ImmutableData.IsValid())
+	{
+		bFound = ImmutableData->FindPackageMapping(Original, Instanced);
+	}
+	return bFound;
 }
 
 bool FLinkerInstancingContext::IsInstanced() const
 {
-	return SharedData->IsInstanced();
+	return SharedData->IsInstanced() || (ImmutableData.IsValid() && ImmutableData->IsInstanced());
 }
 
 /** Remap the package name from the import table to its instanced counterpart, otherwise return the name unmodified. */
 FName FLinkerInstancingContext::RemapPackage(const FName& PackageName) const
 {
-	return SharedData->RemapPackage(PackageName);
+	FName RemappedName = SharedData->RemapPackage(PackageName);
+		
+	if (ImmutableData.IsValid() && RemappedName == PackageName)
+	{
+		RemappedName = ImmutableData->RemapPackage(PackageName);
+	}
+	
+	return RemappedName;
 }
 
 /**
@@ -252,18 +342,37 @@ FName FLinkerInstancingContext::RemapPackage(const FName& PackageName) const
  */
 FSoftObjectPath FLinkerInstancingContext::RemapPath(const FSoftObjectPath& Path) const
 {
-	return SharedData->RemapPath(Path);
+	FSoftObjectPath RemappedPath;
+	bool bRemapped = SharedData->RemapPath(Path, RemappedPath);
+		
+	if (ImmutableData.IsValid() && !bRemapped)
+	{
+		RemappedPath = ImmutableData->RemapPath(Path);
+	}
+
+	return RemappedPath;
 }
 
 /** Add a mapping from a package name to a new package name. There should be no separators (. or :) in these strings. */
 void FLinkerInstancingContext::AddPackageMapping(FName Original, FName Instanced)
 {
+	if (ImmutableData.IsValid())
+	{
+		// validate immutable package mapping doesn't contain this mapping or maps to the same package
+		FName AlreadyMapped;
+		if (ImmutableData->FindPackageMapping(Original, AlreadyMapped))
+		{
+			ensureAlwaysMsgf(Instanced == AlreadyMapped, TEXT("Trying to add duplicate package mapping, this add is masked by the previous one (present in the Immutable linked InstancingContext) Original: %s, Instanced: %s, InMapping: %s"), *Original.ToString(), *Instanced.ToString(), *AlreadyMapped.ToString());
+		}
+	}
+
 	SharedData->AddPackageMapping(Original, Instanced);
 }
 
 /** Add a mapping function from a package name to a new package name. This function should be thread-safe, as it can be invoked from ALT. */
 void FLinkerInstancingContext::AddPackageMappingFunc(TFunction<FName(FName)> InInstancedPackageMapFunc)
 {
+	check(!ImmutableData.IsValid());
 	SharedData->SetPackageMappingFunc(InInstancedPackageMapFunc);
 }
 
@@ -275,6 +384,16 @@ void FLinkerInstancingContext::AddPathMapping(FSoftObjectPath Original, FSoftObj
 	ensureAlwaysMsgf(Instanced.GetSubPathString().IsEmpty(),
 		TEXT("Linker instance remap paths should be top-level assets only: ->%s"), *Instanced.ToString());
 
+	if (ImmutableData.IsValid())
+	{
+		// validate immutable path mapping doesn't contain this mapping or maps to the same package
+		FSoftObjectPath AlreadyMapped;
+		if (ImmutableData->RemapPath(Original, AlreadyMapped))
+		{
+			ensureAlwaysMsgf(Instanced == AlreadyMapped, TEXT("Trying to add duplicate path mapping, this add is masked by the previous one (present in the Immutable linked InstancingContext) Original: %s, Instanced: %s, InMapping: %s"), *Original.ToString(), *Instanced.ToString(), *AlreadyMapped.ToString());
+		}	
+	}
+
 	SharedData->AddPathMapping(Original, Instanced);
 }
 
@@ -290,17 +409,18 @@ void FLinkerInstancingContext::AppendTags(const TSet<FName>& NewTags)
 
 bool FLinkerInstancingContext::HasTag(FName Tag) const
 {
-	return SharedData->HasTag(Tag);
+	return SharedData->HasTag(Tag) || (ImmutableData.IsValid() && ImmutableData->HasTag(Tag));
 }
 
 void FLinkerInstancingContext::SetSoftObjectPathRemappingEnabled(bool bInSoftObjectPathRemappingEnabled)
 {
+	check(!ImmutableData.IsValid());
 	SharedData->SetSoftObjectPathRemappingEnabled(bInSoftObjectPathRemappingEnabled);
 }
 
 bool FLinkerInstancingContext::GetSoftObjectPathRemappingEnabled() const
 {
-	return SharedData->GetSoftObjectPathRemappingEnabled();
+	return SharedData->GetSoftObjectPathRemappingEnabled() || (ImmutableData.IsValid() && ImmutableData->GetSoftObjectPathRemappingEnabled());
 }
 
 void FLinkerInstancingContext::FixupSoftObjectPath(FSoftObjectPath& InOutSoftObjectPath) const
@@ -328,6 +448,10 @@ void FLinkerInstancedPackageMap::AddPackageMapping(FName Original, FName Instanc
 			InstancedPackageMapping.Add(Original, Instanced);
 			bIsInstanced |= !Instanced.IsNone();
 		}
+		else
+		{
+			ensureAlwaysMsgf(InstancedPackageMapping.FindChecked(Original) == Instanced, TEXT("Trying to add duplicate package mapping, this add is masked by the previous one. Original: %s, Instanced: %s, InMapping: %s"), *Original.ToString(), *Instanced.ToString(), *InstancedPackageMapping.FindChecked(Original).ToString());
+		}
 	}
 	else
 	{
@@ -337,6 +461,10 @@ void FLinkerInstancedPackageMap::AddPackageMapping(FName Original, FName Instanc
 			InstancedPackageMapping.Add(Instanced, Original);
 			bIsInstanced |= !Original.IsNone();
 		}
+		else
+		{
+			ensureAlwaysMsgf(InstancedPackageMapping.FindChecked(Instanced) == Original, TEXT("Trying to add duplicate package mapping, this add is masked by the previous one. Original: %s, Instanced: %s, InMapping: %s"), *Original.ToString(), *Instanced.ToString(), *InstancedPackageMapping.FindChecked(Instanced).ToString());
+		}
 	}
 }
 
@@ -402,6 +530,29 @@ void FLinkerInstancedPackageMap::BuildPackageMapping(FName Original, FName Insta
 	}
 }
 
+bool FLinkerInstancedPackageMap::ValidateAgainst(FLinkerInstancedPackageMap& OtherPackageMap) const
+{
+	// This method ensures the object passed as a parameter doesn't contain any mapping that diverge from us
+	bool bResult = true;
+
+	for (TPair<FName, FName> It : InstancedPackageMapping)
+	{
+		FName* OtherInstanced = OtherPackageMap.InstancedPackageMapping.Find(It.Key);
+
+		bool bThisResult = (!OtherInstanced || (*OtherInstanced == It.Value));
+		
+		ensureAlwaysMsgf(bThisResult, TEXT("Discrepancy in instance package mapping for %s, maps to %s and %s"), *It.Key.ToString(), *It.Value.ToString(), *OtherInstanced->ToString());
+
+		bResult &= bThisResult;
+	}
+	
+	bResult &= (GeneratedPackagesFolder == OtherPackageMap.GeneratedPackagesFolder || OtherPackageMap.GeneratedPackagesFolder.IsEmpty());
+	bResult &= (InstancedPackageSuffix == OtherPackageMap.InstancedPackageSuffix || OtherPackageMap.InstancedPackageSuffix.IsEmpty());
+	bResult &= (InstancedPackagePrefix == OtherPackageMap.InstancedPackagePrefix || OtherPackageMap.InstancedPackagePrefix.IsEmpty());
+	
+	return bResult;
+}
+
 bool FLinkerInstancedPackageMap::FixupSoftObjectPath(FSoftObjectPath& InOutSoftObjectPath) const
 {
 	if (IsInstanced())
diff --git a/Engine/Source/Runtime/CoreUObject/Private/UObject/LinkerLoad.cpp b/Engine/Source/Runtime/CoreUObject/Private/UObject/LinkerLoad.cpp
--- a/Engine/Source/Runtime/CoreUObject/Private/UObject/LinkerLoad.cpp
+++ b/Engine/Source/Runtime/CoreUObject/Private/UObject/LinkerLoad.cpp
@@ -877,7 +877,7 @@ FLinkerLoad* FLinkerLoad::CreateLinkerAsync(FUObjectSerializeContext* LoadContex
 		{
 			LoadFlags |= LOAD_Async;
 		}
-		Linker = new FLinkerLoad(Parent, PackagePath, LoadFlags, InstancingContext ? *InstancingContext : FLinkerInstancingContext());
+		Linker = new FLinkerLoad(Parent, PackagePath, LoadFlags, InstancingContext ? FLinkerInstancingContext(*InstancingContext, FLinkerInstancingContext::EContextSharingMode::Link) : FLinkerInstancingContext());
 		Parent->SetLinker(Linker);
 		if (GEventDrivenLoaderEnabled && Linker)
 		{
diff --git a/Engine/Source/Runtime/CoreUObject/Public/UObject/LinkerInstancingContext.h b/Engine/Source/Runtime/CoreUObject/Public/UObject/LinkerInstancingContext.h
--- a/Engine/Source/Runtime/CoreUObject/Public/UObject/LinkerInstancingContext.h
+++ b/Engine/Source/Runtime/CoreUObject/Public/UObject/LinkerInstancingContext.h
@@ -60,6 +60,8 @@ public:
 
 	COREUOBJECT_API bool FixupSoftObjectPath(FSoftObjectPath& InOutSoftObjectPath) const;
 
+	bool ValidateAgainst(FLinkerInstancedPackageMap& OtherPackageMap) const;
+
 private:
 	friend class FLinkerInstancingContext;
 	friend class FLinkerInstancingContextTests;
@@ -98,7 +100,36 @@ private:
 class FLinkerInstancingContext
 {
 public:
+	enum class EContextSharingMode
+	{
+		Copy,
+		Link,
+		Share
+	};
+
 	COREUOBJECT_API FLinkerInstancingContext();
+
+	/** Construct an FLinkerInstancingContext from another and share it's data */
+	COREUOBJECT_API FLinkerInstancingContext(const FLinkerInstancingContext&);
+
+	/** Construct an FLinkerInstancingContext from another and specifying the sharing mode: 
+	* 
+	*	Copy:  This will copy the ptr to the shared immutable data from the Source context, and duplicate the shared data of the Source context 
+	* 
+	*	Link: * If the Source context has no shared immutable data it'll construct this context by using the Source context shared data as the 
+	*		     immutable data for this context and create a new share data object to use for mutations specific to this context. 
+	* 
+	*		   * If the Source context has shared immutable data we copy the ptr to share it as our shared immutable data , and we copy the 
+	*			 the shared data ptr to share it. This means the context has been linked more than 1 time. 
+	* 
+	*	Share: We copy both the shared data and immutable share data ptrs. 
+	* 
+	*	The Link context is intended to be used internally in the loader/linker to ensure that mappings that occur during loading have no side
+	*	effects on further loads done with the same FLinkerInstancingContext thus allowing them to be resued. 
+	* 
+	*/	
+	COREUOBJECT_API FLinkerInstancingContext(const FLinkerInstancingContext& Source, EContextSharingMode Mode);
+
 	COREUOBJECT_API explicit FLinkerInstancingContext(TSet<FName> InTags);
 	COREUOBJECT_API explicit FLinkerInstancingContext(bool bInSoftObjectPathRemappingEnabled);
 
@@ -160,4 +191,5 @@ private:
 	class FSharedLinkerInstancingContextData;
 
 	TSharedPtr<FSharedLinkerInstancingContextData> SharedData;
+	TSharedPtr<FSharedLinkerInstancingContextData> ImmutableData;
 };
\ No newline at end of file
diff --git a/Engine/Source/Runtime/Engine/Private/WorldPartition/ActorDescContainerInstance.cpp b/Engine/Source/Runtime/Engine/Private/WorldPartition/ActorDescContainerInstance.cpp
--- a/Engine/Source/Runtime/Engine/Private/WorldPartition/ActorDescContainerInstance.cpp
+++ b/Engine/Source/Runtime/Engine/Private/WorldPartition/ActorDescContainerInstance.cpp
@@ -144,7 +144,10 @@ void UActorDescContainerInstance::Initialize(const FInitializeParams& InParams)
 			const FString LongActorPackageName = It->GetActorPackage().ToString();
 			const FString InstancedName = ULevel::GetExternalActorPackageInstanceName(OuterWorldPackageNameStr, LongActorPackageName);
 
-			InstancingContext->AddPackageMapping(*LongActorPackageName, *InstancedName);
+			if (LongActorPackageName != OuterWorldContainerPackageName)
+			{
+				InstancingContext->AddPackageMapping(*LongActorPackageName, *InstancedName);
+			}
 			ActorDescInstance.ActorPath = It->GetActorSoftPath().ToString().Replace(*SourceWorldPath, *RemappedWorldPath);
 		}
 
diff --git a/Engine/Source/Runtime/Engine/Private/WorldPartition/WorldPartitionLevelHelper.cpp b/Engine/Source/Runtime/Engine/Private/WorldPartition/WorldPartitionLevelHelper.cpp
--- a/Engine/Source/Runtime/Engine/Private/WorldPartition/WorldPartitionLevelHelper.cpp
+++ b/Engine/Source/Runtime/Engine/Private/WorldPartition/WorldPartitionLevelHelper.cpp
@@ -316,7 +316,7 @@ void FWorldPartitionLevelHelper::MoveExternalActorsToLevel(const TArray<FWorldPa
 						}
 						else
 						{
-							// Move objects in the destination level package if not outerred to the Actor (and thus wouldn't follow)
+							// Move objects in the destination level package
 							// We handle name clashes in a custom way by increasing a suffix on the already existing base name instead
 							// of using the default behavior of having the class name be the base name, this is to maximally preserve
 							// information present in the object name in case somebody is using the name as significant information 
@@ -705,9 +705,13 @@ bool FWorldPartitionLevelHelper::LoadActorsInternal(FLoadActorsParams&& InParams
 			if (bConsiderActorEditorOnly || PackageObjectMapping.ContainerPackage != ContainerPackageInstanceName)
 			{
 				const FName ActorPackageName = *FPackageName::ObjectPathToPackageName(PackageObjectMapping.Package.ToString());
-				const FName ActorPackageInstanceName = bConsiderActorEditorOnly ? NAME_None : FName(*ULevel::GetExternalActorPackageInstanceName(ContainerPackageInstanceName.ToString(), ActorPackageName.ToString()));
 
-				Context->AddPackageMapping(ActorPackageName, ActorPackageInstanceName);
+				// Only add PackageMappings for ExternalPackage Actors, otherwise they're invalid mappings
+				if (ActorPackageName != PackageObjectMapping.ContainerPackage)
+				{
+					const FName ActorPackageInstanceName = bConsiderActorEditorOnly ? NAME_None : FName(*ULevel::GetExternalActorPackageInstanceName(ContainerPackageInstanceName.ToString(), ActorPackageName.ToString()));
+					Context->AddPackageMapping(ActorPackageName, ActorPackageInstanceName);
+				}
 			}
 
 			if (!bConsiderActorEditorOnly)
@@ -736,7 +740,7 @@ bool FWorldPartitionLevelHelper::LoadActorsInternal(FLoadActorsParams&& InParams
 	for (FWorldPartitionRuntimeCellObjectMapping* PackageObjectMapping : ActorPackagesToLoad)
 	{
 		const FName PackageToLoad(*FPackageName::ObjectPathToPackageName(PackageObjectMapping->Package.ToString()));
-		const FLinkerInstancingContext ContainerInstancingContext = FLinkerInstancingContext::DuplicateContext(LinkerInstancingContexts.FindChecked(PackageObjectMapping->ContainerID));
+		const FLinkerInstancingContext& ContainerInstancingContext = LinkerInstancingContexts.FindChecked(PackageObjectMapping->ContainerID);
 		const FName PackageName = ContainerInstancingContext.RemapPackage(PackageToLoad);
 
 		if (!PackageObjectMapping->ContainerID.IsMainContainer())
