Skip reading replaced keys when NumKeysToRemove == INDEX_NONE in FMapProperty::ConvertFromType

#rnx

diff --git a/Engine/Source/Runtime/CoreUObject/Private/UObject/PropertyMap.cpp b/Engine/Source/Runtime/CoreUObject/Private/UObject/PropertyMap.cpp
--- a/Engine/Source/Runtime/CoreUObject/Private/UObject/PropertyMap.cpp
+++ b/Engine/Source/Runtime/CoreUObject/Private/UObject/PropertyMap.cpp
@@ -1536,8 +1536,9 @@ EConvertFromTypeResult FMapProperty::ConvertFromType(const FPropertyTag& Tag, FS
 	// any of the elements that were not present when we saved this Map:
 	int32 NumKeysToRemove = 0;
 	FStructuredArchive::FArray KeysToRemoveArray = ValueRecord.EnterArray(TEXT("KeysToRemove"), NumKeysToRemove);
+	const bool bReplaceMap = NumKeysToRemove == INDEX_NONE;
 
-	if (NumKeysToRemove)
+	if (NumKeysToRemove && !bReplaceMap)
 	{
 		TempKeyValueStorage = (uint8*)FMemory::Malloc(MapLayout.SetLayout.Size);
 		KeyProp->InitializeValue(TempKeyValueStorage);
