RemoteExecutor:  Fix for speculation used with CascadingRetryTransaction.  When the retry happens, the transaction never aborts to user code, so it executes again immediately.  We have to reset our variables to undo the results of the previous run.

diff --git a/Engine/Source/Runtime/CoreUObject/Private/UObject/RemoteExecutor.cpp b/Engine/Source/Runtime/CoreUObject/Private/UObject/RemoteExecutor.cpp
--- a/Engine/Source/Runtime/CoreUObject/Private/UObject/RemoteExecutor.cpp
+++ b/Engine/Source/Runtime/CoreUObject/Private/UObject/RemoteExecutor.cpp
@@ -119,6 +119,7 @@ public:
 	// Pump the network next chance we get
 	bool bShouldPumpNetwork = false;
 
+	// These should only ever be accessed in an AUTORTFM_OPEN scope
 	UE::RemoteExecutor::ERemoteExecutorAbortReason AbortReason = UE::RemoteExecutor::ERemoteExecutorAbortReason::Unspecified;
 	FString AbortReasonDescription;
 
@@ -431,11 +432,6 @@ public:
 					}
 
 					// all of the subsystems are ready, try to perform the work
-					AbortReason = UE::RemoteExecutor::ERemoteExecutorAbortReason::Unspecified;
-					AbortReasonDescription = TEXT("");
-
-					ExecutingWork->ExecutionAttempts++;
-
 					TRACE_CPUPROFILER_EVENT_SCOPE_TEXT(ExecutingWork->Name)
 
 					FRemoteTransactionId RequestId = ExecutingWork->RequestId;
@@ -444,6 +440,15 @@ public:
 						AutoRTFM::Transact(
 							[this, LocalIterationNumber]()
 							{
+								UE_AUTORTFM_OPEN
+								{
+									AbortReason = UE::RemoteExecutor::ERemoteExecutorAbortReason::Unspecified;
+									AbortReasonDescription.Reset();
+									ExecutingWork->NumAbortsRequested = 0;
+
+									ExecutingWork->ExecutionAttempts++;
+								};
+
 								// You can set a breakpoint inside the abort handler below to see what caused the abort,
 								// or you can rely on the log output which only catches livelocks
 								const uint32 ExecutionAttempts = ExecutingWork->ExecutionAttempts;
@@ -458,7 +463,6 @@ public:
 									}
 								};
 
-								ExecutingWork->NumAbortsRequested = 0;
 								ExecutingWork->Work();
 
 								// Did we do something that should late-abort this transaction?
@@ -467,6 +471,10 @@ public:
 									UE_LOG(LogRemoteExec, Warning, TEXT("We ended a transaction marked for abort. We expected to abort before this point (%d abort requests)."), ExecutingWork->NumAbortsRequested);
 									AutoRTFM::AbortTransaction();
 								}
+								else
+								{
+									ensureMsgf(AbortReason == UE::RemoteExecutor::ERemoteExecutorAbortReason::Unspecified && AbortReasonDescription.IsEmpty(), TEXT("RemoteExecutor: AbortReason %s specified during transaction without aborting"), *AbortReasonDescription);
+								}
 
 								if (ExecutingWork->bRequiresMultiServerCommit)
 								{
@@ -768,6 +776,8 @@ public:
 		}
 
 		PendingWorkLists.Empty();
+		AbortReason = UE::RemoteExecutor::ERemoteExecutorAbortReason::Unspecified;
+		AbortReasonDescription.Reset();
 	}
 };
 
@@ -845,7 +855,8 @@ void AbortCurrentTransaction(EAbortTransactionBehavior AbortBehavior, ERemoteExe
 		GRemoteExecutor.AbortReason = AbortReason;
 		GRemoteExecutor.AbortReasonDescription = AbortDescription;
 
-		if (FRemoteExecutorWork* CurrentWorkItem = GRemoteExecutor.ExecutingWork)
+		FRemoteExecutorWork* CurrentWorkItem = GRemoteExecutor.ExecutingWork;
+		if (ensure(CurrentWorkItem))
 		{
 			CurrentWorkItem->NumAbortsRequested++;
 		}
@@ -871,18 +882,6 @@ void AbortTransactionAndAbandonWork(FStringView Description)
 	AbortCurrentTransaction(EAbortTransactionBehavior::Immediately, ERemoteExecutorAbortReason::AbandonWork, Description);
 }
 
-void RollbackTransactionRequiresDependencies(FStringView Description)
-{
-	AutoRTFM::UnreachableIfClosed();
-
-	GRemoteExecutor.AbortReason = ERemoteExecutorAbortReason::RequiresDependencies;
-	GRemoteExecutor.AbortReasonDescription = Description;
-
-#if UE_AUTORTFM
-	AutoRTFM::ForTheRuntime::RollbackTransaction();
-#endif
-}
-
 void TransactionRequiresMultiServerCommit(FStringView Description)
 {
 	if (GRemoteExecutor.ExecutingWork && !GRemoteExecutor.ExecutingWork->bRequiresMultiServerCommit)
diff --git a/Engine/Source/Runtime/CoreUObject/Public/UObject/RemoteExecutor.h b/Engine/Source/Runtime/CoreUObject/Public/UObject/RemoteExecutor.h
--- a/Engine/Source/Runtime/CoreUObject/Public/UObject/RemoteExecutor.h
+++ b/Engine/Source/Runtime/CoreUObject/Public/UObject/RemoteExecutor.h
@@ -271,7 +271,6 @@ namespace UE::RemoteExecutor
 	COREUOBJECT_API void AbortCurrentTransaction(EAbortTransactionBehavior AbortBehavior, ERemoteExecutorAbortReason AbortReason, FStringView AbortDescription);
 
 	COREUOBJECT_API void AbortTransactionRequiresDependencies(FStringView Description);
-	COREUOBJECT_API void RollbackTransactionRequiresDependencies(FStringView Description);
 	COREUOBJECT_API void AbortTransactionAndAbandonWork(FStringView Description);
 
 	COREUOBJECT_API void TransactionRequiresMultiServerCommit(FStringView Description);
