Add package name as well as object path for external objects/actors in FAssetDataToken

diff --git a/Engine/Source/Runtime/CoreUObject/Private/AssetRegistry/AssetDataToken.cpp b/Engine/Source/Runtime/CoreUObject/Private/AssetRegistry/AssetDataToken.cpp
--- a/Engine/Source/Runtime/CoreUObject/Private/AssetRegistry/AssetDataToken.cpp
+++ b/Engine/Source/Runtime/CoreUObject/Private/AssetRegistry/AssetDataToken.cpp
@@ -2,6 +2,8 @@
 
 #include "AssetRegistry/AssetDataToken.h" 
 
+#define LOCTEXT_NAMESPACE "AssetRegistry"
+
 FOnMessageTokenActivated FAssetDataToken::DefaultMessageTokenActivated;
 FAssetDataToken::FOnGetDisplayName FAssetDataToken::DefaultGetAssetDisplayName;
 
@@ -23,9 +25,19 @@ FAssetDataToken::FAssetDataToken(FPrivateToken, const FAssetData& InAssetData, c
 		{
 			CachedText = DefaultGetAssetDisplayName.Execute(InAssetData, false);
 		}
+#if WITH_EDITORONLY_DATA
+		else if (InAssetData.GetOptionalOuterPathName() != FName{})
+		{
+			// Provide both path and stored package
+			CachedText = FText::Format(LOCTEXT("AssetDataToken_ExternalPackageObject", "{0} ({1})"), 
+				FText::FromString(InAssetData.GetExportTextName()),
+				FText::FromStringView(WriteToString<256>(InAssetData.PackageName).ToView())
+			);
+		}
+#endif
 		else 
 		{
-			CachedText = FText::FromString(InAssetData.GetObjectPathString());
+			CachedText = FText::FromString(InAssetData.GetExportTextName());
 		}
 	}
 }
@@ -40,4 +52,6 @@ const FOnMessageTokenActivated& FAssetDataToken::GetOnMessageTokenActivated() co
 	{
 		return DefaultMessageTokenActivated;
 	}
-}
\ No newline at end of file
+}
+
+#undef LOCTEXT_NAMESPACE
\ No newline at end of file
