PackageFileSummary: Bugfix for incorrect error message in AssetRegistry's PackageReader when a packagefile on disk is from a newer version than the current executable. We were incorrectly resetting the version number of the PackageFileSummary and losing the information about whether it was too old or too new.
#rnx

#FailedValidations Preflight

diff --git a/Engine/Source/Runtime/CoreUObject/Private/Serialization/AsyncLoading.cpp b/Engine/Source/Runtime/CoreUObject/Private/Serialization/AsyncLoading.cpp
--- a/Engine/Source/Runtime/CoreUObject/Private/Serialization/AsyncLoading.cpp
+++ b/Engine/Source/Runtime/CoreUObject/Private/Serialization/AsyncLoading.cpp
@@ -7834,7 +7834,8 @@ void FAsyncArchive::ReadCallback(bool bWasCancelled, IAsyncReadRequest* Request)
 			FBufferReader Ar(Mem, FMath::Min<int64>(FMaxPackageSummarySize::Value, FileSize),/*bInFreeOnClose=*/ false, /*bIsPersistent=*/ true);
 			FPackageFileSummary Sum;
 			Ar << Sum;
-			if (Ar.IsError() || Sum.TotalHeaderSize > FileSize || Sum.IsFileVersionTooOld())
+			if (Ar.IsError() || Sum.TotalHeaderSize > FileSize
+				|| Sum.IsFileVersionTooOld() || (IsEnforcePackageCompatibleVersionCheck() && Sum.IsFileVersionTooNew()))
 			{
 				SetError();
 				LoadError = ELoadError::CorruptData;
diff --git a/Engine/Source/Runtime/CoreUObject/Private/UObject/LinkerLoad.cpp b/Engine/Source/Runtime/CoreUObject/Private/UObject/LinkerLoad.cpp
--- a/Engine/Source/Runtime/CoreUObject/Private/UObject/LinkerLoad.cpp
+++ b/Engine/Source/Runtime/CoreUObject/Private/UObject/LinkerLoad.cpp
@@ -1552,7 +1552,8 @@ FLinkerLoad::ELinkerStatus FLinkerLoad::SerializePackageFileSummaryInternal()
 	}
 
 	// Don't load packages that are only compatible with an engine version newer than the current one.
-	if (bLoaderNeedsEngineVersionChecks && IsEnforcePackageCompatibleVersionCheck() && !FEngineVersion::Current().IsCompatibleWith(Summary.CompatibleWithEngineVersion))
+	if (bLoaderNeedsEngineVersionChecks && IsEnforcePackageCompatibleVersionCheck() &&
+		(Summary.IsFileVersionTooNew() || !FEngineVersion::Current().IsCompatibleWith(Summary.CompatibleWithEngineVersion)))
 	{
 		// Send the warning to the game thread as slate is not thread-safe
 		Async(EAsyncExecution::TaskGraphMainThread,
diff --git a/Engine/Source/Runtime/CoreUObject/Private/UObject/PackageFileSummary.cpp b/Engine/Source/Runtime/CoreUObject/Private/UObject/PackageFileSummary.cpp
--- a/Engine/Source/Runtime/CoreUObject/Private/UObject/PackageFileSummary.cpp
+++ b/Engine/Source/Runtime/CoreUObject/Private/UObject/PackageFileSummary.cpp
@@ -7,6 +7,7 @@
 #include "Runtime/Launch/Resources/Version.h"
 #include "Serialization/StructuredArchive.h"
 #include "UObject/Linker.h"
+#include "UObject/LinkerLoad.h"
 #include "UObject/UObjectGlobals.h"
 
 FPackageFileSummary::FPackageFileSummary()
@@ -162,14 +163,13 @@ void operator<<(FStructuredArchive::FSlot Slot, FPackageFileSummary& Sum)
 					Sum.FileVersionLicenseeUE = GPackageFileLicenseeUEVersion;
 				}
 
-				if (Sum.FileVersionUE < VER_UE4_OLDEST_LOADABLE_PACKAGE || Sum.IsFileVersionTooNew())
+				if (Sum.FileVersionUE < VER_UE4_OLDEST_LOADABLE_PACKAGE 
+					|| (IsEnforcePackageCompatibleVersionCheck() && Sum.IsFileVersionTooNew()))
 				{
 					// We early exit as soon as possible if the fileversion is too old or too new so that changes can
 					// be made to the expected format after this point, and still have older versions of the editor
 					// gracefully abort the load of the newer package rather than crashing. They might crash on format
 					// differences if they try to serialize arrays from the wrong offset.
-					Sum.FileVersionUE.Reset();
-					Sum.FileVersionLicenseeUE = 0;
 					return;
 				}
 
