[StructUtils] Remove duplicated function

* Function was added on Main from external users, and was also added by us on 5.7 stream. They do the same thing, so remove the duplicate and tell external users to use the new one.

#rnx

diff --git a/Engine/Source/Runtime/CoreUObject/Private/StructUtils/PropertyBag.cpp b/Engine/Source/Runtime/CoreUObject/Private/StructUtils/PropertyBag.cpp
--- a/Engine/Source/Runtime/CoreUObject/Private/StructUtils/PropertyBag.cpp
+++ b/Engine/Source/Runtime/CoreUObject/Private/StructUtils/PropertyBag.cpp
@@ -2060,11 +2060,6 @@ void FInstancedPropertyBag::CopyMatchingValuesByID(const FInstancedPropertyBag&
 	}
 }
 
-void FInstancedPropertyBag::CopyMatchingValuesByIDWithOverrides(const FInstancedPropertyBag& NewDescs, TConstArrayView<FGuid> OverriddenPropertyIDs)
-{
-	UE::StructUtils::Private::CopyMatchingValuesByIDWithOverrides(NewDescs.Value, Value, OverriddenPropertyIDs);
-}
-
 void FInstancedPropertyBag::CopyMatchingValuesByName(const FInstancedPropertyBag& Other, TOptional<TConstArrayView<FName>> OptionalPropertyNamesSubset)
 {
 	UE::StructUtils::Private::CopyMatchingValues(
diff --git a/Engine/Source/Runtime/CoreUObject/Public/StructUtils/PropertyBag.h b/Engine/Source/Runtime/CoreUObject/Public/StructUtils/PropertyBag.h
--- a/Engine/Source/Runtime/CoreUObject/Public/StructUtils/PropertyBag.h
+++ b/Engine/Source/Runtime/CoreUObject/Public/StructUtils/PropertyBag.h
@@ -456,18 +456,10 @@ struct FInstancedPropertyBag
 	 * Copies matching property values from another bag of potentially mismatching layout.
 	 * The properties are matched between the bags based on the property ID.
 	 * @param Other Reference to the bag to copy the values from
-	 * @param OptionalPropertyIdsSubset Optional array of property ids subset to copy from NewDescs
+	 * @param OptionalPropertyIdsSubset Optional array of property ids subset to copy from Other
 	 */
 	UE_API void CopyMatchingValuesByID(const FInstancedPropertyBag& Other, TOptional<TConstArrayView<FGuid>> OptionalPropertyIdsSubset = {});
 
-	/**
-	 * Copies matching property values from another bag of potentially mismatching layout.
-	 * The properties are matched between the bags based on the property ID.
-	 * @param NewDescs Reference to the bag to copy the values from
-	 * @param OverriddenPropertyIDs Only the matching properties, whose ID is found in this array are copied (if it is empty, nothing is copied)
-	 */
-	UE_API void CopyMatchingValuesByIDWithOverrides(const FInstancedPropertyBag& NewDescs, TConstArrayView<FGuid> OverriddenPropertyIDs);
-
 	/**
 	 * Copies matching property values from another bag of potentially mismatching layout.
 	 * The properties are matched between the bags based on the property name.
