[StructUtils] Fixed incorrect logic in reordering properties in a property bag

* Also converted std::move to MoveTemp to maintain compliance with code standards

[fixed merge error from 38.00->39.00]

#rnx

diff --git a/Engine/Source/Runtime/CoreUObject/Private/StructUtils/PropertyBag.cpp b/Engine/Source/Runtime/CoreUObject/Private/StructUtils/PropertyBag.cpp
--- a/Engine/Source/Runtime/CoreUObject/Private/StructUtils/PropertyBag.cpp
+++ b/Engine/Source/Runtime/CoreUObject/Private/StructUtils/PropertyBag.cpp
@@ -1652,36 +1652,21 @@ namespace UE::StructUtils::Private
 			return EPropertyBagAlterationResult::TargetPropertyNotFound;
 		}
 
-		TArray<FPropertyBagPropertyDesc> NewDescs;
-		NewDescs.Reserve(Descs.Num());
+		FPropertyBagPropertyDesc SourceDesc = MoveTemp(Descs[SourcePropertyIndex]);
+		Descs.RemoveAt(SourcePropertyIndex);
 
-		const int32 Min = FMath::Min(SourcePropertyIndex, TargetPropertyIndex);
-		const int32 SourceSpacer = SourcePropertyIndex < TargetPropertyIndex;
-
-		// Elements before the source and target properties.
-		int32 ReadIndex;
-		for (ReadIndex = 0; ReadIndex < Min; ++ReadIndex)
+		int32 NewTargetPropertyIndex = TargetPropertyIndex;
+		if (SourcePropertyIndex < TargetPropertyIndex)
 		{
-			NewDescs.Emplace(MoveTemp(Descs[ReadIndex]));
+			--NewTargetPropertyIndex;
 		}
 
-		// Between the source and target properties (skip source if needed)
-		for (ReadIndex += SourceSpacer; ReadIndex < TargetPropertyIndex; ++ReadIndex)
-		{
-			NewDescs.Emplace(MoveTemp(Descs[ReadIndex]));
-		}
-
-		// Source and target property
-		NewDescs.Emplace(MoveTemp(Descs[bInsertBefore ? SourcePropertyIndex : TargetPropertyIndex]));
-		NewDescs.Emplace(MoveTemp(Descs[bInsertBefore ? TargetPropertyIndex : SourcePropertyIndex]));
-
-		// The remaining elements
-		for (++ReadIndex; ReadIndex < Descs.Num(); ++ReadIndex)
+		if (!bInsertBefore)
 		{
-			NewDescs.Emplace(MoveTemp(Descs[ReadIndex]));
+			++NewTargetPropertyIndex;
 		}
 
-		Descs = MoveTemp(NewDescs);
+		Descs.Insert(MoveTemp(SourceDesc), NewTargetPropertyIndex);
 
 		return EPropertyBagAlterationResult::Success;
 	}
