Make these two CoreUObject level serializers compatible with UE_WITH_OBJECT_HANDLE_LATE_RESOLVE

#rnx

diff --git a/Engine/Source/Runtime/CoreUObject/Private/StructUtils/InstancedStruct.cpp b/Engine/Source/Runtime/CoreUObject/Private/StructUtils/InstancedStruct.cpp
--- a/Engine/Source/Runtime/CoreUObject/Private/StructUtils/InstancedStruct.cpp
+++ b/Engine/Source/Runtime/CoreUObject/Private/StructUtils/InstancedStruct.cpp
@@ -151,7 +151,7 @@ bool FInstancedStruct::Serialize(FArchive& Ar, UStruct* DefaultsStruct, const vo
 	if (Ar.IsLoading())
 	{
 		// UScriptStruct type
-		UScriptStruct* SerializedScriptStruct = nullptr;
+		TObjectPtr<UScriptStruct> SerializedScriptStruct = nullptr;
 		Ar << SerializedScriptStruct;
 		if (SerializedScriptStruct)
 		{
diff --git a/Engine/Source/Runtime/CoreUObject/Private/UObject/FieldPath.cpp b/Engine/Source/Runtime/CoreUObject/Private/UObject/FieldPath.cpp
--- a/Engine/Source/Runtime/CoreUObject/Private/UObject/FieldPath.cpp
+++ b/Engine/Source/Runtime/CoreUObject/Private/UObject/FieldPath.cpp
@@ -359,7 +359,7 @@ FArchive& operator<<(FArchive& Ar, FFieldPath& InOutPropertyPath)
 		}
 		if (Ar.CustomVer(FFortniteMainBranchObjectVersion::GUID) >= FFortniteMainBranchObjectVersion::FFieldPathOwnerSerialization || Ar.CustomVer(FReleaseObjectVersion::GUID) >= FReleaseObjectVersion::FFieldPathOwnerSerialization)
 		{
-			UStruct* SerializedOwner = InOutPropertyPath.ResolvedOwner.Get();
+			TObjectPtr<UStruct> SerializedOwner = InOutPropertyPath.ResolvedOwner.Get();
 			Ar << SerializedOwner;
 			InOutPropertyPath.ResolvedOwner = SerializedOwner;
 			if (!SerializedOwner)
