[Cleanup] Move PostInitPropertiesCheck back next to PostInitProperties.

This has drifted away over time but belongs immediately after PostInitProperties.

I've also marked the array as non-shrinking instead of individually annotating each Pop call with `EAllowShrinking::No`.

diff --git a/Engine/Source/Runtime/CoreUObject/Private/UObject/UObjectGlobals.cpp b/Engine/Source/Runtime/CoreUObject/Private/UObject/UObjectGlobals.cpp
--- a/Engine/Source/Runtime/CoreUObject/Private/UObject/UObjectGlobals.cpp
+++ b/Engine/Source/Runtime/CoreUObject/Private/UObject/UObjectGlobals.cpp
@@ -3975,7 +3975,7 @@ void UObject::PostInitProperties()
 	};
 	AutoRTFM::PushOnAbortHandler(&ThreadContext.PostInitPropertiesCheck, [this]
 	{ 
-		UObject* Popped = FUObjectThreadContext::Get().PostInitPropertiesCheck.Pop(EAllowShrinking::No);
+		UObject* Popped = FUObjectThreadContext::Get().PostInitPropertiesCheck.Pop();
 		check(this == Popped);
 	});
 #endif
@@ -4443,6 +4443,21 @@ void FObjectInitializer::PostConstructInit()
 		Obj->PostInitProperties();
 	}
 
+#if !(UE_BUILD_SHIPPING || UE_BUILD_TEST)
+	FUObjectThreadContext& ThreadContext = FUObjectThreadContext::Get();
+	// Always write to FUObjectThreadContext from the open, as writing to a
+	// memory location in both the open and closed within the same transaction
+	// can lead to state corruption.
+	AutoRTFM::PopOnAbortHandler(&ThreadContext.PostInitPropertiesCheck);
+	UE_AUTORTFM_OPEN
+	{
+		if (ThreadContext.PostInitPropertiesCheck.IsEmpty() || ThreadContext.PostInitPropertiesCheck.Pop() != Obj)
+		{
+			UE_LOG(LogUObjectGlobals, Fatal, TEXT("%s failed to route PostInitProperties. Call Super::PostInitProperties() in %s::PostInitProperties()."), *Obj->GetClass()->GetName(), *Obj->GetClass()->GetName());
+		}
+	};
+#endif // !(UE_BUILD_SHIPPING || UE_BUILD_TEST)
+
 	// Only invoke the class initializer if subobject instancing is enabled. Since object properties may have been default-initialized to
 	// reference archetypes before instancing, archetypes could end up being mutated by compiled bytecode if executed before fixup occurs.
 	if (bAllowInstancing)
@@ -4457,21 +4472,6 @@ void FObjectInitializer::PostConstructInit()
 	}
 #endif
 
-#if !(UE_BUILD_SHIPPING || UE_BUILD_TEST)
-	FUObjectThreadContext& ThreadContext = FUObjectThreadContext::Get();
-	// Always write to FUObjectThreadContext from the open, as writing to a 
-	// memory location in both the open and closed within the same transaction
-	// can lead to state corruption.
-	AutoRTFM::PopOnAbortHandler(&ThreadContext.PostInitPropertiesCheck);
-	UE_AUTORTFM_OPEN
-	{
-		if (!ThreadContext.PostInitPropertiesCheck.Num() || (ThreadContext.PostInitPropertiesCheck.Pop(EAllowShrinking::No) != Obj))
-		{
-			UE_LOG(LogUObjectGlobals, Fatal, TEXT("%s failed to route PostInitProperties. Call Super::PostInitProperties() in %s::PostInitProperties()."), *Obj->GetClass()->GetName(), *Obj->GetClass()->GetName());
-		}
-	};
-#endif // !(UE_BUILD_SHIPPING || UE_BUILD_TEST)
-
 #if !USE_CIRCULAR_DEPENDENCY_LOAD_DEFERRING
 	if (!Obj->HasAnyFlags(RF_NeedLoad) 
 #else 
diff --git a/Engine/Source/Runtime/CoreUObject/Public/UObject/UObjectThreadContext.h b/Engine/Source/Runtime/CoreUObject/Public/UObject/UObjectThreadContext.h
--- a/Engine/Source/Runtime/CoreUObject/Public/UObject/UObjectThreadContext.h
+++ b/Engine/Source/Runtime/CoreUObject/Public/UObject/UObjectThreadContext.h
@@ -112,9 +112,9 @@ public:
 
 #if !(UE_BUILD_SHIPPING || UE_BUILD_TEST)
 	/** Stack to ensure that PostInitProperties is routed through Super:: calls. **/
-	TArray<UObject*, TInlineAllocator<16>> PostInitPropertiesCheck;
-	/** Used to verify that the Super::PostLoad chain is intact.			*/
-	TArray<UObject*, TInlineAllocator<16> > DebugPostLoad;
+	TArray<UObject*, TInlineAllocator<16, FNonshrinkingAllocator>> PostInitPropertiesCheck;
+	/** Used to verify that the Super::PostLoad chain is intact. */
+	TArray<UObject*, TInlineAllocator<16>> DebugPostLoad;
 #endif
 #if WITH_EDITORONLY_DATA
 	UE_DEPRECATED(5.5, "No longer used; skiponlyeditoronly is used instead and tracks editoronly references via savepackage results.")
