[CoreUObject]
* Fixed non unity compile error

diff --git a/Engine/Source/Runtime/CoreUObject/Private/UObject/ObjectRedirector.cpp b/Engine/Source/Runtime/CoreUObject/Private/UObject/ObjectRedirector.cpp
--- a/Engine/Source/Runtime/CoreUObject/Private/UObject/ObjectRedirector.cpp
+++ b/Engine/Source/Runtime/CoreUObject/Private/UObject/ObjectRedirector.cpp
@@ -12,6 +12,7 @@
 #include "Templates/Casts.h"
 #include "UObject/GarbageCollectionSchema.h"
 #include "UObject/PropertyPortFlags.h"
+#include "UObject/RegisterCompiledInObjects.h"
 #include "UObject/UnrealType.h"
 
 /*-----------------------------------------------------------------------------
