[ZenLoader]
- Fix potential use-after-free when using the multithreaded loader

diff --git a/Engine/Source/Runtime/CoreUObject/Private/Serialization/AsyncLoading2.cpp b/Engine/Source/Runtime/CoreUObject/Private/Serialization/AsyncLoading2.cpp
--- a/Engine/Source/Runtime/CoreUObject/Private/Serialization/AsyncLoading2.cpp
+++ b/Engine/Source/Runtime/CoreUObject/Private/Serialization/AsyncLoading2.cpp
@@ -3277,7 +3277,7 @@ struct FAsyncLoadingThreadState2
 		return bTimeLimitExceeded;
 	}
 
-	bool UseTimeLimit()
+	bool UseTimeLimit() const
 	{
 		return bUseTimeLimit;
 	}
@@ -3292,7 +3292,7 @@ struct FAsyncLoadingThreadState2
 		bIsActive.store(false, std::memory_order_relaxed);
 	}
 
-	bool IsActive()
+	bool IsActive() const
 	{
 		return bIsActive.load(std::memory_order_relaxed);
 	}
@@ -3310,10 +3310,12 @@ struct FAsyncLoadingThreadState2
 	TAsyncAtomic<bool> bIsActive = false;
 	bool bShouldFireNodes = true;
 	bool bUseTimeLimit = false;
+	
 	double TimeLimit = 0.0;
 	double StartTime = 0.0;
 	double LastTestTime = -1.0;
 	static uint32 TlsSlot;
+	const bool bIsWorkerThread = false; // can only be set at construction time
 };
 
 uint32 FAsyncLoadingThreadState2::TlsSlot;
@@ -5924,7 +5926,8 @@ void FEventLoadNode2::ParallelLoadingWork()
 #endif
 
 	FTaskTagScope Scope(ETaskTag::EParallelLoadingThread);
-	FAsyncLoadingThreadState2 LocalThreadState;
+	FAsyncLoadingThreadState2 LocalThreadState { .bIsWorkerThread = true };
+
 	FAsyncLoadingThreadState2Scope WorkerStateScope(LocalThreadState);
 
 	FUObjectThreadContext& ThreadContext = FUObjectThreadContext::Get();
@@ -5941,6 +5944,9 @@ void FEventLoadNode2::ParallelLoadingWork()
 	
 	ThreadContext.AsyncPackage = nullptr;
 
+	// We should never destroy a thread state that still has nodes to fire
+	check(LocalThreadState.NodesToFire.Num() == 0);
+
 	// This is never used for zen loader, everything is handled in NotifyConstructedDuringAsyncLoading
 	// Reset to avoid accumulation and shutdown error complaining about objects being present in the context.
 	ThreadContext.GetSerializeContext()->PRIVATE_GetObjectsLoadedInternalUseOnly().Reset();
@@ -5988,9 +5994,18 @@ void FEventLoadNode2::Fire(FAsyncLoadingThreadState2* ThreadState)
 
 EEventLoadNodeExecutionResult FEventLoadNode2::Execute(FAsyncLoadingThreadState2& ThreadState)
 {
-	//TRACE_CPUPROFILER_EVENT_SCOPE(ExecuteEvent);
 	ThreadState.MarkAsActive();
 
+	// This will avoid a race where the function being called releases a barrier
+	// that allows the package to move forward on other threads and get deleted
+	// before we have the time to process dependencies. This is only required
+	// when outside ALT since ALT is responsible for deleting packages so it
+	// cannot race with itself.
+	if (ThreadState.bIsWorkerThread)
+	{
+		Package->AddRef();
+	}
+
 	check(BarrierCount.load(std::memory_order_relaxed) == 0);
 	EEventLoadNodeExecutionResult Result;
 	{
@@ -6002,6 +6017,11 @@ EEventLoadNodeExecutionResult FEventLoadNode2::Execute(FAsyncLoadingThreadState2
 	{
 		ProcessDependencies(ThreadState);
 	}
+
+	if (ThreadState.bIsWorkerThread)
+	{
+		Package->Release();
+	}
 	return Result;
 }
 
