Removed deprecated functions from FWeakObjectPtr.

diff --git a/Engine/Source/Runtime/CoreUObject/Public/UObject/WeakObjectPtr.h b/Engine/Source/Runtime/CoreUObject/Public/UObject/WeakObjectPtr.h
--- a/Engine/Source/Runtime/CoreUObject/Public/UObject/WeakObjectPtr.h
+++ b/Engine/Source/Runtime/CoreUObject/Public/UObject/WeakObjectPtr.h
@@ -76,12 +76,6 @@ public:
 	}
 #endif
 
-	UE_DEPRECATED(5.6, "Constructing a FWeakObjectPtr from NULL has been deprecated - please use nullptr instead.")
-	UE_FORCEINLINE_HINT FWeakObjectPtr(int)
-	{
-		(*this) = nullptr;
-	}
-
 	/**  
 	 * Construct from an object pointer or something that can be implicitly converted to an object pointer
 	 * @param Object object to create a weak pointer to
@@ -367,13 +361,6 @@ public:
 	 */
 	COREUOBJECT_API void Serialize(FArchive& Ar);
 
-protected:
-	UE_DEPRECATED(5.1, "GetObjectIndex is now deprecated, and will be removed.")
-	UE_FORCEINLINE_HINT int32 GetObjectIndex() const
-	{
-		return ObjectIndex;
-	}
-
 private:
 	inline int32 GetObjectIndex_Private() const
 	{
