Fixed remove during iteration in UObjectMarks UnMarkAllObjects.

#tests PIE, editor repro steps from jira

diff --git a/Engine/Source/Runtime/CoreUObject/Private/UObject/UObjectMarks.cpp b/Engine/Source/Runtime/CoreUObject/Private/UObject/UObjectMarks.cpp
--- a/Engine/Source/Runtime/CoreUObject/Private/UObject/UObjectMarks.cpp
+++ b/Engine/Source/Runtime/CoreUObject/Private/UObject/UObjectMarks.cpp
@@ -24,6 +24,12 @@ struct FObjectMark
 		return Marks == OBJECTMARK_NOMARKS;
 	}
 
+	bool ClearMarks(FObjectMark InMarksToClear)
+	{
+		Marks = static_cast<EObjectMark>(Marks & ~InMarksToClear.Marks);
+		return IsDefault();
+	}
+
 	EObjectMark				Marks = OBJECTMARK_NOMARKS;
 };
 
@@ -114,6 +120,24 @@ public:
 		}
 	}
 
+	void RemoveFromAnnotations(const TAnnotation& Annotation)
+	{
+		if (!AnnotationMap.IsEmpty())
+		{
+			for (typename TMap<const UObjectBase*, TAnnotation>::TIterator It(AnnotationMap); It; ++It)
+			{
+				if (It.Value().ClearMarks(Annotation))
+				{
+					It.RemoveCurrent();
+				}
+			}
+			if (AnnotationMap.IsEmpty())
+			{
+				GUObjectArray.RemoveUObjectDeleteListener(this);
+			}
+		}
+	}
+
 	TAnnotation GetAnnotation(const UObjectBase* Object)
 	{
 		check(Object);
@@ -190,14 +214,7 @@ void UnMarkAllObjects(EObjectMark Marks)
 	}
 	else
 	{
-		const TMap<const UObjectBase *, FObjectMark>& Map = ThreadMarkAnnotation.GetAnnotationMap();
-		for (TMap<const UObjectBase *, FObjectMark>::TConstIterator It(Map); It; ++It)
-		{
-			if(It.Value().Marks & Marks)
-			{
-				ThreadMarkAnnotation.AddAnnotation((UObject*)It.Key(),FObjectMark(EObjectMark(It.Value().Marks & ~Marks)));
-			}
-		}
+		ThreadMarkAnnotation.RemoveFromAnnotations(Marks);
 	}
 }
 
