[Feature] Updated IDT api to allow users to get, preview, and edit the full set of transforms rather than simply appending to the list and saving immediately

diff --git a/Engine/Plugins/Experimental/InstanceDataObjectFixupTool/Source/Private/InstanceDataObjectFixupPanel.cpp b/Engine/Plugins/Experimental/InstanceDataObjectFixupTool/Source/Private/InstanceDataObjectFixupPanel.cpp
--- a/Engine/Plugins/Experimental/InstanceDataObjectFixupTool/Source/Private/InstanceDataObjectFixupPanel.cpp
+++ b/Engine/Plugins/Experimental/InstanceDataObjectFixupTool/Source/Private/InstanceDataObjectFixupPanel.cpp
@@ -70,11 +70,11 @@ static bool IsPathOverridden(UObject* Object, const FPropertyPath& Path)
 	return (Override != EOverriddenPropertyOperation::None);
 }
 
-FInstanceDataObjectFixupPanel::FInstanceDataObjectFixupPanel(
-	const TObjectPtr<UObject>& InstanceDataObject, const TObjectPtr<UObject>& InstanceDataObjectOwner, EViewFlags InViewFlags)
+
+FInstanceDataObjectFixupPanel::FInstanceDataObjectFixupPanel(const TObjectPtr<UObject>& InstanceDataObject, EViewFlags ViewFlags, TWeakPtr<TMap<FTopLevelAssetPath, UE::FInstanceDataTransformSet>> StagedTransforms)
 	: InstanceDataObject(InstanceDataObject)
-	, InstanceDataObjectOwner(InstanceDataObjectOwner)
-	, ViewFlags(InViewFlags)
+	, ViewFlags(ViewFlags)
+	, StagedTransforms(StagedTransforms)
 {
 }
 
@@ -85,19 +85,6 @@ FInstanceDataObjectFixupPanel::~FInstanceDataObjectFixupPanel()
 	ICoreProvider* DataStorage = GetMutableDataStorageFeature<ICoreProvider>(StorageFeatureName);
 	ICompatibilityProvider* DataStorageCompatibility = GetMutableDataStorageFeature<ICompatibilityProvider>(CompatibilityFeatureName);
 	UE::FPropertyBagRepository& Repository = UE::FPropertyBagRepository::Get();
-
-	if (DataStorageCompatibility != nullptr && DataStorage != nullptr)
-	{
-		if (!Repository.RequiresFixup(Repository.FindInstanceForDataObject(InstanceDataObject)))
-		{
-			static FName AlertName = FName("EntityLoosePropertiesErrorAlert");
-
-			// If a UObject isn't registered with TEDS, there's a chance its parent is registered and is the one
-			// with the alert column on it, so search upward until the nearest registered parent is found.
-			RowHandle Row = DataStorageCompatibility->FindRowWithCompatibleObject(InstanceDataObjectOwner ? InstanceDataObjectOwner.Get() : InstanceDataObject.Get());
-			Alerts::RemoveAlert(*DataStorage, Row, AlertName);
-		}
-	}
 }
 
 static bool RemoveCustomizationsWithLooseProperties(const FFieldVariant& FieldVariant, const TSharedPtr<IDetailsView>& DetailsView)
@@ -422,6 +409,21 @@ static int32 GetLastSharedStructSegment(const UE::FPropertyPathName& PathA, cons
 	return LastStructIndex;
 }
 
+UE::FInstanceDataTransformSet* FInstanceDataObjectFixupPanel::GetStagedTransformsForStruct(UStruct* ClassOrStruct)
+{
+	check(ClassOrStruct);
+	if (TSharedPtr<TMap<FTopLevelAssetPath, UE::FInstanceDataTransformSet>> Transforms = StagedTransforms.Pin())
+	{
+		FTopLevelAssetPath Key = ClassOrStruct->GetStructPathName();
+		if (UE::FInstanceDataTransformSet* Found = Transforms->Find(Key))
+		{
+			return Found;
+		}
+		return &Transforms->Add(Key, UE::FInstanceDataTransforms::Get().GetTransformSet(*ClassOrStruct));
+	}
+	return nullptr;
+}
+
 void FInstanceDataObjectFixupPanel::RedirectProperty(const FPropertyPath& From, const FPropertyPath& To)
 {
 	FUObjectSerializeContext* SerializeContext = FUObjectThreadContext::Get().GetSerializeContext();
@@ -455,15 +457,17 @@ void FInstanceDataObjectFixupPanel::RedirectProperty(const FPropertyPath& From,
 			SrcPath.Push(FromPathNames[0].GetSegment(SegmentIndex));
 		}
 	}
+	
+	if (UE::FInstanceDataTransformSet* Transforms = GetStagedTransformsForStruct(BaseStruct))
+	{
+		TSharedPtr<FJsonObject> Params = MakeShared<FJsonObject>();
+		Params->SetField(TEXT("src"), SrcPath.ToJsonValue());
+		Params->SetField(TEXT("dst"), DstPath.ToJsonValue());
+		Transforms->Operations.Add({ .OpCode = TEXT("RedirectProperty"), .Params = Params });
 
-	TSharedPtr<FJsonObject> Params = MakeShared<FJsonObject>();
-	Params->SetField(TEXT("src"), SrcPath.ToJsonValue());
-	Params->SetField(TEXT("dst"), DstPath.ToJsonValue());
-	UE::FInstanceDataTransforms::Get().WriteInstanceDataTransform(*BaseStruct, TEXT("RedirectProperty"), Params);
-	TObjectPtr<UObject> OwnerObjectPtr = const_cast<UObject*>(Owner);
-
-	// this will likely trigger a IDO reinstance which will cause a details panel refresh for the new object!
-	UE::FInstanceDataTransforms::Get().PatchInstanceDataObject(InstanceDataObject, OwnerObjectPtr);
+		TObjectPtr<UObject> OwnerObjectPtr = const_cast<UObject*>(Owner);
+		UE::FInstanceDataTransforms::Get().ApplyTransformSet(*Transforms, InstanceDataObject, OwnerObjectPtr);
+	}
 }
 
 void FInstanceDataObjectFixupPanel::OnRedirectProperty(FPropertyPath From, FPropertyPath To)
@@ -516,11 +520,16 @@ void FInstanceDataObjectFixupPanel::MarkForDelete(const FPropertyPath& CurrentPa
 		return;
 	}
 	const UObject* Owner = UE::FPropertyBagRepository::Get().FindInstanceForDataObject(InstanceDataObject);
-	TSharedPtr<FJsonObject> Params = MakeShared<FJsonObject>();
-	Params->SetField(TEXT("src"), srcPathNames[0].ToJsonValue());
-	UE::FInstanceDataTransforms::Get().WriteInstanceDataTransform(*Owner->GetClass(), TEXT("RemoveProperty"), Params);
-	TObjectPtr<UObject> OwnerObjectPtr = const_cast<UObject*>(Owner);
-	UE::FInstanceDataTransforms::Get().PatchInstanceDataObject(InstanceDataObject, OwnerObjectPtr);
+
+	if (UE::FInstanceDataTransformSet* Transforms = GetStagedTransformsForStruct(Owner->GetClass()))
+	{
+		TSharedPtr<FJsonObject> Params = MakeShared<FJsonObject>();
+		Params->SetField(TEXT("src"), srcPathNames[0].ToJsonValue());
+		Transforms->Operations.Add({ .OpCode = TEXT("RemoveProperty"), .Params = Params });
+
+		TObjectPtr<UObject> OwnerObjectPtr = const_cast<UObject*>(Owner);
+		UE::FInstanceDataTransforms::Get().ApplyTransformSet(*Transforms, InstanceDataObject, OwnerObjectPtr);
+	}
 }
 
 void FInstanceDataObjectFixupPanel::OnMarkForDelete(FPropertyPath Path)
diff --git a/Engine/Plugins/Experimental/InstanceDataObjectFixupTool/Source/Private/InstanceDataObjectFixupPanel.h b/Engine/Plugins/Experimental/InstanceDataObjectFixupTool/Source/Private/InstanceDataObjectFixupPanel.h
--- a/Engine/Plugins/Experimental/InstanceDataObjectFixupTool/Source/Private/InstanceDataObjectFixupPanel.h
+++ b/Engine/Plugins/Experimental/InstanceDataObjectFixupTool/Source/Private/InstanceDataObjectFixupPanel.h
@@ -6,6 +6,7 @@
 #include "AsyncDetailViewDiff.h"
 #include "UObject/Object.h"
 #include "Templates/SharedPointer.h"
+#include <UObject/InstanceDataTransforms.h>
 
 class SLinkableScrollBar;
 
@@ -26,8 +27,7 @@ public:
 		DefaultRightPanel = HideLooseProperties,
 	};
 	
-	FInstanceDataObjectFixupPanel(
-		const TObjectPtr<UObject>& InstanceDataObject, const TObjectPtr<UObject>& InstanceDataObjectOwner, EViewFlags ViewFlags);
+	FInstanceDataObjectFixupPanel(const TObjectPtr<UObject>& InstanceDataObject, EViewFlags ViewFlags, TWeakPtr<TMap<FTopLevelAssetPath, UE::FInstanceDataTransformSet>> StagedTransforms);
 	~FInstanceDataObjectFixupPanel();
 	
 	TSharedPtr<IDetailsView>& GenerateDetailsView(bool bScrollbarOnLeft = false);
@@ -68,12 +68,13 @@ private:
 		bool bWasHidden;
 	};
 
+	UE::FInstanceDataTransformSet* GetStagedTransformsForStruct(UStruct* ClassOrStruct);
+
 	void RedirectProperty(const FPropertyPath& From, const FPropertyPath& To);
 	// pass-by-copy version for delegates. Use RedirectProperty when possible
 	void OnRedirectProperty(FPropertyPath From, FPropertyPath To);
 
 	TObjectPtr<UObject> InstanceDataObject; // stores either InstanceDataObject property bag or just regular objects
-	TObjectPtr<UObject> InstanceDataObjectOwner; // if not null, then it stores the common parent for the instances.
 
 	TMap<FPropertyPath, FRevertInfo> RevertInfo;
 	TSet<FPropertyPath> MarkedForDelete;
@@ -81,4 +82,5 @@ private:
 	TWeakPtr<FAsyncDetailViewDiff> DiffAgainstLeft;
 	TWeakPtr<FAsyncDetailViewDiff> DiffAgainstRight;
 	EViewFlags ViewFlags = EViewFlags::None;
+	TWeakPtr<TMap<FTopLevelAssetPath, UE::FInstanceDataTransformSet>> StagedTransforms;
 };
diff --git a/Engine/Plugins/Experimental/InstanceDataObjectFixupTool/Source/Private/InstanceDataObjectFixupTool.cpp b/Engine/Plugins/Experimental/InstanceDataObjectFixupTool/Source/Private/InstanceDataObjectFixupTool.cpp
--- a/Engine/Plugins/Experimental/InstanceDataObjectFixupTool/Source/Private/InstanceDataObjectFixupTool.cpp
+++ b/Engine/Plugins/Experimental/InstanceDataObjectFixupTool/Source/Private/InstanceDataObjectFixupTool.cpp
@@ -72,6 +72,7 @@ BEGIN_SLATE_FUNCTION_BUILD_OPTIMIZATION
 void SInstanceDataObjectFixupTool::Construct(const FArguments& InArgs)
 {
 	SelectedClassPath = InArgs._SelectedClassPath;
+	StagedTransforms = InArgs._StagedTransforms;
 
 	ChildSlot
 	[
@@ -142,15 +143,6 @@ void SInstanceDataObjectFixupTool::SetDockTab(const TSharedRef<SDockTab>& DockTa
 	OwningDockTab = DockTab;
 }
 
-FReply SInstanceDataObjectFixupTool::OnConfirmClicked() const
-{
-	if (const TSharedPtr<SDockTab> DockTab = OwningDockTab.Pin())
-	{
-		DockTab->RequestCloseTab();
-	}
-	return FReply::Handled();
-}
-
 void SInstanceDataObjectFixupTool::ResetDiff()
 {
 	Border->SetContent(SNullWidget::NullWidget);
@@ -178,8 +170,8 @@ void SInstanceDataObjectFixupTool::GenerateDetailsViews()
 	constexpr int32 LeftIndex = 0;
 	constexpr int32 RightIndex = 1;
 
-	Panels[LeftIndex] = MakeShared<FInstanceDataObjectFixupPanel>(InstanceDataObject, nullptr, FInstanceDataObjectFixupPanel::EViewFlags::DefaultLeftPanel);
-	Panels[RightIndex] = MakeShared<FInstanceDataObjectFixupPanel>(InstanceDataObject, nullptr, FInstanceDataObjectFixupPanel::EViewFlags::DefaultRightPanel);
+	Panels[LeftIndex] = MakeShared<FInstanceDataObjectFixupPanel>(InstanceDataObject, FInstanceDataObjectFixupPanel::EViewFlags::DefaultLeftPanel, StagedTransforms);
+	Panels[RightIndex] = MakeShared<FInstanceDataObjectFixupPanel>(InstanceDataObject, FInstanceDataObjectFixupPanel::EViewFlags::DefaultRightPanel, StagedTransforms);
 
 	// generate the panels
 	Panels[LeftIndex]->GenerateDetailsView(true);
diff --git a/Engine/Plugins/Experimental/InstanceDataObjectFixupTool/Source/Public/InstanceDataObjectFixupTool.h b/Engine/Plugins/Experimental/InstanceDataObjectFixupTool/Source/Public/InstanceDataObjectFixupTool.h
--- a/Engine/Plugins/Experimental/InstanceDataObjectFixupTool/Source/Public/InstanceDataObjectFixupTool.h
+++ b/Engine/Plugins/Experimental/InstanceDataObjectFixupTool/Source/Public/InstanceDataObjectFixupTool.h
@@ -14,6 +14,10 @@ class SDetailsSplitter;
 class FInstanceDataObjectFixupPanel;
 struct FPropertySoftPath;
 class IStructureDetailsView;
+namespace UE
+{
+	class FInstanceDataTransformSet;
+}
 
 /**
  * This tool diffs multiple property bags of one format against the same number of property bags of another format.
@@ -22,9 +26,12 @@ class IStructureDetailsView;
 class SInstanceDataObjectFixupTool : public SCompoundWidget
 {
 public:
+	using FStagedTransformMap = TMap<FTopLevelAssetPath, UE::FInstanceDataTransformSet>;
+
 	SLATE_BEGIN_ARGS(SInstanceDataObjectFixupTool)
 	{}
 		SLATE_ARGUMENT(TWeakPtr<FTopLevelAssetPath>, SelectedClassPath)
+		SLATE_ARGUMENT(TWeakPtr<FStagedTransformMap>, StagedTransforms)
 
 	SLATE_END_ARGS()
 
@@ -39,7 +46,6 @@ public:
 	UE_API FReply OnAutoMarkForDeletion() const;
 	
 private:
-	UE_API FReply OnConfirmClicked() const;
 	void ResetDiff();
 
 	TWeakPtr<SDockTab> OwningDockTab;
@@ -48,6 +54,8 @@ private:
 
 	TSharedPtr<SDetailsSplitter> Splitter;
 
+	TWeakPtr<TMap<FTopLevelAssetPath, UE::FInstanceDataTransformSet>> StagedTransforms;
+
 	TSharedPtr<SBorder> Border;
 	TWeakPtr<FTopLevelAssetPath> SelectedClassPath;
 };
diff --git a/Engine/Source/Runtime/CoreUObject/Internal/UObject/InstanceDataTransforms.h b/Engine/Source/Runtime/CoreUObject/Internal/UObject/InstanceDataTransforms.h
--- a/Engine/Source/Runtime/CoreUObject/Internal/UObject/InstanceDataTransforms.h
+++ b/Engine/Source/Runtime/CoreUObject/Internal/UObject/InstanceDataTransforms.h
@@ -30,6 +30,24 @@ namespace UE::Private::IDT
 
 namespace UE
 {
+	// public facing list of operations that can be applied to an object or saved to disk
+	class FInstanceDataTransformSet
+	{
+	public:
+		using FTransformParams = UE::Private::IDT::FTransformParams;
+		using ETransformResult = UE::Private::IDT::ETransformResult;
+		struct FOperation
+		{
+			FString OpCode; // ie RemoveProperty, RedirectProperty, etc
+			FTransformParams Params;
+			ETransformResult operator()(TObjectPtr<UObject>& IDO, TObjectPtr<UObject>& Owner, const FPropertyPathName& BasePath) const;
+			ETransformResult operator()(TObjectPtr<UObject>& IDO, TObjectPtr<UObject>& Owner) const;
+		};
+
+		FTopLevelAssetPath StructPath;
+		TArray<FOperation> Operations;
+	};
+
 	class FInstanceDataTransforms
 	{
 	public:
@@ -39,18 +57,16 @@ namespace UE
 
 		COREUOBJECT_API static FInstanceDataTransforms& Get();
 
-		COREUOBJECT_API void WriteInstanceDataTransform(UStruct& ClassOrStruct, const FString& OpCode, const FTransformParams& Params);
 		COREUOBJECT_API TObjectPtr<UObject> PatchInstanceDataObject(TNotNull<UObject*> IDO, TNotNull<UObject*> Owner);
 
+		COREUOBJECT_API FInstanceDataTransformSet GetTransformSet(UStruct& ClassOrStruct);
+		COREUOBJECT_API TObjectPtr<UObject> ApplyTransformSet(const FInstanceDataTransformSet& TransformSet, TNotNull<UObject*> IDO, TNotNull<UObject*> Owner);
+		COREUOBJECT_API void SaveTransformSet(const FInstanceDataTransformSet& TransformSet);
+
 	private:
-		struct FTransform
-		{
-			FTransformParams Args;
-			FTransformFunction Func;
-		};
 
-		const TArray<TSharedPtr<FInstanceDataTransforms::FTransform>>* FindTransformsInternal(const FTopLevelAssetPath& ClassPath);
-		const TArray<TSharedPtr<FInstanceDataTransforms::FTransform>>* FindTransformsInternal(const TNotNull<UStruct*> Struct);
+		const FInstanceDataTransformSet* FindTransformsInternal(const FTopLevelAssetPath& StructPath);
+		const FInstanceDataTransformSet* FindTransformsInternal(const TNotNull<UStruct*> Struct);
 
 		void ReadConfigFile(const FString& ConfigFilepath);
 
@@ -60,8 +76,10 @@ namespace UE
 		void PatchClass(TObjectPtr<UObject>& IDO, TObjectPtr<UObject>& Owner);
 		void PatchStructs(TObjectPtr<UObject>& IDO, TObjectPtr<UObject>& Owner);
 
+		friend struct FInstanceDataTransformSet::FOperation;
+
 		TMap<FString, FDateTime> ConfigFileLastReadTime;
-		TMap<FTopLevelAssetPath, TArray<TSharedPtr<FTransform>>> StructPathToIDTs;
+		TMap<FTopLevelAssetPath, FInstanceDataTransformSet> StructPathToIDTs;
 		TMap<FString, FTransformFunction> RegisteredOperations;
 	};
 }
diff --git a/Engine/Source/Runtime/CoreUObject/Private/UObject/InstanceDataTransforms.cpp b/Engine/Source/Runtime/CoreUObject/Private/UObject/InstanceDataTransforms.cpp
--- a/Engine/Source/Runtime/CoreUObject/Private/UObject/InstanceDataTransforms.cpp
+++ b/Engine/Source/Runtime/CoreUObject/Private/UObject/InstanceDataTransforms.cpp
@@ -228,6 +228,55 @@ namespace UE::Private
 		}
 	}
 
+	void GetTreePathsForStruct(const FPropertyPathNameTree& UnknownPropertyTree, TArray<FPropertyPathName>& OutPaths, FPropertyTypeName StructType, FPropertyPathName& CurPath)
+	{
+		for (FPropertyPathNameTree::FConstIterator Itr = UnknownPropertyTree.CreateConstIterator(); Itr; ++Itr)
+		{
+			CurPath.Push(FPropertyPathNameSegment{ Itr.GetName(), Itr.GetType() });
+
+			FPropertyTypeName Type = Itr.GetType();
+			if (Type.GetName() == NAME_StructProperty && Type.GetParameter(0) == StructType)
+			{
+				OutPaths.Push(CurPath);
+			}
+			else if (Type.GetName() == NAME_ArrayProperty || Type.GetName() == NAME_SetProperty)
+			{
+				FPropertyTypeName InnerType = Type.GetParameter(0);
+				if (InnerType.GetName() == NAME_StructProperty && InnerType.GetParameter(0) == StructType)
+				{
+					OutPaths.Push(CurPath);
+				}
+			}
+			else if (Type.GetName() == NAME_MapProperty)
+			{
+				FPropertyTypeName KeyType = Type.GetParameter(0);
+				if (KeyType.GetName() == NAME_StructProperty && KeyType.GetParameter(0) == StructType)
+				{
+					OutPaths.Push(CurPath);
+				}
+				FPropertyTypeName ValType = Type.GetParameter(1);
+				if (ValType.GetName() == NAME_StructProperty && ValType.GetParameter(0) == StructType)
+				{
+					OutPaths.Push(CurPath);
+				}
+			}
+
+			if (const FPropertyPathNameTree* SubTree = Itr.GetNode().GetSubTree())
+			{
+				GetTreePathsForStruct(*SubTree, OutPaths, StructType, CurPath);
+			}
+			CurPath.Pop();
+		}
+	}
+
+	void GetTreePathsForStruct(const FPropertyPathNameTree& UnknownPropertyTree, TArray<FPropertyPathName>& OutPaths, UScriptStruct* Struct)
+	{
+		FPropertyPathName BasePath;
+		FPropertyTypeNameBuilder Builder;
+		Builder.AddPath(Struct);
+		GetTreePathsForStruct(UnknownPropertyTree, OutPaths, Builder.Build(), BasePath);
+	}
+
 	bool SubPathIsLeaf(const TSharedPtr<FPropertyPathNameTree>& UnknownPropertyTree, FPropertyPathName PathName)
 	{
 		PathName.Pop();
@@ -395,19 +444,79 @@ namespace UE::Private::IDT
 
 namespace UE
 {
+	FInstanceDataTransformSet::ETransformResult UE::FInstanceDataTransformSet::FOperation::operator()(TObjectPtr<UObject>& IDO, TObjectPtr<UObject>& Owner, const FPropertyPathName& BasePath) const
+	{
+		if (FInstanceDataTransforms::FTransformFunction* Function = FInstanceDataTransforms::Get().RegisteredOperations.Find(OpCode))
+		{
+			return (*Function)(IDO, Owner, BasePath, Params);
+		}
+		return FInstanceDataTransformSet::ETransformResult::Break;
+	}
+
+	FInstanceDataTransformSet::ETransformResult FInstanceDataTransformSet::FOperation::operator()(TObjectPtr<UObject>& IDO, TObjectPtr<UObject>& Owner) const
+	{
+		return this->operator()(IDO, Owner, {});
+	}
+
 	FInstanceDataTransforms& FInstanceDataTransforms::Get()
 	{
 		static FInstanceDataTransforms Instance;
 		return Instance;
 	}
 
-	void FInstanceDataTransforms::WriteInstanceDataTransform(UStruct& ClassOrStruct, const FString& OpCode, const FTransformParams& Params)
+	TObjectPtr<UObject> FInstanceDataTransforms::PatchInstanceDataObject(TNotNull<UObject*> IDO, TNotNull<UObject*> Owner)
+	{
+		TObjectPtr<UObject> ResultIdo = IDO;
+		TObjectPtr<UObject> ResultOwner = Owner;
+		PatchClass(ResultIdo, ResultOwner);
+		PatchStructs(ResultIdo, ResultOwner);
+
+		return ResultIdo;
+	}
+
+	FInstanceDataTransformSet FInstanceDataTransforms::GetTransformSet(UStruct& ClassOrStruct)
+	{
+		const FTopLevelAssetPath StructPath = ClassOrStruct.GetStructPathName();
+		if (const FInstanceDataTransformSet* TransformSet = FindTransformsInternal(StructPath))
+		{
+			return *TransformSet;
+		}
+		return FInstanceDataTransformSet{.StructPath= StructPath };
+	}
+
+	TObjectPtr<UObject> FInstanceDataTransforms::ApplyTransformSet(const FInstanceDataTransformSet& TransformSet, TNotNull<UObject*> IDO, TNotNull<UObject*> Owner)
+	{
+		TArray<FPropertyPathName> PathsToStruct;
+		if (TransformSet.StructPath == Owner->GetClass()->GetClassPathName())
+		{
+			PathsToStruct.Emplace();
+		}
+		else if (UScriptStruct* FoundStruct = FindObject<UScriptStruct>(TransformSet.StructPath))
+		{
+			if (TSharedPtr<FPropertyPathNameTree> UnknownPropertyTree = FUnknownPropertyTree(Owner).Find())
+			{
+				Private::GetTreePathsForStruct(*UnknownPropertyTree, PathsToStruct, FoundStruct);
+			}
+		}
+
+		TObjectPtr<UObject> ResultIdo = IDO;
+		TObjectPtr<UObject> ResultOwner = Owner;
+		for (const FPropertyPathName& BasePath : PathsToStruct)
+		{
+			for (const FInstanceDataTransformSet::FOperation& Operation : TransformSet.Operations)
+			{
+				Operation(ResultIdo, ResultOwner, BasePath);
+			}
+		}
+		return ResultIdo;
+	}
+
+	void FInstanceDataTransforms::SaveTransformSet(const FInstanceDataTransformSet& TransformSet)
 	{
-		FTopLevelAssetPath AssetPath = ClassOrStruct.GetStructPathName();
-		FString ConfigFilepath = Private::InstanceDataTransformDirectory(AssetPath);
+		FString ConfigFilepath = Private::InstanceDataTransformDirectory(TransformSet.StructPath);
 		if (ConfigFilepath.IsEmpty())
 		{
-			UE_LOG(LogInstanceDataTransforms, Error, TEXT("Couldn't find a valid IDT config path for '%s'. Perhaps it's not mounted?"), *AssetPath.ToString());
+			UE_LOG(LogInstanceDataTransforms, Error, TEXT("Couldn't find a valid IDT config path for '%s'. Perhaps it's not mounted?"), *TransformSet.StructPath.ToString());
 			return;
 		}
 
@@ -425,37 +534,35 @@ namespace UE
 		else
 		{
 			ConfigJson = MakeShared<FJsonObject>();
-			ConfigJson->SetStringField(TEXT("Package"), AssetPath.GetPackageName().ToString());
-			ConfigJson->SetArrayField(TEXT("Transforms"), {});
+			ConfigJson->SetStringField(TEXT("Package"), TransformSet.StructPath.GetPackageName().ToString());
+			ConfigJson->SetObjectField(TEXT("Transforms"), MakeShared<FJsonObject>());
 		}
 
-		const TArray<TSharedPtr<FJsonValue>>* ExistingTransformsPtr = nullptr;
-		if (!ConfigJson->TryGetArrayField(TEXT("Transforms"), ExistingTransformsPtr) || !ExistingTransformsPtr)
+		const TSharedPtr<FJsonObject>* TransformsMapPtr = nullptr;
+		if (!ConfigJson->TryGetObjectField(TEXT("Transforms"), TransformsMapPtr) || !TransformsMapPtr)
 		{
-			UE_LOG(LogInstanceDataTransforms, Error, TEXT("'Transforms' array was missing or malformed in '%s'"), *ConfigFilepath);
+			UE_LOG(LogInstanceDataTransforms, Error, TEXT("'Transforms' object was missing or malformed in '%s'"), *ConfigFilepath);
 			return;
 		}
 
-		// create the new transform operation
-		TArray<TSharedPtr<FJsonValue>> TransformsJson = *ExistingTransformsPtr;
-		TSharedPtr<FJsonObject> NewTransform = MakeShared<FJsonObject>();
-		if (ClassOrStruct.IsA<UClass>())
-		{
-			NewTransform->SetStringField(TEXT("Class"), AssetPath.GetAssetName().ToString());
-		}
-		else
-		{
-			NewTransform->SetStringField(TEXT("Struct"), AssetPath.GetAssetName().ToString());
-		}
-		NewTransform->SetStringField(TEXT("OpCode"), OpCode);
-		if (Params)
+		// create a json array for every transform in this struct
+		TArray<TSharedPtr<FJsonValue>> StructTransformsArray;
+		for (const FInstanceDataTransformSet::FOperation& Operation : TransformSet.Operations)
 		{
-			NewTransform->SetObjectField(TEXT("Params"), Params);
+			// create the new transform operation
+			TSharedPtr<FJsonObject> NewTransform = MakeShared<FJsonObject>();
+			NewTransform->SetStringField(TEXT("OpCode"), Operation.OpCode);
+			if (Operation.Params)
+			{
+				NewTransform->SetObjectField(TEXT("Params"), Operation.Params);
+			}
+
+			// append the new operation to the Transforms list
+			StructTransformsArray.Add(StaticCastSharedRef<FJsonValue>(MakeShared<FJsonValueObject>(NewTransform)));
 		}
 
-		// append the new operation to the Transforms list
-		TransformsJson.Add(StaticCastSharedRef<FJsonValue>(MakeShared<FJsonValueObject>(NewTransform)));
-		ConfigJson->SetArrayField(TEXT("Transforms"), TransformsJson);
+		// put StructTransformsArray in the map keyed to the struct name
+		(*TransformsMapPtr)->SetArrayField(TransformSet.StructPath.GetAssetName().ToString(), StructTransformsArray);
 
 		// update Json file on disk
 		JsonRaw.Empty();
@@ -473,17 +580,7 @@ namespace UE
 		}
 	}
 
-	TObjectPtr<UObject> FInstanceDataTransforms::PatchInstanceDataObject(TNotNull<UObject*> IDO, TNotNull<UObject*> Owner)
-	{
-		TObjectPtr<UObject> ResultIdo = IDO;
-		TObjectPtr<UObject> ResultOwner = Owner;
-		PatchClass(ResultIdo, ResultOwner);
-		PatchStructs(ResultIdo, ResultOwner);
-
-		return ResultIdo;
-	}
-
-	const TArray<TSharedPtr<FInstanceDataTransforms::FTransform>>* FInstanceDataTransforms::FindTransformsInternal(const FTopLevelAssetPath& ClassPath)
+	const FInstanceDataTransformSet* FInstanceDataTransforms::FindTransformsInternal(const FTopLevelAssetPath& ClassPath)
 	{
 		// the transform config for a class is found in <RootPackageDirectory>/Config/InstanceDataTransforms.json
 		// TODO: I need to some research on wheather any work needs to be done in regards to making sure this file is treated properly by source control and the content worker.
@@ -497,7 +594,6 @@ namespace UE
 				// check if the config file has been updated since the last time we read it
 				if (*LastReadTime < FileModificationTime)
 				{
-					StructPathToIDTs.Empty(); // kill all cache
 					ReadConfigFile(ConfigFilepath);
 				}
 			}
@@ -510,7 +606,7 @@ namespace UE
 		return StructPathToIDTs.Find(ClassPath);
 	}
 
-	const TArray<TSharedPtr<FInstanceDataTransforms::FTransform>>* FInstanceDataTransforms::FindTransformsInternal(const TNotNull<UStruct*> Struct)
+	const FInstanceDataTransformSet* FInstanceDataTransforms::FindTransformsInternal(const TNotNull<UStruct*> Struct)
 	{
 		if(Struct->GetOuter() != Struct->GetPackage())
 		{
@@ -544,59 +640,58 @@ namespace UE
 					return;
 				}
 
-				const TArray<TSharedPtr<FJsonValue>>* TransformsJson = nullptr;
-				if (!ConfigJson->TryGetArrayField(TEXT("Transforms"), TransformsJson) || !TransformsJson)
+				const TSharedPtr<FJsonObject>* TransformsJson = nullptr;
+				if (!ConfigJson->TryGetObjectField(TEXT("Transforms"), TransformsJson) || !TransformsJson)
 				{
-					UE_LOG(LogInstanceDataTransforms, Warning, TEXT("'Transforms' array was missing or malformed in '%s'"), *ConfigFilepath);
+					UE_LOG(LogInstanceDataTransforms, Warning, TEXT("'Transforms' object was missing or malformed in '%s'"), *ConfigFilepath);
 					return;
 				}
 
-				for (const TSharedPtr<FJsonValue>& TransformInfoJson : *TransformsJson)
+				for (TPair<FString, TSharedPtr<FJsonValue>>& TransformSetInfo : (*TransformsJson)->Values)
 				{
-					TSharedPtr<FJsonObject>* TransformInfo;
-					if (!TransformInfoJson->TryGetObject(TransformInfo) || !TransformInfo)
-					{
-						UE_LOG(LogInstanceDataTransforms, Warning, TEXT("Found a malfomed object in the 'Transforms' array in '%s'"), *ConfigFilepath);
-						continue;
-					}
+					FTopLevelAssetPath StructPath(PackagePath.GetPackageFName(), FName(TransformSetInfo.Key));
 
-					FTopLevelAssetPath StructPath;
-					if (FString ClassNameString; (*TransformInfo)->TryGetStringField(TEXT("Class"), ClassNameString))
-					{
-						FName ClassName = FName(ClassNameString);
-						StructPath = FTopLevelAssetPath(PackagePath.GetPackageFName(), ClassName);
-					}
-					else if (FString StructNameString; (*TransformInfo)->TryGetStringField(TEXT("Struct"), StructNameString))
-					{
-						FName StructName = FName(StructNameString);
-						StructPath = FTopLevelAssetPath(PackagePath.GetPackageFName(), StructName);
-					}
-					else
-					{
-						UE_LOG(LogInstanceDataTransforms, Warning, TEXT("an element of the 'Transforms' array in '%s' was missing a 'Class' or 'Struct' field"), *ConfigFilepath);
-						continue;
-					}
+					FInstanceDataTransformSet ReadTransformSet{ .StructPath = StructPath };
 
-					FString OpCode;
-					if (!(*TransformInfo)->TryGetStringField(TEXT("OpCode"), OpCode))
+					const TArray<TSharedPtr<FJsonValue>>* TransformOperationsJson;
+					if (!TransformSetInfo.Value->TryGetArray(TransformOperationsJson) || !TransformOperationsJson)
 					{
-						UE_LOG(LogInstanceDataTransforms, Warning, TEXT("an element of the 'Transforms' array in '%s' was missing a 'OpCode' field"), *ConfigFilepath);
-						continue;
+						UE_LOG(LogInstanceDataTransforms, Warning, TEXT("Transform array was malformed for '%s' in '%s'"), *TransformSetInfo.Key, *ConfigFilepath);
+						return;
 					}
-					const FTransformFunction* Func = RegisteredOperations.Find(OpCode);
-					if (!Func)
-					{
-						UE_LOG(LogInstanceDataTransforms, Warning, TEXT("The transform OpCode in '%s' was malformed. '%s' is not a valid operation."), *ConfigFilepath, *OpCode);
-						continue;
-					}
-					FTransformParams Params;
-					if (const FTransformParams* ParamsPtr; (*TransformInfo)->TryGetObjectField(TEXT("Params"), ParamsPtr))
+
+					for (const TSharedPtr<FJsonValue>& TransformOperationJson : *TransformOperationsJson)
 					{
-						Params = *ParamsPtr;
+						TSharedPtr<FJsonObject>* TransformInfo;
+						if (!TransformOperationJson->TryGetObject(TransformInfo) || !TransformInfo)
+						{
+							UE_LOG(LogInstanceDataTransforms, Warning, TEXT("Found a malfomed object in the Transform array for '%s' in '%s'"), *TransformSetInfo.Key, *ConfigFilepath);
+							continue;
+						}
+
+						FInstanceDataTransformSet::FOperation Operation;
+
+						if (!(*TransformInfo)->TryGetStringField(TEXT("OpCode"), Operation.OpCode))
+						{
+							UE_LOG(LogInstanceDataTransforms, Warning, TEXT("An element of the Transform array for '%s' in '%s' was missing a 'OpCode' field"), *TransformSetInfo.Key, *ConfigFilepath);
+							continue;
+						}
+						const FTransformFunction* Func = RegisteredOperations.Find(Operation.OpCode);
+						if (!Func)
+						{
+							UE_LOG(LogInstanceDataTransforms, Warning, TEXT("The transform OpCode in '%s' was malformed. '%s' is not a valid operation."), *ConfigFilepath, *Operation.OpCode);
+							continue;
+						}
+
+						if (const FTransformParams* ParamsPtr; (*TransformInfo)->TryGetObjectField(TEXT("Params"), ParamsPtr))
+						{
+							Operation.Params = *ParamsPtr;
+						}
+						ReadTransformSet.Operations.Emplace(MoveTemp(Operation));
 					}
-					TArray<TSharedPtr<FTransform>>& Transforms = StructPathToIDTs.FindOrAdd(StructPath);
-					Transforms.Add(MakeShared<FTransform>(Params, *Func));
+					StructPathToIDTs.Add(StructPath, ReadTransformSet);
 				}
+
 			}
 		}
 		ConfigFileLastReadTime.Add(ConfigFilepath, FDateTime::Now());
@@ -610,12 +705,12 @@ namespace UE
 	void FInstanceDataTransforms::PatchClass(TObjectPtr<UObject>& IDO, TObjectPtr<UObject>& Owner)
 	{
 		// find and run class IDTs
-		if (const TArray<TSharedPtr<FTransform>>* IDTs = FindTransformsInternal(Owner->GetClass()))
+		if (const FInstanceDataTransformSet* TransformSet = FindTransformsInternal(Owner->GetClass()))
 		{
-			for (const TSharedPtr<FTransform>& IDT : *IDTs)
+			for (const FInstanceDataTransformSet::FOperation& Operation : TransformSet->Operations)
 			{
 				FPropertyPathName RootPath; // class IDTs operate at the root of the class
-				ETransformResult Result = IDT->Func(IDO, Owner, RootPath, IDT->Args);
+				ETransformResult Result = Operation(IDO, Owner);
 				if (Result == ETransformResult::Break)
 				{
 					break;
@@ -626,7 +721,6 @@ namespace UE
 
 	void FInstanceDataTransforms::PatchStructs(TObjectPtr<UObject>& IDO, TObjectPtr<UObject>& Owner)
 	{
-
 		// search the unknown property tree for structs. If there are any, look up IDTs associated with that struct type and run them
 		TSharedPtr<FPropertyPathNameTree> UnknownPropertyTree = FUnknownPropertyTree(Owner).FindOrCreate();
 
@@ -641,16 +735,15 @@ namespace UE
 			FPropertyPathName& Path = Paths[i];
 			FPropertyPathNameSegment LeafSegment = Path.GetSegment(Path.GetSegmentCount() - 1);
 
-			auto PatchStruct = [&Path, &IDO, &Owner](const FPropertyTypeName& StructPath)
+			auto PatchStruct = [&Path, &IDO, &Owner, this](const FPropertyTypeName& StructPath)
 				{
 					if (UStruct* Struct = UE::FindObjectByTypePath<UScriptStruct>(StructPath))
 					{
-
-						if (const TArray<TSharedPtr<FTransform>>* IDTs = FInstanceDataTransforms::Get().FindTransformsInternal(Struct))
+						if (const FInstanceDataTransformSet* TransformSet = FindTransformsInternal(Struct))
 						{
-							for (const TSharedPtr<FTransform>& IDT : *IDTs)
+							for (const FInstanceDataTransformSet::FOperation& Operation : TransformSet->Operations)
 							{
-								ETransformResult Result = IDT->Func(IDO, Owner, Path, IDT->Args);
+								ETransformResult Result = Operation(IDO, Owner, Path);
 								if (Result == ETransformResult::Break)
 								{
 									break;
