Fixed not invalidating AnnotationCacheKey and Value from CL 47169102.

#tests editor

diff --git a/Engine/Source/Runtime/CoreUObject/Private/UObject/UObjectMarks.cpp b/Engine/Source/Runtime/CoreUObject/Private/UObject/UObjectMarks.cpp
--- a/Engine/Source/Runtime/CoreUObject/Private/UObject/UObjectMarks.cpp
+++ b/Engine/Source/Runtime/CoreUObject/Private/UObject/UObjectMarks.cpp
@@ -124,6 +124,8 @@ public:
 	{
 		if (!AnnotationMap.IsEmpty())
 		{
+			AnnotationCacheKey = nullptr;
+			AnnotationCacheValue = TAnnotation();
 			for (typename TMap<const UObjectBase*, TAnnotation>::TIterator It(AnnotationMap); It; ++It)
 			{
 				if (It.Value().ClearMarks(Annotation))
