Add CSV stats FileIO/AsyncLoading and Exclusive/FlushAsyncLoading
#tests tested on console

diff --git a/Engine/Source/Runtime/CoreUObject/Private/Serialization/AsyncLoading.cpp b/Engine/Source/Runtime/CoreUObject/Private/Serialization/AsyncLoading.cpp
--- a/Engine/Source/Runtime/CoreUObject/Private/Serialization/AsyncLoading.cpp
+++ b/Engine/Source/Runtime/CoreUObject/Private/Serialization/AsyncLoading.cpp
@@ -4330,6 +4330,8 @@ struct FScopedRecursionNotAllowed
 EAsyncPackageState::Type FAsyncLoadingThread::ProcessAsyncLoading(int32& OutPackagesProcessed, bool bUseTimeLimit /*= false*/, bool bUseFullTimeLimit /*= false*/, float TimeLimit /*= 0.0f*/, FFlushRequest& FlushRequest)
 {
 	SCOPED_LOADTIMER(AsyncLoadingTime);
+	CSV_SCOPED_TIMING_STAT(FileIO, AsyncLoadingTime);
+
 	check(!IsInGameThread() || !IsMultithreaded());
 
 	// If we're not multithreaded and flushing async loading, update the thread heartbeat
diff --git a/Engine/Source/Runtime/CoreUObject/Private/Serialization/AsyncLoading2.cpp b/Engine/Source/Runtime/CoreUObject/Private/Serialization/AsyncLoading2.cpp
--- a/Engine/Source/Runtime/CoreUObject/Private/Serialization/AsyncLoading2.cpp
+++ b/Engine/Source/Runtime/CoreUObject/Private/Serialization/AsyncLoading2.cpp
@@ -9795,6 +9795,7 @@ void FAsyncLoadingThread2::UpdateSyncLoadContext(FAsyncLoadingThreadState2& Thre
 EAsyncPackageState::Type FAsyncLoadingThread2::ProcessAsyncLoadingFromGameThread(FAsyncLoadingThreadState2& ThreadState, bool& bDidSomething)
 {
 	SCOPED_LOADTIMER(AsyncLoadingTime);
+	CSV_SCOPED_TIMING_STAT(FileIO, AsyncLoadingTime);
 
 	check(IsInGameThread());
 
@@ -10391,6 +10392,7 @@ uint32 FAsyncLoadingThread2::Run()
 		else if (CurrentState == EMainState::Loading)
 		{
 			TRACE_CPUPROFILER_EVENT_SCOPE(AsyncLoadingTime);
+			CSV_SCOPED_TIMING_STAT(FileIO, AsyncLoadingTime);
 
 #if USING_INSTRUMENTATION
 			if (GDetectRaceDuringLoading)
diff --git a/Engine/Source/Runtime/CoreUObject/Private/Serialization/AsyncPackageLoader.cpp b/Engine/Source/Runtime/CoreUObject/Private/Serialization/AsyncPackageLoader.cpp
--- a/Engine/Source/Runtime/CoreUObject/Private/Serialization/AsyncPackageLoader.cpp
+++ b/Engine/Source/Runtime/CoreUObject/Private/Serialization/AsyncPackageLoader.cpp
@@ -14,6 +14,7 @@
 #include "IO/IoDispatcher.h"
 #include "IO/IoDispatcherInternal.h"
 #include "HAL/IConsoleManager.h"
+#include "ProfilingDebugging/CsvProfiler.h"
 #include "AutoRTFM.h"
 
 static bool GUseOldLoaderAsFallback = false;
@@ -29,7 +30,6 @@ static FAutoConsoleVariableRef CVarGUseOldLoaderAsFallback(
 #if DO_TRACK_ASYNC_LOAD_REQUESTS
 #include "Containers/StackTracker.h"
 #include "Misc/OutputDeviceFile.h"
-#include "ProfilingDebugging/CsvProfiler.h"
 #endif
 
 volatile int32 GIsLoaderCreated;
@@ -325,6 +325,7 @@ void FlushAsyncLoading(int32 RequestId /* = INDEX_NONE */)
 void FlushAsyncLoading(TConstArrayView<int32> RequestIds)
 {
 	TRACE_CPUPROFILER_EVENT_SCOPE(FlushAsyncLoading);
+	CSV_SCOPED_TIMING_STAT_EXCLUSIVE(FlushAsyncLoading);
 
 #if defined(WITH_CODE_GUARD_HANDLER) && WITH_CODE_GUARD_HANDLER
 	void CheckImageIntegrityAtRuntime();
