ObjectHandle: Fixed missing assignment of OutObjectId in MakeReferenceIds

diff --git a/Engine/Source/Runtime/CoreUObject/Private/UObject/ObjectHandleIndex.cpp b/Engine/Source/Runtime/CoreUObject/Private/UObject/ObjectHandleIndex.cpp
--- a/Engine/Source/Runtime/CoreUObject/Private/UObject/ObjectHandleIndex.cpp
+++ b/Engine/Source/Runtime/CoreUObject/Private/UObject/ObjectHandleIndex.cpp
@@ -254,23 +254,26 @@ namespace UE::CoreUObject::Private
 			return;
 		}
 
-		FRWScopeLock LocalLockScope(PackageData->Lock, SLT_ReadOnly);
-		int32 PathIndex = PackageData->ObjectDescriptors.IndexOfByPredicate([&ObjectPath](const FObjectDescriptor& Desc)
+		int32 PathIndex;
+		{
+			TReadScopeLock LocalLockScope(PackageData->Lock);
+			PathIndex = PackageData->ObjectDescriptors.IndexOfByPredicate([&ObjectPath](const FObjectDescriptor& Desc)
 			{
 				return Desc.ObjectPath == ObjectPath;
 
 			}); //linear search is fine as typically there is only one.
+		}
 
 		if (PathIndex == INDEX_NONE)
 		{
 			//ObjectPaths could have been modified when the read lock was released and the write
 			//lock was acquired, so we must check and see if the ObjectPath was added in that window.
-			LocalLockScope.ReleaseReadOnlyLockAndAcquireWriteLock_USE_WITH_CAUTION();
+			TWriteScopeLock LocalLockScope(PackageData->Lock);
 			PathIndex = PackageData->ObjectDescriptors.IndexOfByPredicate([&ObjectPath](const FObjectDescriptor& Desc)
-				{
-					return Desc.ObjectPath == ObjectPath;
+			{
+				return Desc.ObjectPath == ObjectPath;
 
-				}); //linear search is fine as typically there is only one.
+			}); //linear search is fine as typically there is only one.
 
 			if (PathIndex == INDEX_NONE)
 			{
@@ -282,15 +285,11 @@ namespace UE::CoreUObject::Private
 				}
 
 				FObjectDescriptor ObjDesc = { ObjectPath, DataClassDesc };
-				PathIndex = PackageData->ObjectDescriptors.Add(ObjDesc) + 1;
-				OutObjectId = FObjectId::FromIndex(PathIndex);
+				PathIndex = PackageData->ObjectDescriptors.Add(ObjDesc);
 			}
 		}
-		else
-		{
-			OutObjectId = FObjectId::FromIndex(PathIndex + 1);
-		}
-		return;
+
+		OutObjectId = FObjectId::FromIndex(PathIndex + 1);
 	}
 
 	static inline FPackedObjectRef MakePackedObjectRef(FName PackageName, FName ClassPackageName, FName ClassName, FObjectPathId ObjectPath, bool bHasPlaceholderType = false)
