Support VCell in FObjectWriter and FObjectReader

This is used when creating IDOs for classes with FVRestValueProperty members, and similar scenarios.

#rnx

diff --git a/Engine/Source/Runtime/Core/Public/Serialization/MemoryArchive.h b/Engine/Source/Runtime/Core/Public/Serialization/MemoryArchive.h
--- a/Engine/Source/Runtime/Core/Public/Serialization/MemoryArchive.h
+++ b/Engine/Source/Runtime/Core/Public/Serialization/MemoryArchive.h
@@ -57,6 +57,15 @@ public:
 		return *this;
 	}
 
+#if WITH_VERSE_VM || defined(__INTELLISENSE__)
+	virtual FArchive& operator<<( Verse::VCell*& Res ) override
+	{
+		// Not supported through this archive
+		check(0);
+		return *this;
+	}
+#endif
+
 protected:
 
 	/** Marked as protected to avoid instantiating this class directly */
diff --git a/Engine/Source/Runtime/CoreUObject/Private/Serialization/ObjectReader.cpp b/Engine/Source/Runtime/CoreUObject/Private/Serialization/ObjectReader.cpp
--- a/Engine/Source/Runtime/CoreUObject/Private/Serialization/ObjectReader.cpp
+++ b/Engine/Source/Runtime/CoreUObject/Private/Serialization/ObjectReader.cpp
@@ -27,6 +27,14 @@ FArchive& FObjectReader::operator<<(UObject*& Res)
 	return *this;
 }
 
+#if WITH_VERSE_VM || defined(__INTELLISENSE__)
+FArchive& FObjectReader::operator<<(Verse::VCell*& Res)
+{
+	ByteOrderSerialize(&Res, sizeof(Res));
+	return *this;
+}
+#endif
+
 FArchive& FObjectReader::operator<<(FObjectPtr& Value)
 {
 	ByteOrderSerialize(&Value, sizeof(Value));
diff --git a/Engine/Source/Runtime/CoreUObject/Private/Serialization/ObjectWriter.cpp b/Engine/Source/Runtime/CoreUObject/Private/Serialization/ObjectWriter.cpp
--- a/Engine/Source/Runtime/CoreUObject/Private/Serialization/ObjectWriter.cpp
+++ b/Engine/Source/Runtime/CoreUObject/Private/Serialization/ObjectWriter.cpp
@@ -25,6 +25,14 @@ FArchive& FObjectWriter::operator<<( class UObject*& Res )
 	return *this;
 }
 
+#if WITH_VERSE_VM || defined(__INTELLISENSE__)
+FArchive& FObjectWriter::operator<<( Verse::VCell*& Res )
+{
+	ByteOrderSerialize(&Res, sizeof(Res));
+	return *this;
+}
+#endif
+
 FArchive& FObjectWriter::operator<<(FObjectPtr& Value)
 {
 	ByteOrderSerialize(&Value, sizeof(Value));
diff --git a/Engine/Source/Runtime/CoreUObject/Public/Serialization/ObjectReader.h b/Engine/Source/Runtime/CoreUObject/Public/Serialization/ObjectReader.h
--- a/Engine/Source/Runtime/CoreUObject/Public/Serialization/ObjectReader.h
+++ b/Engine/Source/Runtime/CoreUObject/Public/Serialization/ObjectReader.h
@@ -72,6 +72,9 @@ public:
 
 	COREUOBJECT_API virtual FArchive& operator<<(FName& N) override;
 	COREUOBJECT_API virtual FArchive& operator<<(UObject*& Res) override;
+#if WITH_VERSE_VM || defined(__INTELLISENSE__)
+	COREUOBJECT_API virtual FArchive& operator<<(Verse::VCell*& Res) override;
+#endif
 	COREUOBJECT_API virtual FArchive& operator<<(FObjectPtr& Value) override;
 	COREUOBJECT_API virtual FArchive& operator<<(FLazyObjectPtr& LazyObjectPtr) override;
 	COREUOBJECT_API virtual FArchive& operator<<(FSoftObjectPtr& Value) override;
diff --git a/Engine/Source/Runtime/CoreUObject/Public/Serialization/ObjectWriter.h b/Engine/Source/Runtime/CoreUObject/Public/Serialization/ObjectWriter.h
--- a/Engine/Source/Runtime/CoreUObject/Public/Serialization/ObjectWriter.h
+++ b/Engine/Source/Runtime/CoreUObject/Public/Serialization/ObjectWriter.h
@@ -49,6 +49,9 @@ public:
 	//~ Begin FArchive Interface
 	COREUOBJECT_API virtual FArchive& operator<<(FName& N) override;
 	COREUOBJECT_API virtual FArchive& operator<<(UObject*& Res) override;
+#if WITH_VERSE_VM || defined(__INTELLISENSE__)
+	COREUOBJECT_API virtual FArchive& operator<<(Verse::VCell*& Res) override;
+#endif
 	COREUOBJECT_API virtual FArchive& operator<<(FObjectPtr& Value) override;
 	COREUOBJECT_API virtual FArchive& operator<<(FLazyObjectPtr& Value) override;
 	COREUOBJECT_API virtual FArchive& operator<<(FSoftObjectPtr& Value) override;
