Adjusted ConditionalFlushAsyncLoadingForLinkers to use the Linker's instancing context when available.
This prevents unnecessary object creation when resetting loaders for objects that were loaded with an instancing context.

diff --git a/Engine/Source/Runtime/CoreUObject/Private/UObject/Linker.cpp b/Engine/Source/Runtime/CoreUObject/Private/UObject/Linker.cpp
--- a/Engine/Source/Runtime/CoreUObject/Private/UObject/Linker.cpp
+++ b/Engine/Source/Runtime/CoreUObject/Private/UObject/Linker.cpp
@@ -832,7 +832,16 @@ void ConditionalFlushAsyncLoadingForLinkers(TConstArrayView<FLinkerLoad*> InLink
 		TArray<int32, TInlineAllocator<4>> RequestIds;
 		for (FLinkerLoad* Linker : InLinkers)
 		{
-			int32 Request = LoadPackageAsync(Linker->GetPackagePath(), FLoadPackageAsyncOptionalParams{ .PackagePriority = MAX_int32 });
+			int32 Request = INDEX_NONE;
+			if (!Linker->GetInstancingContext().IsInstanced())
+			{
+				Request = LoadPackageAsync(Linker->GetPackagePath(), FLoadPackageAsyncOptionalParams{ .PackagePriority = MAX_int32 });
+			}
+			else
+			{
+				FName RemappedPackageName = Linker->GetInstancingContext().RemapPackage(Linker->GetPackagePath().GetPackageFName());
+				Request = LoadPackageAsync(Linker->GetPackagePath(), FLoadPackageAsyncOptionalParams{ .CustomPackageName = RemappedPackageName, .PackagePriority = MAX_int32, .InstancingContext = &Linker->GetInstancingContext() });
+			}
 			RequestIds.Add(Request);
 		}
 		FlushAsyncLoading(RequestIds);
