Emit a CSV event for FlushAsyncLoading

diff --git a/Engine/Source/Runtime/CoreUObject/Private/Serialization/AsyncPackageLoader.cpp b/Engine/Source/Runtime/CoreUObject/Private/Serialization/AsyncPackageLoader.cpp
--- a/Engine/Source/Runtime/CoreUObject/Private/Serialization/AsyncPackageLoader.cpp
+++ b/Engine/Source/Runtime/CoreUObject/Private/Serialization/AsyncPackageLoader.cpp
@@ -326,6 +326,7 @@ void FlushAsyncLoading(TConstArrayView<int32> RequestIds)
 {
 	TRACE_CPUPROFILER_EVENT_SCOPE(FlushAsyncLoading);
 	CSV_SCOPED_TIMING_STAT_EXCLUSIVE(FlushAsyncLoading);
+	CSV_EVENT_GLOBAL(TEXT("FlushAsyncLoading"));
 
 #if defined(WITH_CODE_GUARD_HANDLER) && WITH_CODE_GUARD_HANDLER
 	void CheckImageIntegrityAtRuntime();
