Add casting operator from TObjectKey to FObjectKey
#rnx

diff --git a/Engine/Source/Runtime/CoreUObject/Public/UObject/ObjectKey.h b/Engine/Source/Runtime/CoreUObject/Public/UObject/ObjectKey.h
--- a/Engine/Source/Runtime/CoreUObject/Public/UObject/ObjectKey.h
+++ b/Engine/Source/Runtime/CoreUObject/Public/UObject/ObjectKey.h
@@ -278,6 +278,12 @@ public:
 		return ObjectKey >= Other.ObjectKey;
 	}
 
+	/** @return the object key */
+	UE_REWRITE operator FObjectKey() const
+	{
+		return ObjectKey;
+	}
+
 	//** Hash function */
 	[[nodiscard]] friend uint32 GetTypeHash(const TObjectKey& Key)
 	{
