Make FReferenceChainSearch::DumpChain public to allow for custom dump tools.

#rnx

diff --git a/Engine/Source/Runtime/CoreUObject/Public/UObject/ReferenceChainSearch.h b/Engine/Source/Runtime/CoreUObject/Public/UObject/ReferenceChainSearch.h
--- a/Engine/Source/Runtime/CoreUObject/Public/UObject/ReferenceChainSearch.h
+++ b/Engine/Source/Runtime/CoreUObject/Public/UObject/ReferenceChainSearch.h
@@ -382,10 +382,9 @@ private:
 
 public:
 	/** Returns a string with all flags (we care about) set on an object */
-	static FString GetObjectFlags(const FGCObjectInfo& InObject);
-private:
+	COREUOBJECT_API static FString GetObjectFlags(const FGCObjectInfo& InObject);
 
 	/** Dumps a reference chain to log */
-	static void DumpChain(FReferenceChainSearch::FReferenceChain* Chain, TFunctionRef<bool(FCallbackParams& Params)> ReferenceCallback,
+	COREUOBJECT_API static void DumpChain(FReferenceChainSearch::FReferenceChain* Chain, TFunctionRef<bool(FCallbackParams& Params)> ReferenceCallback,
 		TMap<uint64, FString>& CallstackCache, FOutputDevice& Out, ELogVerbosity::Type InVerbosityForPrint);
 };
