Fix for Crashing during Shutdown if we trigger a Crash within the UObjectBase constructor. This allows crash reports to be sent from executables that crash in this way.
If UObjectBase::AddObject fails on HashObject, then we've added an object to the GUObjectArray which hasn't been fully constructed yet, so the static cast inside StaticShutdownAfterError would cause us to have a UObject* which has a vtbl pointer which points to UObjectBase's vtbl instead, causing undefined code execution.
This change fixes the issue by using an Iterator type that checks against PendingConstruction which will not return the partially constructed UObjectBase

#rnx
#tests Tested against Cooked Dedicated Server

diff --git a/Engine/Source/Runtime/CoreUObject/Private/UObject/Obj.cpp b/Engine/Source/Runtime/CoreUObject/Private/UObject/Obj.cpp
--- a/Engine/Source/Runtime/CoreUObject/Private/UObject/Obj.cpp
+++ b/Engine/Source/Runtime/CoreUObject/Private/UObject/Obj.cpp
@@ -4055,10 +4055,13 @@ static void StaticShutdownAfterError()
 		bShutdown = true;
 		UE_LOG(LogExit, Log, TEXT("Executing StaticShutdownAfterError") );
 
-		for( FRawObjectIterator It; It; ++It )
+		for( FThreadSafeObjectIterator It; It; ++It )
 		{
-			UObject* Object = static_cast<UObject*>(It->GetObject());
-			Object->ShutdownAfterError();
+			UObject* Object = *It;
+			if(Object)
+			{
+				Object->ShutdownAfterError();
+			}
 		}
 	}
 }
