UE::IsGCLocked should check FGCCSyncObject::IsGCLocked, not IsAsyncLocked

This function guards delegate broadcasts that happen in response to assertion failures- if the GC itself hits an assertion failure while holding the GC lock, these broadcasts will deadlock in TWeakObjectPtr::Pin, which also acquires the GC lock.

However, the GC lock is split into two: an Async part acquired by non-GC code (including TWeakObjectPtr::Pin), and a GC part acquired by the GC itself. These parts are mutually exclusive, so IsAsyncLocked prevents the GC from taking the lock, while IsGCLocked prevents non-GC code from taking the lock.

This means that checking IsAsyncLocked before calling TWeakObjectPtr::Pin does not prevent deadlocks when the lock is held by the GC. (In fact it does nothing at all, because the lock allows multiple concurrent LockAsyncs.)

#rnx

diff --git a/Engine/Source/Runtime/CoreUObject/Private/UObject/GarbageCollection.cpp b/Engine/Source/Runtime/CoreUObject/Private/UObject/GarbageCollection.cpp
--- a/Engine/Source/Runtime/CoreUObject/Private/UObject/GarbageCollection.cpp
+++ b/Engine/Source/Runtime/CoreUObject/Private/UObject/GarbageCollection.cpp
@@ -151,15 +151,10 @@ public:
 	}
 };
 
-bool IsGarbageCollectionLocked()
-{
-	return FGCCSyncObject::Get().IsAsyncLocked();
-}
-
 FGCCSyncObject::FGCCSyncObject()
 {
 	GCUnlockedEvent = FPlatformProcess::GetSynchEventFromPool(true);
-	UE::Private::SetIsGCLockedFunction(IsGarbageCollectionLocked);
+	UE::Private::SetIsGCLockedFunction([] { return FGCCSyncObject::Get().IsGCLocked(); });
 }
 FGCCSyncObject::~FGCCSyncObject()
 {
@@ -225,6 +220,11 @@ FGCScopeTryGuard::~FGCScopeTryGuard()
 	}
 }
 
+bool IsGarbageCollectionLocked()
+{
+	return FGCCSyncObject::Get().IsAsyncLocked();
+}
+
 static int32 GIncrementalBeginDestroyEnabled = 1;
 static FAutoConsoleVariableRef CIncrementalBeginDestroyEnabled(
 	TEXT("gc.IncrementalBeginDestroyEnabled"),
