When remote object handles are enabled, unresolved TObjectPtr will now properly return false when calling TObjectPtr::IsRemote() when the stub it points to represents a local object.

diff --git a/Engine/Source/Runtime/CoreUObject/Private/UObject/RemoteObject.cpp b/Engine/Source/Runtime/CoreUObject/Private/UObject/RemoteObject.cpp
--- a/Engine/Source/Runtime/CoreUObject/Private/UObject/RemoteObject.cpp
+++ b/Engine/Source/Runtime/CoreUObject/Private/UObject/RemoteObject.cpp
@@ -548,6 +548,7 @@ bool IsRemote(FRemoteObjectId ObjectId)
 		return IsRemote(Object);
 	}
 
+	// if a stub exists but we don't have the object's memory (the UObject doesn't exist) then the object must be remote
 	if (FindRemoteObjectStub(ObjectId))
 	{
 		return true;
@@ -558,6 +559,22 @@ bool IsRemote(FRemoteObjectId ObjectId)
 	return ServerId.IsValid() && !ServerId.IsLocal() && !ServerId.IsAsset();
 }
 
+bool IsRemote(const FRemoteObjectStub* Stub)
+{
+	if (!Stub)
+	{
+		return false;
+	}
+
+	if (UObject* Object = StaticFindObjectFastInternal(Stub->Id))
+	{
+		return IsRemote(Object);
+	}
+
+	// Similar to the above IsRemote(ObjectId): if a stub exists but we don't have the object's memory (the UObject doesn't exist) then the object must be remote
+	return true;
+}
+
 bool IsRemote(const UObject* Object)
 {
 #if UE_WITH_REMOTE_OBJECT_HANDLE
diff --git a/Engine/Source/Runtime/CoreUObject/Public/UObject/ObjectHandle.h b/Engine/Source/Runtime/CoreUObject/Public/UObject/ObjectHandle.h
--- a/Engine/Source/Runtime/CoreUObject/Public/UObject/ObjectHandle.h
+++ b/Engine/Source/Runtime/CoreUObject/Public/UObject/ObjectHandle.h
@@ -249,6 +249,25 @@ inline bool IsObjectHandleResolved_ForGC(FObjectHandle Handle)
 #endif
 }
 
+/* checks if a handle represents a remote object.
+ * nullptr counts as not remote
+ * all handles are not remote when remote handles are not enabled
+ */
+inline bool IsObjectHandleRemote(FObjectHandle Handle)
+{
+#if UE_WITH_REMOTE_OBJECT_HANDLE
+	if (!!(Handle.PointerOrHandle & 1))
+	{
+		return UE::RemoteObject::Handle::IsRemote(Handle.ToStub());
+	}
+	else if (!IsObjectHandleNull(Handle))
+	{
+		return UE::RemoteObject::Handle::IsRemote(ReadObjectHandlePointerNoCheck(Handle));
+	}
+#endif
+	return false;
+}
+
 /* return true if a handle is type safe.
  * null and resolved handles are considered type safe
  */ 
diff --git a/Engine/Source/Runtime/CoreUObject/Public/UObject/ObjectPtr.h b/Engine/Source/Runtime/CoreUObject/Public/UObject/ObjectPtr.h
--- a/Engine/Source/Runtime/CoreUObject/Public/UObject/ObjectPtr.h
+++ b/Engine/Source/Runtime/CoreUObject/Public/UObject/ObjectPtr.h
@@ -189,7 +189,7 @@ public:
 	 * FObjectPtr::IsRemote is only used when UE_WITH_REMOTE_OBJECT_HANDLE is true, and is mutually exclusive with other uses of unresolved FObjectPtr.
 	 * When using UE_WITH_REMOTE_OBJECT_HANDLE, unresolved FObjectPtrs mean the object is not locally available and is present on another server, see @RemoteObject.h
 	*/
-	[[nodiscard]] FORCEINLINE bool IsRemote() const { return !IsResolved(); }
+	[[nodiscard]] FORCEINLINE bool IsRemote() const { return IsObjectHandleRemote(Handle); }
 
 	// Gets the PathName of the object without resolving the object reference.
 	// @TODO OBJPTR: Deprecate this.
diff --git a/Engine/Source/Runtime/CoreUObject/Public/UObject/RemoteObject.h b/Engine/Source/Runtime/CoreUObject/Public/UObject/RemoteObject.h
--- a/Engine/Source/Runtime/CoreUObject/Public/UObject/RemoteObject.h
+++ b/Engine/Source/Runtime/CoreUObject/Public/UObject/RemoteObject.h
@@ -191,6 +191,12 @@ namespace UE::RemoteObject::Handle
 	*/
 	COREUOBJECT_API bool IsRemote(FRemoteObjectId ObjectId);
 
+	/**
+	* Checks if an object associated with the specified stub is remote
+	* @return True if an object associated with the specified stub is remote
+	*/
+	COREUOBJECT_API bool IsRemote(const FRemoteObjectStub* Stub);
+
 	/**
 	* Checks if an object (memory that has not yet been GC'd) is remote
 	* @return True if the object is remote
