Improvements to the serialization of ImportTypeHierarchy entries.

#rnx

diff --git a/Engine/Source/Runtime/CoreUObject/Private/UObject/LinkerLoad.cpp b/Engine/Source/Runtime/CoreUObject/Private/UObject/LinkerLoad.cpp
--- a/Engine/Source/Runtime/CoreUObject/Private/UObject/LinkerLoad.cpp
+++ b/Engine/Source/Runtime/CoreUObject/Private/UObject/LinkerLoad.cpp
@@ -2734,13 +2734,35 @@ FLinkerLoad::ELinkerStatus FLinkerLoad::SerializeImportTypeHierarchies()
 		return LINKER_Loaded;
 	}
 
+	if (Summary.ImportTypeHierarchiesCount < 0)
+	{
+		return LINKER_Failed;
+	}
+
 	Seek(Summary.ImportTypeHierarchiesOffset);
 
-	FStructuredArchive::FStream Stream = StructuredArchiveRootRecord->EnterStream(TEXT("ImportTypeHierarchies"));
-	Stream << ImportTypeHierarchies;
+	int32 NumTypeHierarchies = 0;
+	FStructuredArchive::FArray Array = StructuredArchiveRootRecord->EnterArray(TEXT("ImportTypeHierarchies"), NumTypeHierarchies);
+
+	if (NumTypeHierarchies == Summary.ImportTypeHierarchiesCount)
+	{
+		for (int32 Hierarchy = 0; Hierarchy < Summary.ImportTypeHierarchiesCount; ++Hierarchy)
+		{
+			FStructuredArchive::FSlot ElementSlot = Array.EnterElement();
+
+			TPair<FPackageIndex, UE::Serialization::Private::FImportTypeHierarchy> Pair;
+			ElementSlot << Pair;
+
+			ImportTypeHierarchies.Add(Pair.Key, Pair.Value);
+		}
+	}
+	else
+	{
+		return LINKER_Failed;
+	}
 
 	// Return whether we finished this step and it's safe to start with the next.
-	if (Stream.GetUnderlyingArchive().IsCriticalError())
+	if (StructuredArchiveRootRecord->GetUnderlyingArchive().IsCriticalError())
 	{
 		return LINKER_Failed;
 	}
