Remove the now defunct old (bad) trashing behaviour since we haven't had to toggle the CVar in a few releases.

diff --git a/Engine/Source/Runtime/CoreUObject/Private/Serialization/TransactionallySafeAsyncLoading.cpp b/Engine/Source/Runtime/CoreUObject/Private/Serialization/TransactionallySafeAsyncLoading.cpp
--- a/Engine/Source/Runtime/CoreUObject/Private/Serialization/TransactionallySafeAsyncLoading.cpp
+++ b/Engine/Source/Runtime/CoreUObject/Private/Serialization/TransactionallySafeAsyncLoading.cpp
@@ -10,10 +10,7 @@
 #include "Misc/ScopeLock.h"
 #include "HAL/IConsoleManager.h"
 
-#if UE_AUTORTFM_ENABLED
-
-bool CVar_TransactionallySafeAsyncLoadingLegacyTrashedPackages = false;
-static FAutoConsoleVariableRef CVarTransactionallySafeAsyncLoadingLegacyTrashedPackages(TEXT("TransactionallySafeAsyncLoading.LegacyTrashedPackages"), CVar_TransactionallySafeAsyncLoadingLegacyTrashedPackages, TEXT("Use the legacy trashed packages handling (EG. ignore them!)."), ECVF_Default);
+#if UE_AUTORTFM_ENABLED || defined(__INTELLISENSE__)
 
 // A transactionally safe async package loader that wraps an underlying actual
 // async package loader but allows it to be interacted with safely while also
@@ -214,15 +211,12 @@ public:
 
 						UPackage* const Package = MapPayload->GetPackage().Get();
 
-						if (LIKELY(!CVar_TransactionallySafeAsyncLoadingLegacyTrashedPackages))
+						// Now that the package has come into actual use we rename it to the actual
+						// package name that it originally was conceived with, making it discoverable
+						// in `FindObject` calls.
+						if (Package)
 						{
-							// Now that the package has come into actual use we rename it to the actual
-							// package name that it originally was conceived with, making it discoverable
-							// in `FindObject` calls.
-							if (Package)
-							{
-								Package->Rename(*CustomPackageName.ToString(), nullptr, REN_DontCreateRedirectors | REN_NonTransactional | REN_DoNotDirty);
-							}
+							Package->Rename(*CustomPackageName.ToString(), nullptr, REN_DontCreateRedirectors | REN_NonTransactional | REN_DoNotDirty);
 						}
 
 						InCompletionDelegate.ExecuteIfBound(CustomPackageName, Package, MapPayload->Result);
@@ -544,12 +538,6 @@ private:
 				return nullptr;
 			}
 
-			if (UNLIKELY(CVar_TransactionallySafeAsyncLoadingLegacyTrashedPackages))
-			{
-				ensure(LoadedPackages.Num() <= 1);
-				return LoadedPackages[0].Get();
-			}
-
 			for (TWeakObjectPtr<UPackage> LoadedPackage : LoadedPackages)
 			{
 				if (LoadedPackage.IsValid())
@@ -650,7 +638,7 @@ private:
 		for (const FTransactionallyDeferredLoadPackagePayload& Payload : TransactionallyDeferredLoadPackages)
 		{
 			const int32 RequestId = LoadAndCachePackage(Payload.PackagePath,
-				bIsRetrying && LIKELY(!CVar_TransactionallySafeAsyncLoadingLegacyTrashedPackages) ? MakeUniqueObjectName(nullptr, UPackage::StaticClass(), NAME_AutoRTFMUnusedPackage) : Payload.CustomPackageName,
+				bIsRetrying ? MakeUniqueObjectName(nullptr, UPackage::StaticClass(), NAME_AutoRTFMUnusedPackage) : Payload.CustomPackageName,
 				bIsRetrying ? FLoadPackageAsyncDelegate() : Payload.CompletionDelegate,
 				Payload.PackageFlags,
 				Payload.PIEInstanceID,
@@ -678,7 +666,7 @@ private:
 
 const char* const FTransactionallySafeAsyncPackageLoader::UnreachableMessage = "Cannot call function within a transaction!";
 
-#endif // UE_AUTORTFM_ENABLED
+#endif // UE_AUTORTFM_ENABLED || defined(__INTELLISENSE__)
 
 IAsyncPackageLoader* MakeTransactionallySafeAsyncPackageLoader(IAsyncPackageLoader* const InWrappedPackageLoader)
 {
