Update PostInitProperties doc-comment to clarify that Verse construction happens afterwards.

No functional change, just a comment tweak.

diff --git a/Engine/Source/Runtime/CoreUObject/Public/UObject/Object.h b/Engine/Source/Runtime/CoreUObject/Public/UObject/Object.h
--- a/Engine/Source/Runtime/CoreUObject/Public/UObject/Object.h
+++ b/Engine/Source/Runtime/CoreUObject/Public/UObject/Object.h
@@ -217,7 +217,7 @@ protected:
 public:
 	/**
 	 * Called after the C++ constructor and after the properties have been initialized, including those loaded from config.
-	 * This is called before any serialization or other setup has happened.
+	 * This is called before any serialization, Verse construction, or other setup has occurred.
 	 */
 	COREUOBJECT_API virtual void PostInitProperties();
 	
