[Zenloader]
- Split runtime loader nodes to the same granularity as the editor to allow some nodes to run in parallel

diff --git a/Engine/Source/Runtime/CoreUObject/Private/Serialization/AsyncLoading2.cpp b/Engine/Source/Runtime/CoreUObject/Private/Serialization/AsyncLoading2.cpp
--- a/Engine/Source/Runtime/CoreUObject/Private/Serialization/AsyncLoading2.cpp
+++ b/Engine/Source/Runtime/CoreUObject/Private/Serialization/AsyncLoading2.cpp
@@ -213,7 +213,7 @@ FArchive& operator<<(FArchive& Ar, FExportMapEntry& ExportMapEntry)
 
 	uint32 ObjectFlags = uint32(ExportMapEntry.ObjectFlags);
 	Ar << ObjectFlags;
-	
+
 	if (Ar.IsLoading())
 	{
 		ExportMapEntry.ObjectFlags = EObjectFlags(ObjectFlags);
@@ -242,7 +242,7 @@ FArchive& operator<<(FArchive& Ar, FCellExportMapEntry& CellExportMapEntry)
 
 	return Ar;
 }
-	
+
 FArchive& operator<<(FArchive& Ar, FBulkDataMapEntry& BulkDataEntry)
 {
 	Ar << BulkDataEntry.SerialOffset;
@@ -327,7 +327,7 @@ void FindAllRuntimeScriptPackages(FRuntimeScriptPackages& OutPackages)
 #endif
 
 #ifndef ALT2_ENABLE_SCRIPT_IMPORTS_DEBUGGING
-	#define ALT2_ENABLE_SCRIPT_IMPORTS_DEBUGGING !UE_BUILD_SHIPPING && !IS_PROGRAM
+#define ALT2_ENABLE_SCRIPT_IMPORTS_DEBUGGING !UE_BUILD_SHIPPING && !IS_PROGRAM
 #endif
 
 #if ALT2_ENABLE_SCRIPT_IMPORTS_DEBUGGING
@@ -338,7 +338,7 @@ static FAutoConsoleVariableRef CEnableScriptImportsDebugging(
 	TEXT("Loads all native script objects (classes, enums, structs, CDOs) that were available in the editor during cook.")
 	TEXT("Provides more informative log messages at the cost of up to ~1.5 MiB memory and ~1 MiB FName memory depending on project size."),
 	ECVF_Default
-	);
+);
 #endif
 
 #define checkObject(Object, expr) checkf(expr, TEXT("Object='%s' (%p), Flags=%s, InternalFlags=0x%08X") \
@@ -384,23 +384,23 @@ static FAutoConsoleVariableRef CVar_DebugPackageNames(
 	GAsyncLoading2_DebugPackageNamesString,
 	TEXT("Add debug breaks for all listed package names, also automatically added to s.VerbosePackageNames."),
 	FConsoleVariableDelegate::CreateLambda([](IConsoleVariable* Variable)
-	{
-		GAsyncLoading2_DebugPackageIds.Reset();
-		ParsePackageNames(Variable->GetString(), GAsyncLoading2_DebugPackageIds);
-		ParsePackageNames(Variable->GetString(), GAsyncLoading2_VerbosePackageIds);
-		GAsyncLoading2_VerboseLogFilter = GAsyncLoading2_VerbosePackageIds.Num() > 0 ? 1 : 2;
-	}),
+		{
+			GAsyncLoading2_DebugPackageIds.Reset();
+			ParsePackageNames(Variable->GetString(), GAsyncLoading2_DebugPackageIds);
+			ParsePackageNames(Variable->GetString(), GAsyncLoading2_VerbosePackageIds);
+			GAsyncLoading2_VerboseLogFilter = GAsyncLoading2_VerbosePackageIds.Num() > 0 ? 1 : 2;
+		}),
 	ECVF_Default);
 static FAutoConsoleVariableRef CVar_VerbosePackageNames(
 	TEXT("s.VerbosePackageNames"),
 	GAsyncLoading2_VerbosePackageNamesString,
 	TEXT("Restrict verbose logging to listed package names."),
 	FConsoleVariableDelegate::CreateLambda([](IConsoleVariable* Variable)
-	{
-		GAsyncLoading2_VerbosePackageIds.Reset();
-		ParsePackageNames(Variable->GetString(), GAsyncLoading2_VerbosePackageIds);
-		GAsyncLoading2_VerboseLogFilter = GAsyncLoading2_VerbosePackageIds.Num() > 0 ? 1 : 2;
-	}),
+		{
+			GAsyncLoading2_VerbosePackageIds.Reset();
+			ParsePackageNames(Variable->GetString(), GAsyncLoading2_VerbosePackageIds);
+			GAsyncLoading2_VerboseLogFilter = GAsyncLoading2_VerbosePackageIds.Num() > 0 ? 1 : 2;
+		}),
 	ECVF_Default);
 #endif
 
@@ -460,7 +460,7 @@ static FAutoConsoleVariableRef CVarGRemoveUnreachableObjectsOnGT(
 	GRemoveUnreachableObjectsOnGT,
 	TEXT("Remove unreachable objects from GlobalImportStore on the GT from the GC callback NotifyUnreachableObjects (slow)."),
 	ECVF_Default
-	);
+);
 
 #if UE_BUILD_DEBUG || (UE_BUILD_DEVELOPMENT && !WITH_EDITOR)
 static bool GVerifyUnreachableObjects = true;
@@ -475,14 +475,14 @@ static FAutoConsoleVariableRef CVarGVerifyUnreachableObjects(
 	GVerifyUnreachableObjects,
 	TEXT("Run GlobalImportStore verifications for unreachable objects from the GC callback NotifyUnreachableObjects (slow)."),
 	ECVF_Default
-	);
+);
 
 static FAutoConsoleVariableRef CVarGVerifyObjectLoadFlags(
 	TEXT("s.VerifyObjectLoadFlags"),
 	GVerifyObjectLoadFlags,
 	TEXT("Run AsyncFlags verifications for all objects when finished loading from the GC callback NotifyUnreachableObjects (slow)."),
 	ECVF_Default
-	);
+);
 
 static bool GAllowMultithreadedLoading = false;
 static FAutoConsoleVariableRef CVarGAllowMultithreadedLoading(
@@ -490,7 +490,7 @@ static FAutoConsoleVariableRef CVarGAllowMultithreadedLoading(
 	GAllowMultithreadedLoading,
 	TEXT("Allow some loading steps to be fully multithreaded."),
 	ECVF_Default
-	);
+);
 
 static bool GOnlyProcessRequiredPackagesWhenSyncLoading = true;
 static FAutoConsoleVariableRef CVarGOnlyProcessRequiredPackagesWhenSyncLoading(
@@ -601,26 +601,26 @@ FString FormatPackageId(FPackageId PackageId)
 static bool AsyncLoading2_ShouldSkipLoadingExport(const EExportFilterFlags FilterFlags)
 {
 #if WITH_EDITOR
-		return false;
+	return false;
 #elif UE_SERVER
-		return !!(static_cast<uint32>(FilterFlags) & static_cast<uint32>(EExportFilterFlags::NotForServer));
+	return !!(static_cast<uint32>(FilterFlags) & static_cast<uint32>(EExportFilterFlags::NotForServer));
 #elif !WITH_SERVER_CODE
-		return !!(static_cast<uint32>(FilterFlags) & static_cast<uint32>(EExportFilterFlags::NotForClient));
+	return !!(static_cast<uint32>(FilterFlags) & static_cast<uint32>(EExportFilterFlags::NotForClient));
 #else
-		static const bool bIsDedicatedServer = !GIsClient && GIsServer;
-		static const bool bIsClientOnly = GIsClient && !GIsServer;
+	static const bool bIsDedicatedServer = !GIsClient && GIsServer;
+	static const bool bIsClientOnly = GIsClient && !GIsServer;
 
-		if (bIsDedicatedServer && static_cast<uint32>(FilterFlags) & static_cast<uint32>(EExportFilterFlags::NotForServer))
-		{
-			return true;
-		}
+	if (bIsDedicatedServer && static_cast<uint32>(FilterFlags) & static_cast<uint32>(EExportFilterFlags::NotForServer))
+	{
+		return true;
+	}
 
-		if (bIsClientOnly && static_cast<uint32>(FilterFlags) & static_cast<uint32>(EExportFilterFlags::NotForClient))
-		{
-			return true;
-		}
+	if (bIsClientOnly && static_cast<uint32>(FilterFlags) & static_cast<uint32>(EExportFilterFlags::NotForClient))
+	{
+		return true;
+	}
 
-		return false;
+	return false;
 #endif
 }
 
@@ -749,8 +749,8 @@ const TCHAR* LexToString(EPackageLoader PackageLoader)
 	case EPackageLoader::Zen: return TEXT("Zen");
 	case EPackageLoader::Unknown: return TEXT("Unknown");
 	default:
-		checkNoEntry();
-		return TEXT("Not Implemented");
+	checkNoEntry();
+	return TEXT("Not Implemented");
 	}
 };
 
@@ -773,7 +773,7 @@ struct FAsyncPackageDesc2
 	// It will be used as key when tracking active async packages
 	FPackageId UPackageId;
 	// The id of the package being loaded from disk
-	FPackageId PackageIdToLoad; 
+	FPackageId PackageIdToLoad;
 	// The name of the UPackage being loaded
 	// Set to none for imported packages up until the package summary has been serialized
 	FName UPackageName;
@@ -866,7 +866,7 @@ private:
 	friend class FPublicExportMapBasicTest;
 	friend class FPublicExportMapIterationTest;
 	friend class FPublicExportMapPreInsertPublicExportsTest;
-	
+
 	class FPublicExportMap
 	{
 	public:
@@ -948,7 +948,7 @@ private:
 		void PreInsertPublicExports(TConstArrayView<TConstArrayView<FExportMapEntry>> ExportLists)
 		{
 			if (Capacity <= 1)
-			{ 
+			{
 				int32 PublicExportCount = 0;
 				uint64 LastNewKeySeen = 0;
 				for (TConstArrayView<FExportMapEntry> Exports : ExportLists)
@@ -1050,7 +1050,7 @@ private:
 						}
 					}
 				}
-				
+
 				// We expect the most common case here to be KeysToAdd.Num() == 0 as we're likely to be reloading the same package data.
 				// If we aren't, it implies some UGC/delivered content scenario where the old content was not fully GC'd before attemping 
 				// to reload the newly downloaded version.
@@ -1077,7 +1077,7 @@ private:
 							NewValues[InsertIndex] = OldValues[OldIndex];
 							++OldIndex;
 						}
-						else 
+						else
 						{
 							NewKeys[InsertIndex] = KeysToAdd[ToAddIndex];
 							NewValues[InsertIndex] = InvalidValue;
@@ -1238,7 +1238,7 @@ private:
 			int32 Pos;
 		public:
 			FValueIterator(const TArrayView<int32>& Data) : Array(Data), Pos(0)
-			{ 
+			{
 				const int32 ArraySize = Array.Num();
 
 				// Move to the first valid element
@@ -1307,7 +1307,7 @@ private:
 			uint8* Allocation = nullptr;
 			uint64 SingleItemKey;
 		};
-		
+
 		// Capacity is the total number of keys stored, not necessarily number of keys that map to valid values
 		// Capacity == 1 means we use SingleItemKey and SingleItemValue 
 		// Capacity > 1 means we use Allocation 
@@ -1340,7 +1340,7 @@ public:
 	FLoadedPackageRef& operator=(const FLoadedPackageRef& Other) = delete;
 
 	FLoadedPackageRef& operator=(FLoadedPackageRef&& Other) = default;
-	
+
 	inline int32 GetRefCount() const
 	{
 		return RefCount;
@@ -1790,7 +1790,7 @@ public:
 								*FoundPackage->GetName(),
 								*LexToString(PackageId)
 							);
-							
+
 							check(OldPackageRef->GetRefCount() == 0);
 							RemoveUnreferencedObsoletePackageNoLock(*OldPackageRef);
 							RemovePackageNoLock(OldPackageId);
@@ -2234,7 +2234,7 @@ public:
 				TEXT("The existing object will be replaced since it or its package was most likely renamed after it was loaded the first time."),
 				Object ? *Object->GetFullName() : TEXT("null"),
 				ObjectIndex,
-				*FormatPackageId(Key.GetPackageId()), 
+				*FormatPackageId(Key.GetPackageId()),
 				Key.GetExportHash(),
 				*ExistingObject->GetFullName(),
 				ExistingObject->GetFlags(),
@@ -2330,7 +2330,7 @@ public:
 #endif
 };
 
-struct FAsyncPackageHeaderData: public FZenPackageHeader
+struct FAsyncPackageHeaderData : public FZenPackageHeader
 {
 	// Backed by allocation in FAsyncPackageData
 	TArrayView<FPackageId> ImportedPackageIds;
@@ -2659,7 +2659,7 @@ public:
 		Cell = nullptr;
 	}
 
-	virtual FArchive& operator<<( UObject*& Object ) override
+	virtual FArchive& operator<<(UObject*& Object) override
 	{
 		FPackageIndex Index;
 		FArchive& Ar = *this;
@@ -2681,7 +2681,7 @@ public:
 				FName ObjectName = HeaderData->NameMap.GetName(Export.ObjectName);
 				UE_ASYNC_PACKAGE_CLOG_VERBOSE(!Object, VeryVerbose, *PackageDesc,
 					TEXT("FExportArchive: Object"), TEXT("Export %s at index %d is null."),
-					*ObjectName.ToString(), 
+					*ObjectName.ToString(),
 					ExportIndex);
 #endif
 			}
@@ -2710,7 +2710,7 @@ public:
 	}
 
 #if WITH_VERSE_VM || defined(__INTELLISENSE__)
-	virtual FArchive& operator<<( Verse::VCell*& Cell ) override
+	virtual FArchive& operator<<(Verse::VCell*& Cell) override
 	{
 		FPackageIndex Index;
 		FArchive& Ar = *this;
@@ -2802,7 +2802,7 @@ public:
 	inline virtual bool SerializeBulkData(FBulkData& BulkData, const FBulkDataSerializationParams& Params) override
 	{
 		const FPackageId& PackageId = PackageDesc->PackageIdToLoad;
-		const uint16 ChunkIndex = bIsOptionalSegment ? 1 : 0; 
+		const uint16 ChunkIndex = bIsOptionalSegment ? 1 : 0;
 
 		UE::BulkData::Private::FBulkMetaData& Meta = BulkData.BulkMeta;
 		FBulkDataCookedIndex CookedIndex;
@@ -2812,16 +2812,16 @@ public:
 		const bool bIsInline = Meta.HasAnyFlags(BULKDATA_PayloadAtEndOfFile) == false;
 		if (bIsInline)
 		{
-PRAGMA_DISABLE_DEPRECATION_WARNINGS
-			if (int64 PayloadSize = Meta.GetSize(); PayloadSize > 0 && Meta.HasAnyFlags(BULKDATA_Unused) == false)
-PRAGMA_ENABLE_DEPRECATION_WARNINGS
-			{
-				// Set the offset from the beginning of the I/O chunk in order to make inline bulk data reloadable
-				const int64 ExportBundleChunkOffset = (ActiveFPLB->StartFastPathLoadBuffer - ActiveFPLB->OriginalFastPathLoadBuffer);
-				Meta.SetOffset(ExportBundleChunkOffset);
-				BulkData.BulkChunkId = CreateIoChunkId(PackageId.Value(), ChunkIndex, EIoChunkType::ExportBundleData);
-				Serialize(BulkData.ReallocateData(PayloadSize), PayloadSize);
-			}
+			PRAGMA_DISABLE_DEPRECATION_WARNINGS
+				if (int64 PayloadSize = Meta.GetSize(); PayloadSize > 0 && Meta.HasAnyFlags(BULKDATA_Unused) == false)
+					PRAGMA_ENABLE_DEPRECATION_WARNINGS
+				{
+					// Set the offset from the beginning of the I/O chunk in order to make inline bulk data reloadable
+					const int64 ExportBundleChunkOffset = (ActiveFPLB->StartFastPathLoadBuffer - ActiveFPLB->OriginalFastPathLoadBuffer);
+					Meta.SetOffset(ExportBundleChunkOffset);
+					BulkData.BulkChunkId = CreateIoChunkId(PackageId.Value(), ChunkIndex, EIoChunkType::ExportBundleData);
+					Serialize(BulkData.ReallocateData(PayloadSize), PayloadSize);
+				}
 		}
 		else if (Meta.HasAnyFlags(BULKDATA_MemoryMappedPayload))
 		{
@@ -2859,7 +2859,7 @@ PRAGMA_ENABLE_DEPRECATION_WARNINGS
 #if UE_DISABLE_COOKEDINDEX_FOR_NONDUPLICATE
 				CookedIndex = FBulkDataCookedIndex::Default;
 #endif //UE_DISABLE_COOKEDINDEX_FOR_NONDUPLICATE
-				
+
 				const FIoChunkId OptionalChunkId = CreateBulkDataIoChunkId(PackageId.Value(), ChunkIndex, CookedIndex.GetValue(), EIoChunkType::OptionalBulkData);
 
 				if (IoDispatcher.DoesChunkExist(OptionalChunkId))
@@ -2884,7 +2884,7 @@ private:
 
 		if (UNLIKELY(HeaderData->BulkDataMap.IsEmpty()))
 		{
-			FBulkMetaData::FromSerialized(Ar, ElementSize,  Meta, DuplicateSerialOffset);
+			FBulkMetaData::FromSerialized(Ar, ElementSize, Meta, DuplicateSerialOffset);
 		}
 		else
 		{
@@ -2956,14 +2956,10 @@ enum class EAsyncPackageLoadingState2 : uint8
 	ProcessDependencies,
 	WaitingForDependencies,
 	DependenciesReady,
-// This is the path taken by LinkerLoad packages
-	CreateLinkerLoadExports,
-	WaitingForLinkerLoadDependencies,
-	ResolveLinkerLoadImports,
-	PreloadLinkerLoadExports,
-// This is the path taken by Runtime/cooked packages
-	ProcessExportBundles,
-// Both LinkerLoad and Cooked packages should converge at this point
+	CreateExports,
+	WaitingForImportDependencies,
+	ResolveImports,
+	PreloadExports,
 	WaitingForExternalReads,
 	ExportsDone,
 	PostLoad,
@@ -2978,7 +2974,7 @@ enum class EAsyncPackageLoadingState2 : uint8
 
 const TCHAR* LexToString(EAsyncPackageLoadingState2 AsyncPackageLoadingState)
 {
-	switch(AsyncPackageLoadingState)
+	switch (AsyncPackageLoadingState)
 	{
 	case EAsyncPackageLoadingState2::NewPackage: return TEXT("NewPackage");
 	case EAsyncPackageLoadingState2::WaitingForIo: return TEXT("WaitingForIo");
@@ -2986,11 +2982,10 @@ const TCHAR* LexToString(EAsyncPackageLoadingState2 AsyncPackageLoadingState)
 	case EAsyncPackageLoadingState2::ProcessDependencies: return TEXT("ProcessDependencies");
 	case EAsyncPackageLoadingState2::WaitingForDependencies: return TEXT("WaitingForDependencies");
 	case EAsyncPackageLoadingState2::DependenciesReady: return TEXT("DependenciesReady");
-	case EAsyncPackageLoadingState2::CreateLinkerLoadExports: return TEXT("CreateLinkerLoadExports");
-	case EAsyncPackageLoadingState2::WaitingForLinkerLoadDependencies: return TEXT("WaitingForLinkerLoadDependencies");
-	case EAsyncPackageLoadingState2::ResolveLinkerLoadImports: return TEXT("ResolveLinkerLoadImports");
-	case EAsyncPackageLoadingState2::PreloadLinkerLoadExports: return TEXT("PreloadLinkerLoadExports");
-	case EAsyncPackageLoadingState2::ProcessExportBundles: return TEXT("ProcessExportBundles");
+	case EAsyncPackageLoadingState2::CreateExports: return TEXT("CreateExports");
+	case EAsyncPackageLoadingState2::WaitingForImportDependencies: return TEXT("WaitingForImportDependencies");
+	case EAsyncPackageLoadingState2::ResolveImports: return TEXT("ResolveImports");
+	case EAsyncPackageLoadingState2::PreloadExports: return TEXT("PreloadExports");
 	case EAsyncPackageLoadingState2::WaitingForExternalReads: return TEXT("WaitingForExternalReads");
 	case EAsyncPackageLoadingState2::ExportsDone: return TEXT("ExportsDone");
 	case EAsyncPackageLoadingState2::PostLoad: return TEXT("PostLoad");
@@ -3002,8 +2997,8 @@ const TCHAR* LexToString(EAsyncPackageLoadingState2 AsyncPackageLoadingState)
 	case EAsyncPackageLoadingState2::Complete: return TEXT("Complete");
 	case EAsyncPackageLoadingState2::DeferredDelete: return TEXT("DeferredDelete");
 	default:
-		checkNoEntry();
-		return TEXT("Unknown");
+	checkNoEntry();
+	return TEXT("Unknown");
 	}
 };
 
@@ -3065,7 +3060,7 @@ private:
 	uint32 DependenciesCapacity = 0;
 	int32 Priority = 0;
 	int32 ImportOrExportIndex = -1;
-	TAsyncAtomic<int32> BarrierCount { 0 };
+	TAsyncAtomic<int32> BarrierCount{ 0 };
 	enum EQueueStatus
 	{
 		QueueStatus_None = 0,
@@ -3073,8 +3068,8 @@ private:
 		QueueStatus_External = 2
 	};
 	uint8 QueueStatus = QueueStatus_None;
-	TAsyncAtomic<bool> bIsUpdatingDependencies { false };
-	TAsyncAtomic<bool> bIsDone { false };
+	TAsyncAtomic<bool> bIsUpdatingDependencies{ false };
+	TAsyncAtomic<bool> bIsDone{ false };
 };
 
 class FAsyncLoadEventGraphAllocator
@@ -3098,11 +3093,11 @@ public:
 		TotalArcCount -= Count;
 	}
 
-	TAsyncAtomic<int64> TotalArcCount { 0 };
-	TAsyncAtomic<int64> TotalAllocated { 0 };
+	TAsyncAtomic<int64> TotalArcCount{ 0 };
+	TAsyncAtomic<int64> TotalAllocated{ 0 };
 };
 
-/** Queue based on DepletableMpscQueue.h 
+/** Queue based on DepletableMpscQueue.h
 	Producers will add nodes to the linked list defined by Sentinel->Tail
 	Nodes will be executed from the linked list defined by LocalHead->LocalTail
 	All the nodes from the producer queue will be moved to the local queue when needed */
@@ -3218,7 +3213,7 @@ struct FAsyncLoadingThreadState2
 				ThreadState.StartTime = 0.0;
 			}
 		}
-		
+
 		~FTimeLimitScope()
 		{
 			ThreadState.bUseTimeLimit = bUseTimeLimit;
@@ -3266,7 +3261,7 @@ struct FAsyncLoadingThreadState2
 		if (!bTimeLimitExceeded)
 		{
 			const bool bGCIsWaiting = IsGarbageCollectionWaiting();
-			
+
 			// Do not time out sync loads due to GC 
 			if (bGCIsWaiting && SyncLoadContextStack.IsEmpty() && !IsInParallelLoadingThread())
 			{
@@ -3287,7 +3282,7 @@ struct FAsyncLoadingThreadState2
 	{
 		bIsActive.store(true, std::memory_order_relaxed);
 	}
-	
+
 	void ResetActivity()
 	{
 		bIsActive.store(false, std::memory_order_relaxed);
@@ -3311,7 +3306,7 @@ struct FAsyncLoadingThreadState2
 	TAsyncAtomic<bool> bIsActive = false;
 	bool bShouldFireNodes = true;
 	bool bUseTimeLimit = false;
-	
+
 	double TimeLimit = 0.0;
 	double StartTime = 0.0;
 	double LastTestTime = -1.0;
@@ -3329,16 +3324,13 @@ enum EEventLoadNode2 : uint8
 	Package_ProcessSummary,
 	Package_ProcessDependencies,
 	Package_DependenciesReady,
-	Package_CreateLinkerLoadExports,
-	Package_ResolveLinkerLoadImports,
-	Package_PreloadLinkerLoadExports,
+	Package_CreateExports,
+	Package_ResolveImports,
+	Package_PreloadExports,
 	Package_ExportsSerialized,
+	Package_PostLoad,
+	Package_DeferredPostLoad,
 	Package_NumPhases,
-
-	ExportBundle_Process = 0,
-	ExportBundle_PostLoad,
-	ExportBundle_DeferredPostLoad,
-	ExportBundle_NumPhases,
 };
 
 struct FAsyncPackageData
@@ -3348,7 +3340,6 @@ struct FAsyncPackageData
 	TArrayView<FExportObject> Exports;
 	TArrayView<FExportCell> CellExports;
 	TArrayView<FAsyncPackage2*> ImportedAsyncPackages;
-	TArrayView<FEventLoadNode2> ExportBundleNodes;
 	TArrayView<const FSHAHash> ShaderMapHashes;
 	int32 TotalExportBundleCount = 0;
 	TAsyncAtomic<bool> bCellExportsInitialized = false;
@@ -3470,30 +3461,30 @@ private:
 	void StartSizeRequests()
 	{
 		FAsyncFileCallBack UassetSizeRequestCallbackFunction = [this](bool bWasCancelled, IAsyncReadRequest* Request)
-		{
-			if (!bWasCancelled)
-			{
-				UassetSize = Request->GetSizeResults();
-			}
-			if (--PendingSizeRequests == 0)
-			{
-				StartReadRequests();
-			}
-		};
-		UassetSizeRequest = UassetFileHandle->SizeRequest(&UassetSizeRequestCallbackFunction);
-		if (UexpFileHandle)
-		{
-			FAsyncFileCallBack UexpSizeRequestCallbackFunction = [this](bool bWasCancelled, IAsyncReadRequest* Request)
 			{
 				if (!bWasCancelled)
 				{
-					UexpSize = Request->GetSizeResults();
+					UassetSize = Request->GetSizeResults();
 				}
 				if (--PendingSizeRequests == 0)
 				{
 					StartReadRequests();
 				}
 			};
+		UassetSizeRequest = UassetFileHandle->SizeRequest(&UassetSizeRequestCallbackFunction);
+		if (UexpFileHandle)
+		{
+			FAsyncFileCallBack UexpSizeRequestCallbackFunction = [this](bool bWasCancelled, IAsyncReadRequest* Request)
+				{
+					if (!bWasCancelled)
+					{
+						UexpSize = Request->GetSizeResults();
+					}
+					if (--PendingSizeRequests == 0)
+					{
+						StartReadRequests();
+					}
+				};
 			UexpSizeRequest = UexpFileHandle->SizeRequest(&UexpSizeRequestCallbackFunction);
 		}
 	}
@@ -3501,16 +3492,16 @@ private:
 	void StartReadRequests()
 	{
 		FAsyncFileCallBack ReadRequestCallbackFunction = [this](bool bWasCancelled, IAsyncReadRequest* Request)
-		{
-			if (bWasCancelled)
-			{
-				bFailed = true;
-			}
-			if (--PendingReadRequests == 0)
 			{
-				FinishedReading();
-			}
-		};
+				if (bWasCancelled)
+				{
+					bFailed = true;
+				}
+				if (--PendingReadRequests == 0)
+				{
+					FinishedReading();
+				}
+			};
 		if (UassetSize <= 0 || (UexpFileHandle && UexpSize <= 0))
 		{
 			SetError();
@@ -3773,9 +3764,11 @@ private:
 	/** Current index into Export objects array used to spread routing PostLoadInstances over several frames			*/
 	int32						PostLoadInstanceIndex = 0;
 	/** Current loading state of a package. */
-	TAsyncAtomic<EAsyncPackageLoadingState2> AsyncPackageLoadingState { EAsyncPackageLoadingState2::NewPackage };
+	TAsyncAtomic<EAsyncPackageLoadingState2> AsyncPackageLoadingState{ EAsyncPackageLoadingState2::NewPackage };
 	/** Whether the constructed objects have been moved to the second loading phase. */
 	TAsyncAtomic<bool> ObjectsNowInPhase2 = false;
+	/** Protects against the same package being serialized from different threads. */
+	UE::FMutex                  SerializationMutex;
 
 	struct FAllDependenciesState
 	{
@@ -3806,7 +3799,7 @@ private:
 			check(WaitListPackage != PackageToAdd);
 			FAllDependenciesState& WaitListPackageState = WaitListPackage->*StateMemberPtr;
 			FAllDependenciesState& PackageToAddState = PackageToAdd->*StateMemberPtr;
-			
+
 			if (PackageToAddState.WaitingForPackage == WaitListPackage)
 			{
 				return;
@@ -3885,7 +3878,7 @@ private:
 	bool						bCreatedLinkerRoot = false;
 	/** List of all request handles. */
 	TArray<int32> RequestIDs;
-	
+
 	UE::FMutex ConstructedObjectsMutex;
 	/** List of ConstructedObjects = Exports + UPackage + ObjectsCreatedFromExports */
 	TArray<UObject*> ConstructedObjects;
@@ -3911,18 +3904,15 @@ public:
 
 	static FAsyncPackage2* GetCurrentlyExecutingPackage(FAsyncLoadingThreadState2& ThreadState, FAsyncPackage2* PackageToFilter = nullptr);
 
-	/** [EDL] Begin Event driven loader specific stuff */
-
-	static EEventLoadNodeExecutionResult Event_ProcessExportBundle(FAsyncLoadingThreadState2& ThreadState, FAsyncPackage2* Package, int32 ExportBundleIndex);
-	static EEventLoadNodeExecutionResult Event_CreateLinkerLoadExports(FAsyncLoadingThreadState2& ThreadState, FAsyncPackage2* Package, int32);
-	static EEventLoadNodeExecutionResult Event_ResolveLinkerLoadImports(FAsyncLoadingThreadState2& ThreadState, FAsyncPackage2* Package, int32);
-	static EEventLoadNodeExecutionResult Event_PreloadLinkerLoadExports(FAsyncLoadingThreadState2& ThreadState, FAsyncPackage2* Package, int32);
+	static EEventLoadNodeExecutionResult Event_CreateExports(FAsyncLoadingThreadState2& ThreadState, FAsyncPackage2* Package, int32);
+	static EEventLoadNodeExecutionResult Event_ResolveImports(FAsyncLoadingThreadState2& ThreadState, FAsyncPackage2* Package, int32);
+	static EEventLoadNodeExecutionResult Event_PreloadExports(FAsyncLoadingThreadState2& ThreadState, FAsyncPackage2* Package, int32);
 	static EEventLoadNodeExecutionResult Event_ProcessPackageSummary(FAsyncLoadingThreadState2& ThreadState, FAsyncPackage2* Package, int32);
 	static EEventLoadNodeExecutionResult Event_ProcessDependencies(FAsyncLoadingThreadState2& ThreadState, FAsyncPackage2* Package, int32);
 	static EEventLoadNodeExecutionResult Event_DependenciesReady(FAsyncLoadingThreadState2& ThreadState, FAsyncPackage2* Package, int32);
 	static EEventLoadNodeExecutionResult Event_ExportsDone(FAsyncLoadingThreadState2& ThreadState, FAsyncPackage2* Package, int32);
-	static EEventLoadNodeExecutionResult Event_PostLoadExportBundle(FAsyncLoadingThreadState2& ThreadState, FAsyncPackage2* Package, int32 ExportBundleIndex);
-	static EEventLoadNodeExecutionResult Event_DeferredPostLoadExportBundle(FAsyncLoadingThreadState2& ThreadState, FAsyncPackage2* Package, int32 ExportBundleIndex);
+	static EEventLoadNodeExecutionResult Event_PostLoad(FAsyncLoadingThreadState2& ThreadState, FAsyncPackage2* Package, int32 ExportBundleIndex);
+	static EEventLoadNodeExecutionResult Event_DeferredPostLoad(FAsyncLoadingThreadState2& ThreadState, FAsyncPackage2* Package, int32 ExportBundleIndex);
 
 	bool CanRunNodeAsync(const FAsyncLoadEventSpec* Spec) const;
 
@@ -3945,7 +3935,6 @@ public:
 	}
 
 	FEventLoadNode2& GetPackageNode(EEventLoadNode2 Phase);
-	FEventLoadNode2& GetExportBundleNode(EEventLoadNode2 Phase, uint32 ExportBundleIndex);
 
 	/** [EDL] End Event driven loader specific stuff */
 
@@ -3954,7 +3943,6 @@ public:
 private:
 	void InitializeExportArchive(FExportArchive& Ar, bool bIsOptionalSegment);
 	void CreatePackageNodes(const FAsyncLoadEventSpec* EventSpecs);
-	void CreateExportBundleNodes(const FAsyncLoadEventSpec* EventSpecs);
 	void SetupScriptDependencies();
 #if ALT2_ENABLE_PACKAGE_DEPENDENCY_DEBUGGING
 	bool HasDependencyToPackageDebug(FAsyncPackage2* Package);
@@ -3989,9 +3977,12 @@ private:
 	void ConditionalBeginDeferredPostLoad(FAsyncLoadingThreadState2& ThreadState);
 	EEventLoadNodeExecutionResult ProcessLinkerLoadPackageSummary(FAsyncLoadingThreadState2& ThreadState);
 	bool CreateLinkerLoadExports(FAsyncLoadingThreadState2& ThreadState);
-	void ConditionalBeginResolveLinkerLoadImports(FAsyncLoadingThreadState2& ThreadState);
+	bool CreateRuntimeExports(FAsyncLoadingThreadState2& ThreadState);
+	void ConditionalBeginResolveImports(FAsyncLoadingThreadState2& ThreadState);
 	bool ResolveLinkerLoadImports(FAsyncLoadingThreadState2& ThreadState);
+	bool ResolveRuntimeImports(FAsyncLoadingThreadState2& ThreadState);
 	bool PreloadLinkerLoadExports(FAsyncLoadingThreadState2& ThreadState);
+	bool PreloadRuntimeExports(FAsyncLoadingThreadState2& ThreadState);
 	EEventLoadNodeExecutionResult ExecutePostLoadLinkerLoadPackageExports(FAsyncLoadingThreadState2& ThreadState);
 	EEventLoadNodeExecutionResult ExecuteDeferredPostLoadLinkerLoadPackageExports(FAsyncLoadingThreadState2& ThreadState);
 
@@ -4201,7 +4192,7 @@ private:
 	TMpscQueue<FAsyncPackage2*> DeferredDeletePackages;
 	/** Used to protect DeferredDeletePackages when dequeuing since its a single consumer queue. */
 	UE::FMutex DeferredDeletePackagesLock;
-	
+
 	struct FCompletedPackageRequest
 	{
 		FName PackageName;
@@ -4217,7 +4208,7 @@ private:
 			EAsyncLoadingResult::Type Result,
 			TUniquePtr<FLoadPackageAsyncDelegate>&& CompletionCallback)
 		{
-			FCompletedPackageRequest Res = 
+			FCompletedPackageRequest Res =
 			{
 				.PackageName = Desc.UPackageName,
 				.Result = Result
@@ -4251,7 +4242,7 @@ private:
 #if WITH_EDITOR
 			UE::Core::Private::FPlayInEditorLoadingScope PlayInEditorIDScope(AsyncPackage ? AsyncPackage->Desc.PIEInstanceID : INDEX_NONE);
 #endif
-			
+
 			if (CompletionCallbacks.Num() != 0)
 			{
 				TRACE_CPUPROFILER_EVENT_SCOPE(PackageCompletionCallbacks);
@@ -4276,7 +4267,7 @@ private:
 						.LoadedPackage = UPackage,
 						.ProgressType = ProgressType
 					};
-					
+
 					ProgressCallback->Invoke(Params);
 				}
 				ProgressCallbacks.Empty();
@@ -4292,7 +4283,7 @@ private:
 
 	TMpscQueue<FAsyncPackage2*> ExternalReadQueue;
 	TAsyncAtomic<int32> PendingIoRequestsCounter{ 0 };
-	
+
 	/** List of all pending package requests */
 	TSet<int32> PendingRequests;
 	/** Synchronization object for PendingRequests list */
@@ -4300,23 +4291,23 @@ private:
 	TMap<int32, FAsyncPackage2*> RequestIdToPackageMap; // Only accessed from the async loading thread
 
 	/** [ASYNC/GAME THREAD] Number of package load requests in the async loading queue */
-	TAsyncAtomic<int32> QueuedPackagesCounter { 0 };
+	TAsyncAtomic<int32> QueuedPackagesCounter{ 0 };
 	/** [ASYNC/GAME THREAD] Number of packages being loaded on the async thread and post loaded on the game thread */
-	TAsyncAtomic<int32> LoadingPackagesCounter { 0 };
+	TAsyncAtomic<int32> LoadingPackagesCounter{ 0 };
 	/** Encapsulate our counter to make sure all decrements are monitored. */
 	class FPackagesWithRemainingWorkCounter
 	{
 		UE::FManualResetEvent* WakeEvent = nullptr;
-		TAsyncAtomic<int32> PackagesWithRemainingWorkCounter {0};
+		TAsyncAtomic<int32> PackagesWithRemainingWorkCounter{ 0 };
 	public:
-		void SetWakeEvent(UE::FManualResetEvent* InWakeEvent) {	WakeEvent = InWakeEvent; }
+		void SetWakeEvent(UE::FManualResetEvent* InWakeEvent) { WakeEvent = InWakeEvent; }
 		int32 operator++() { return ++PackagesWithRemainingWorkCounter; }
 		int32 operator++(int) { return PackagesWithRemainingWorkCounter++; }
 		operator int32() const { return PackagesWithRemainingWorkCounter; }
 
 		// Only implement prefix
-		int32 operator--() 
-		{ 
+		int32 operator--()
+		{
 			int32 newValue = --PackagesWithRemainingWorkCounter;
 			if (newValue == 0)
 			{
@@ -4719,7 +4710,7 @@ private:
 							UnloadableImportPackageId = ImportId;
 							break;
 						}
-						else if(bLocalFailLoadOnNotInstalledImport && PackageRef->IsNotInstalled())
+						else if (bLocalFailLoadOnNotInstalledImport && PackageRef->IsNotInstalled())
 						{
 							bFoundNotInstalled |= PackageRef->IsNotInstalled();
 							UnloadableImportPackageId = ImportId;
@@ -4946,15 +4937,12 @@ private:
 #endif
 		const int32 ShaderMapHashesCount = PackageStoreEntry.ShaderMapHashes.Num();
 
-		const int32 ExportBundleNodeCount = TotalExportBundleCount * EEventLoadNode2::ExportBundle_NumPhases;
-		const uint64 ExportBundleNodesMemSize = Align(sizeof(FEventLoadNode2) * ExportBundleNodeCount, 8);
 		const uint64 ImportedPackagesMemSize = Align(sizeof(FAsyncPackage2*) * TotalImportedPackagesCount, 8);
 		const uint64 ShaderMapHashesMemSize = Align(sizeof(FSHAHash) * ShaderMapHashesCount, 8);
 		const uint64 MemoryBufferSize =
 #if WITH_EDITOR
 			OptionalSegmentImportedPackageIdsMemSize +
 #endif
-			ExportBundleNodesMemSize +
 			ImportedPackageIdsMemSize +
 			ImportedPackagesMemSize +
 			ShaderMapHashesMemSize;
@@ -4964,9 +4952,6 @@ private:
 		uint8* DataPtr = Data.MemoryBuffer0;
 
 		Data.TotalExportBundleCount = TotalExportBundleCount;
-		Data.ExportBundleNodes = MakeArrayView(reinterpret_cast<FEventLoadNode2*>(DataPtr), ExportBundleNodeCount);
-		DataPtr += ExportBundleNodesMemSize;
-		AsyncPackage->CreateExportBundleNodes(EventSpecs.GetData());
 
 		Data.ShaderMapHashes = MakeArrayView(reinterpret_cast<const FSHAHash*>(DataPtr), ShaderMapHashesCount);
 		FMemory::Memcpy((void*)Data.ShaderMapHashes.GetData(), PackageStoreEntry.ShaderMapHashes.GetData(), sizeof(FSHAHash) * ShaderMapHashesCount);
@@ -5038,7 +5023,7 @@ private:
 			ExportBundleEntriesCopyMemSize;
 
 		Data.MemoryBuffer1 = reinterpret_cast<uint8*>(FMemory::Malloc(MemoryBufferSize));
-		
+
 		uint8* DataPtr = Data.MemoryBuffer1;
 
 		Data.Exports = MakeArrayView(reinterpret_cast<FExportObject*>(DataPtr), TotalExportCount);
@@ -5363,7 +5348,7 @@ void FAsyncLoadingThread2::InitializeLoading()
 	{
 		AltZenaphore.NotifyOne();
 	});
-	
+
 	AsyncThreadReady.fetch_add(1);
 
 	UE_LOG(LogStreaming, Log, TEXT("AsyncLoading2 - Initialized"));
@@ -5420,7 +5405,7 @@ void FAsyncLoadingThread2::UpdatePackagePriorityRecursive(FAsyncLoadingThreadSta
 
 	// For LinkerLoad, dependencies are loaded in ProcessPackageSummary so if this is called before,
 	// early out to avoid races and we'll reprocess them in the ProcessDependencies step
-	if (Package->LinkerLoadState.IsSet() && 
+	if (Package->LinkerLoadState.IsSet() &&
 		Package->AsyncPackageLoadingState <= EAsyncPackageLoadingState2::ProcessPackageSummary)
 	{
 		return;
@@ -5457,15 +5442,12 @@ void FAsyncPackage2::ConditionalBeginDeferredPostLoad(FAsyncLoadingThreadState2&
 	// Move everything to the second phase so that StaticFind from inside postload can
 	// see objects ready for postload.
 	MoveConstructedObjectsToPhase2();
-	
+
 	check(AsyncPackageLoadingState == EAsyncPackageLoadingState2::PostLoad);
 	AsyncPackageLoadingState = EAsyncPackageLoadingState2::DeferredPostLoad;
 	ConditionalReleasePartialRequests(ThreadState);
 
-	for (int32 ExportBundleIndex = 0; ExportBundleIndex < Data.TotalExportBundleCount; ++ExportBundleIndex)
-	{
-		GetExportBundleNode(EEventLoadNode2::ExportBundle_DeferredPostLoad, ExportBundleIndex).ReleaseBarrier(&ThreadState);
-	}
+	GetPackageNode(Package_DeferredPostLoad).ReleaseBarrier(&ThreadState);
 }
 
 FAsyncPackage2* FAsyncLoadingThread2::FindOrInsertPackage(FAsyncLoadingThreadState2& ThreadState, FAsyncPackageDesc2& Desc, bool& bInserted, FAsyncPackage2* ImportedByPackage, TUniquePtr<FLoadPackageAsyncDelegate>&& PackageLoadedDelegate, TSharedPtr<FLoadPackageAsyncProgressDelegate> PackageProgressDelegate)
@@ -5580,14 +5562,14 @@ void FAsyncLoadingThread2::IncludePackageInSyncLoadContextRecursive(FAsyncLoadin
 }
 
 UE_TRACE_EVENT_BEGIN(CUSTOM_LOADTIMER_LOG, CreateAsyncPackage, NoSync)
-	UE_TRACE_EVENT_FIELD(uint64, PackageId)
-	UE_TRACE_EVENT_FIELD(UE::Trace::WideString, PackageName)
+UE_TRACE_EVENT_FIELD(uint64, PackageId)
+UE_TRACE_EVENT_FIELD(UE::Trace::WideString, PackageName)
 UE_TRACE_EVENT_END()
 
 bool FAsyncLoadingThread2::CreateAsyncPackagesFromQueue(FAsyncLoadingThreadState2& ThreadState, bool bForceExitForGarbageCollect)
 {
 	TRACE_CPUPROFILER_EVENT_SCOPE(CreateAsyncPackagesFromQueue);
-	
+
 	// Package creation needs access to all objects
 	TGuardValue GuardVisibilityFilter(FUObjectThreadContext::Get().AsyncVisibilityFilter, EInternalObjectFlags::None);
 
@@ -5656,7 +5638,7 @@ bool FAsyncLoadingThread2::CreateAsyncPackagesFromQueue(FAsyncLoadingThreadState
 			--QueuedPackagesCounter;
 			++NumDequeued;
 			TRACE_COUNTER_SET(AsyncLoadingQueuedPackages, QueuedPackagesCounter);
-		
+
 			FPackageRequest& Request = OptionalRequest.GetValue();
 			EPackageStoreEntryStatus PackageStatus = EPackageStoreEntryStatus::Missing;
 			EPackageLoader PackageLoader = EPackageLoader::Zen;
@@ -5698,14 +5680,14 @@ bool FAsyncLoadingThread2::CreateAsyncPackagesFromQueue(FAsyncLoadingThreadState
 
 			FPackagePath::TryFromPackageName(PackageNameToLoad, Request.PackagePath);
 			EPackageExtension Extension;
-		    if (!PackageNameToLoad.IsNone() && TryGetExistingLoadedPackagePath(PackageIdToLoad, PackageNameToLoad, Extension, PackageLoader, PackageEntry))
-		    {
+			if (!PackageNameToLoad.IsNone() && TryGetExistingLoadedPackagePath(PackageIdToLoad, PackageNameToLoad, Extension, PackageLoader, PackageEntry))
+			{
 				Request.PackagePath.SetHeaderExtension(Extension);
-			    PackageStatus = EPackageStoreEntryStatus::Ok;
-		    }
-		    else
-		    {
-			    PackageStatus = TryGetPackageInfoFromPackageStore(PackageIdToLoad, PackageNameToLoad, UPackageName, PackageEntry);
+				PackageStatus = EPackageStoreEntryStatus::Ok;
+			}
+			else
+			{
+				PackageStatus = TryGetPackageInfoFromPackageStore(PackageIdToLoad, PackageNameToLoad, UPackageName, PackageEntry);
 
 				// try again after looking for an matching active load request
 				if (PackageStatus == EPackageStoreEntryStatus::Missing)
@@ -5800,7 +5782,7 @@ bool FAsyncLoadingThread2::CreateAsyncPackagesFromQueue(FAsyncLoadingThreadState
 				RequestIdToPackageMap.Add(PackageDesc.RequestID, Package);
 			}
 		}
-		
+
 		bPackagesCreated |= NumDequeued > 0;
 
 		if (!NumDequeued || ThreadState.IsTimeLimitExceeded(TEXT("CreateAsyncPackagesFromQueue")))
@@ -5810,7 +5792,7 @@ bool FAsyncLoadingThread2::CreateAsyncPackagesFromQueue(FAsyncLoadingThreadState
 	}
 
 	IoBatch.Issue();
-	
+
 	return bPackagesCreated;
 }
 
@@ -5927,7 +5909,7 @@ void FEventLoadNode2::ParallelLoadingWork()
 #endif
 
 	FTaskTagScope Scope(ETaskTag::EParallelLoadingThread);
-	FAsyncLoadingThreadState2 LocalThreadState { .bIsWorkerThread = true };
+	FAsyncLoadingThreadState2 LocalThreadState{ .bIsWorkerThread = true };
 
 	FAsyncLoadingThreadState2Scope WorkerStateScope(LocalThreadState);
 
@@ -5942,7 +5924,7 @@ void FEventLoadNode2::ParallelLoadingWork()
 
 	EEventLoadNodeExecutionResult Result = Execute(LocalThreadState);
 	check(Result == EEventLoadNodeExecutionResult::Complete);
-	
+
 	ThreadContext.AsyncPackage = nullptr;
 
 	// We should never destroy a thread state that still has nodes to fire
@@ -5989,7 +5971,7 @@ void FEventLoadNode2::Fire(FAsyncLoadingThreadState2* ThreadState)
 
 		// fallthrough to push the task to the queue
 	}
-	
+
 	Spec->EventQueue->Push(ThreadState, this);
 }
 
@@ -6168,7 +6150,7 @@ bool FAsyncLoadEventQueue2::PopAndExecute(FAsyncLoadingThreadState2& ThreadState
 		return false;
 	}
 	Node->QueueStatus = FEventLoadNode2::QueueStatus_None;
-	
+
 	EEventLoadNodeExecutionResult Result = Node->Execute(ThreadState);
 	if (Result == EEventLoadNodeExecutionResult::Timeout)
 	{
@@ -6188,9 +6170,9 @@ bool FAsyncLoadEventQueue2::ExecuteSyncLoadEvents(FAsyncLoadingThreadState2& Thr
 	int32 ThisCallCounter = ++ExecuteSyncLoadEventsCallCounter;
 
 	auto ShouldExecuteNode = [&SyncLoadContext](FEventLoadNode2& Node) -> bool
-	{
-		return Node.Package->SyncLoadContextId >= SyncLoadContext.ContextId;
-	};
+		{
+			return Node.Package->SyncLoadContextId >= SyncLoadContext.ContextId;
+		};
 
 	bool bDidSomething = false;
 	if (TimedOutEventNode && ShouldExecuteNode(*TimedOutEventNode))
@@ -6243,35 +6225,24 @@ void FAsyncLoadEventQueue2::UpdatePackagePriority(FAsyncPackage2* Package)
 {
 	FScopeLock Lock(&ExternalCritical);
 	auto ReprioritizeNode = [this](FEventLoadNode2& Node)
-	{
-		if (Node.Spec->EventQueue == this && Node.Priority < Node.Package->Desc.Priority)
 		{
-			if (Node.QueueStatus == FEventLoadNode2::QueueStatus_Local)
+			if (Node.Spec->EventQueue == this && Node.Priority < Node.Package->Desc.Priority)
 			{
-				LocalQueue.Reprioritize(&Node, Node.Package->Desc.Priority);
-			}
-			else if (Node.QueueStatus == FEventLoadNode2::QueueStatus_External)
-			{
-				ExternalQueue.Reprioritize(&Node, Node.Package->Desc.Priority);
+				if (Node.QueueStatus == FEventLoadNode2::QueueStatus_Local)
+				{
+					LocalQueue.Reprioritize(&Node, Node.Package->Desc.Priority);
+				}
+				else if (Node.QueueStatus == FEventLoadNode2::QueueStatus_External)
+				{
+					ExternalQueue.Reprioritize(&Node, Node.Package->Desc.Priority);
+				}
 			}
-		}
-	};
+		};
 
 	for (FEventLoadNode2& Node : Package->PackageNodes)
 	{
 		ReprioritizeNode(Node);
 	}
-
-	// Only try to reprioritize export bundle nodes if they have been initialized
-	// otherwise they can't even be scheduled yet and will be prioritized correctly later.
-	if (Package->AsyncPackageLoadingState > EAsyncPackageLoadingState2::ProcessPackageSummary)
-	{
-		for (FEventLoadNode2& Node : Package->Data.ExportBundleNodes)
-		{
-			ReprioritizeNode(Node);
-		}
-	}
-
 	UpdateExternalQueueState();
 }
 
@@ -6431,10 +6402,10 @@ void FCookedScriptObjectsDebug::LoadDebugData()
 		for (const FScriptObjectEntry& Entry : Entries)
 		{
 			FName Name = NameMap[Entry.Mapped.GetIndex()].ToName(Entry.Mapped.GetNumber());
-			Objects.Emplace(Entry.GlobalIndex, {FMinimalName(Name), Entry.OuterIndex});
+			Objects.Emplace(Entry.GlobalIndex, { FMinimalName(Name), Entry.OuterIndex });
 		}
 		UE_LOG(LogStreaming, Log, TEXT("AsyncLoading2 - LoadDebugData: Loaded %d cooked script object entries (%.2f KB)"),
-			Objects.Num(), (float)Objects.GetAllocatedSize()/1024.f);
+			Objects.Num(), (float)Objects.GetAllocatedSize() / 1024.f);
 	}
 }
 
@@ -6453,23 +6424,23 @@ void FCookedScriptObjectsDebug::TrimDebugData(const TMap<FPackageObjectIndex, UO
 
 	int32 NumRemaining = Num;
 	TFunction<void(int32)> KeepOuters = [this, &KeepOuters, &Remove, &Visited, &NumRemaining](int32 I)
-	{
-		if (!Visited[I])
 		{
-			Visited[I] = true;
-			if (Remove[I])
-			{
-				++NumRemaining;
-				Remove[I] = false;
-			}
-			const TPair<FMinimalName, FPackageObjectIndex>& Entry = Objects.Get(FSetElementId::FromInteger(I)).Value;
-			if (!Entry.Value.IsNull())
+			if (!Visited[I])
 			{
-				FSetElementId OuterId = Objects.FindId(Entry.Value);
-				KeepOuters(OuterId.AsInteger());
+				Visited[I] = true;
+				if (Remove[I])
+				{
+					++NumRemaining;
+					Remove[I] = false;
+				}
+				const TPair<FMinimalName, FPackageObjectIndex>& Entry = Objects.Get(FSetElementId::FromInteger(I)).Value;
+				if (!Entry.Value.IsNull())
+				{
+					FSetElementId OuterId = Objects.FindId(Entry.Value);
+					KeepOuters(OuterId.AsInteger());
+				}
 			}
-		}
-	};
+		};
 	// 1. Mark all runtime objects to be removed
 	for (const TPair<FPackageObjectIndex, UObject*>& Runtime : RuntimeScriptObjects)
 	{
@@ -6509,7 +6480,7 @@ void FCookedScriptObjectsDebug::TrimDebugData(const TMap<FPackageObjectIndex, UO
 	}
 
 	UE_LOG(LogStreaming, Log, TEXT("AsyncLoading2 - TrimDebugData: Kept %d/%d cooked script object entries (%.2f KB)"),
-		NumRemaining, Num, (float)Objects.GetAllocatedSize()/1024.f);
+		NumRemaining, Num, (float)Objects.GetAllocatedSize() / 1024.f);
 }
 
 void FCookedScriptObjectsDebug::GetPathNameInternal(const TPair<FMinimalName, FPackageObjectIndex>& Entry, FStringBuilderBase& PathName, int32& NumParts) const
@@ -6539,7 +6510,7 @@ void FCookedScriptObjectsDebug::GetPathNameInternal(const TPair<FMinimalName, FP
 
 FString FCookedScriptObjectsDebug::GetPathName(FPackageObjectIndex Index) const
 {
-	const TPair<FMinimalName,FPackageObjectIndex>* Entry = Objects.Find(Index);
+	const TPair<FMinimalName, FPackageObjectIndex>* Entry = Objects.Find(Index);
 	if (Entry)
 	{
 		TStringBuilder<256> PathName;
@@ -6688,13 +6659,13 @@ void FAsyncPackage2::ImportPackagesRecursiveInner(FAsyncLoadingThreadState2& Thr
 		// Add the imported package reference and confirm that its loader and extension match what was found from the package store or import store directly, 
 		// and if not, ensure the package was marked as missing or not installed
 		FLoadedPackageRef& ImportedPackageRef = ImportStore.AddImportedPackageReference(ImportedPackageId, ImportedPackageUPackageName, ImportedPackageLoader, ImportedPackagePath.GetHeaderExtension());
-		checkf((ImportedPackageStatus == EPackageStoreEntryStatus::Missing || ImportedPackageStatus == EPackageStoreEntryStatus::NotInstalled) 
-			|| ImportedPackageRef.GetPackageLoader() == ImportedPackageLoader, 
-			TEXT("ImportPackageStatus=%s, ImportedPackageRef.GetPackageLoader()=%s, ImportedPackageLoader=%s"), 
+		checkf((ImportedPackageStatus == EPackageStoreEntryStatus::Missing || ImportedPackageStatus == EPackageStoreEntryStatus::NotInstalled)
+			|| ImportedPackageRef.GetPackageLoader() == ImportedPackageLoader,
+			TEXT("ImportPackageStatus=%s, ImportedPackageRef.GetPackageLoader()=%s, ImportedPackageLoader=%s"),
 			LexToString(ImportedPackageStatus), LexToString(ImportedPackageRef.GetPackageLoader()), LexToString(ImportedPackageLoader));
-		checkf((ImportedPackageStatus == EPackageStoreEntryStatus::Missing || ImportedPackageStatus == EPackageStoreEntryStatus::NotInstalled) 
-			|| ImportedPackageRef.GetPackageHeaderExtension() == ImportedPackagePath.GetHeaderExtension(), 
-			TEXT("ImportPackageStatus=%s, ImportedPackageRef.GetPackageHeaderExtension()=%s, ImportedPackagePath.GetHeaderExtension()=%s"), 
+		checkf((ImportedPackageStatus == EPackageStoreEntryStatus::Missing || ImportedPackageStatus == EPackageStoreEntryStatus::NotInstalled)
+			|| ImportedPackageRef.GetPackageHeaderExtension() == ImportedPackagePath.GetHeaderExtension(),
+			TEXT("ImportPackageStatus=%s, ImportedPackageRef.GetPackageHeaderExtension()=%s, ImportedPackagePath.GetHeaderExtension()=%s"),
 			LexToString(ImportedPackageStatus), LexToString(ImportedPackageRef.GetPackageHeaderExtension()), LexToString(ImportedPackagePath.GetHeaderExtension()));
 
 		FAsyncPackage2* ImportedPackage = nullptr;
@@ -6944,18 +6915,18 @@ void FAsyncPackage2::DetachLinker()
 
 namespace UE::AsyncPackage2
 {
-	EAsyncLoadingResult::Type AsyncLoadingResultFromIoError(EIoErrorCode Error)
+EAsyncLoadingResult::Type AsyncLoadingResultFromIoError(EIoErrorCode Error)
+{
+	switch (Error)
 	{
-		switch (Error)
-		{
-		case EIoErrorCode::Ok:
-			return EAsyncLoadingResult::Succeeded;
-		case EIoErrorCode::NotInstalled:
-			return EAsyncLoadingResult::FailedNotInstalled;
-		}
-
-		return EAsyncLoadingResult::Failed;
+	case EIoErrorCode::Ok:
+	return EAsyncLoadingResult::Succeeded;
+	case EIoErrorCode::NotInstalled:
+	return EAsyncLoadingResult::FailedNotInstalled;
 	}
+
+	return EAsyncLoadingResult::Failed;
+}
 }
 
 void FAsyncPackage2::StartLoading(FAsyncLoadingThreadState2& ThreadState, FIoBatch& IoBatch)
@@ -7051,24 +7022,24 @@ void FAsyncPackage2::StartLoading(FAsyncLoadingThreadState2& ThreadState, FIoBat
 	{
 		TRACE_CPUPROFILER_EVENT_SCOPE(StartShaderMapRequests);
 		auto ReadShaderMapFunc = [this, &IoBatch](const FIoChunkId& ChunkId, FGraphEventRef GraphEvent)
-		{
-			GetPackageNode(Package_ExportsSerialized).AddBarrier();
-			int32 LocalPendingIoRequestsCounter = AsyncLoadingThread.PendingIoRequestsCounter.fetch_add(1) + 1;
-			TRACE_COUNTER_SET(AsyncLoadingPendingIoRequests, LocalPendingIoRequestsCounter);
-			return IoBatch.ReadWithCallback(ChunkId, FIoReadOptions(), Desc.Priority,
-				[this, GraphEvent](TIoStatusOr<FIoBuffer> Result)
-				{
-					GraphEvent->DispatchSubsequents();
-					int32 LocalPendingIoRequestsCounter = AsyncLoadingThread.PendingIoRequestsCounter.fetch_sub(1) - 1;
-					TRACE_COUNTER_SET(AsyncLoadingPendingIoRequests, LocalPendingIoRequestsCounter);
-					FAsyncLoadingThread2& LocalAsyncLoadingThread = AsyncLoadingThread;
-					GetPackageNode(Package_ExportsSerialized).ReleaseBarrier(nullptr);
-					if (LocalPendingIoRequestsCounter == 0)
+			{
+				GetPackageNode(Package_ExportsSerialized).AddBarrier();
+				int32 LocalPendingIoRequestsCounter = AsyncLoadingThread.PendingIoRequestsCounter.fetch_add(1) + 1;
+				TRACE_COUNTER_SET(AsyncLoadingPendingIoRequests, LocalPendingIoRequestsCounter);
+				return IoBatch.ReadWithCallback(ChunkId, FIoReadOptions(), Desc.Priority,
+					[this, GraphEvent](TIoStatusOr<FIoBuffer> Result)
 					{
-						LocalAsyncLoadingThread.AltZenaphore.NotifyOne();
-					}
-				});
-		};
+						GraphEvent->DispatchSubsequents();
+						int32 LocalPendingIoRequestsCounter = AsyncLoadingThread.PendingIoRequestsCounter.fetch_sub(1) - 1;
+						TRACE_COUNTER_SET(AsyncLoadingPendingIoRequests, LocalPendingIoRequestsCounter);
+						FAsyncLoadingThread2& LocalAsyncLoadingThread = AsyncLoadingThread;
+						GetPackageNode(Package_ExportsSerialized).ReleaseBarrier(nullptr);
+						if (LocalPendingIoRequestsCounter == 0)
+						{
+							LocalAsyncLoadingThread.AltZenaphore.NotifyOne();
+						}
+					});
+			};
 		FCoreDelegates::PreloadPackageShaderMaps.ExecuteIfBound(Data.ShaderMapHashes, ReadShaderMapFunc);
 	}
 }
@@ -7123,7 +7094,7 @@ EEventLoadNodeExecutionResult FAsyncPackage2::ProcessLinkerLoadPackageSummary(FA
 		}
 		else
 		{
-			ImportPackageName =  PackageImport.ObjectName;
+			ImportPackageName = PackageImport.ObjectName;
 		}
 		TCHAR NameStr[FName::StringBufferSize];
 		uint32 NameLen = ImportPackageName.ToString(NameStr);
@@ -7308,9 +7279,9 @@ bool FAsyncPackage2::PreloadLinkerLoadExports(FAsyncLoadingThreadState2& ThreadS
 
 bool FAsyncPackage2::ResolveLinkerLoadImports(FAsyncLoadingThreadState2& ThreadState)
 {
-	check(AsyncPackageLoadingState >= EAsyncPackageLoadingState2::WaitingForLinkerLoadDependencies);
+	check(AsyncPackageLoadingState >= EAsyncPackageLoadingState2::WaitingForImportDependencies);
 
-	check (!LinkerLoadState->bIsCurrentlyResolvingImports);
+	check(!LinkerLoadState->bIsCurrentlyResolvingImports);
 	TGuardValue<bool> GuardIsCurrentlyResolvingImports(LinkerLoadState->bIsCurrentlyResolvingImports, true);
 
 	// Validate that all imports are in the appropriate state and had their exports created
@@ -7322,7 +7293,7 @@ bool FAsyncPackage2::ResolveLinkerLoadImports(FAsyncLoadingThreadState2& ThreadS
 		{
 			if (ImportedPackage->LinkerLoadState.IsSet())
 			{
-				if (ImportedPackage->AsyncPackageLoadingState < EAsyncPackageLoadingState2::WaitingForLinkerLoadDependencies)
+				if (ImportedPackage->AsyncPackageLoadingState < EAsyncPackageLoadingState2::WaitingForImportDependencies)
 				{
 					check(ThreadState.PackagesOnStack.Contains(ImportedPackage));
 					UE_LOG(LogStreaming, Warning, TEXT("Package %s might be missing an import from package %s because of a circular dependency between them."),
@@ -7405,7 +7376,7 @@ bool FAsyncPackage2::CreateLinkerLoadExports(FAsyncLoadingThreadState2& ThreadSt
 {
 	check(AsyncPackageLoadingState >= EAsyncPackageLoadingState2::DependenciesReady);
 
-	check (!LinkerLoadState->bIsCurrentlyCreatingExports);
+	check(!LinkerLoadState->bIsCurrentlyCreatingExports);
 	TGuardValue<bool> GuardIsCurrentlyCreatingExports(LinkerLoadState->bIsCurrentlyCreatingExports, true);
 
 	UE::TUniqueLock LinkerLock(LinkerLoadState->Linker->Mutex);
@@ -7471,6 +7442,10 @@ bool FAsyncPackage2::CreateLinkerLoadExports(FAsyncLoadingThreadState2& ThreadSt
 
 EEventLoadNodeExecutionResult FAsyncPackage2::ExecutePostLoadLinkerLoadPackageExports(FAsyncLoadingThreadState2& ThreadState)
 {
+#if WITH_EDITOR
+	UE::Core::Private::FPlayInEditorLoadingScope PlayInEditorIDScope(Desc.PIEInstanceID);
+#endif
+
 	if (!HasLoadFailed())
 	{
 		SCOPED_LOADTIMER(PostLoadObjectsTime);
@@ -7478,10 +7453,6 @@ EEventLoadNodeExecutionResult FAsyncPackage2::ExecutePostLoadLinkerLoadPackageEx
 
 		FAsyncLoadingTickScope2 InAsyncLoadingTick(AsyncLoadingThread);
 
-#if WITH_EDITOR
-		UE::Core::Private::FPlayInEditorLoadingScope PlayInEditorIDScope(Desc.PIEInstanceID);
-#endif
-
 		// Begin async loading, simulates BeginLoad
 		BeginAsyncLoad();
 
@@ -7489,7 +7460,7 @@ EEventLoadNodeExecutionResult FAsyncPackage2::ExecutePostLoadLinkerLoadPackageEx
 		TArray<UObject*>& ThreadObjLoaded = LoadContext->PRIVATE_GetObjectsLoadedInternalUseOnly();
 
 		// End async loading, simulates EndLoad
-		ON_SCOPE_EXIT { ThreadObjLoaded.Reset(); EndAsyncLoad(); };
+		ON_SCOPE_EXIT{ ThreadObjLoaded.Reset(); EndAsyncLoad(); };
 
 		FUObjectThreadContext& ThreadContext = FUObjectThreadContext::Get();
 		TGuardValue<bool> GuardIsRoutingPostLoad(ThreadContext.IsRoutingPostLoad, true);
@@ -7502,7 +7473,7 @@ EEventLoadNodeExecutionResult FAsyncPackage2::ExecutePostLoadLinkerLoadPackageEx
 		{
 			const int32 ExportIndex = LinkerLoadState->PostLoadExportIndex++;
 			const FExportObject& Export = Data.Exports[ExportIndex];
-		
+
 			if (UObject* Object = Export.Object)
 			{
 				if (Object->HasAnyFlags(RF_NeedPostLoad) && CanPostLoadOnAsyncLoadingThread(Object) && Object->IsReadyForAsyncPostLoad())
@@ -7602,48 +7573,49 @@ EEventLoadNodeExecutionResult FAsyncPackage2::ExecuteDeferredPostLoadLinkerLoadP
 	return EEventLoadNodeExecutionResult::Complete;
 }
 
-EEventLoadNodeExecutionResult FAsyncPackage2::Event_CreateLinkerLoadExports(FAsyncLoadingThreadState2& ThreadState, FAsyncPackage2* Package, int32)
+EEventLoadNodeExecutionResult FAsyncPackage2::Event_CreateExports(FAsyncLoadingThreadState2& ThreadState, FAsyncPackage2* Package, int32)
 {
-	TRACE_CPUPROFILER_EVENT_SCOPE(Event_CreateLinkerLoadExports);
+	TRACE_CPUPROFILER_EVENT_SCOPE(Event_CreateExports);
 	UE_ASYNC_PACKAGE_DEBUG(Package->Desc);
 
 	check(FUObjectThreadContext::Get().AsyncVisibilityFilter == EInternalObjectFlags::None);
-	check(Package->AsyncPackageLoadingState == EAsyncPackageLoadingState2::CreateLinkerLoadExports);
+	check(Package->AsyncPackageLoadingState == EAsyncPackageLoadingState2::CreateExports);
 
 	FAsyncPackageScope2 Scope(Package);
 #if WITH_EDITOR
 	UE::Core::Private::FPlayInEditorLoadingScope PlayInEditorIDScope(Package->Desc.PIEInstanceID);
 #endif
 
-	check(Package->LinkerLoadState.IsSet());
-
-	if (!Package->CreateLinkerLoadExports(ThreadState))
+	if (Package->LinkerLoadState.IsSet())
 	{
-		return EEventLoadNodeExecutionResult::Timeout;
+		if (!Package->CreateLinkerLoadExports(ThreadState))
+		{
+			return EEventLoadNodeExecutionResult::Timeout;
+		}
+	}
+	else
+	{
+		if (!Package->CreateRuntimeExports(ThreadState))
+		{
+			return EEventLoadNodeExecutionResult::Timeout;
+		}
 	}
 
-	Package->AsyncPackageLoadingState = EAsyncPackageLoadingState2::WaitingForLinkerLoadDependencies;
-	Package->ConditionalBeginResolveLinkerLoadImports(ThreadState);
+	Package->AsyncPackageLoadingState = EAsyncPackageLoadingState2::WaitingForImportDependencies;
+	Package->ConditionalBeginResolveImports(ThreadState);
 	return EEventLoadNodeExecutionResult::Complete;
 }
 
-void FAsyncPackage2::ConditionalBeginResolveLinkerLoadImports(FAsyncLoadingThreadState2& ThreadState)
+void FAsyncPackage2::ConditionalBeginResolveImports(FAsyncLoadingThreadState2& ThreadState)
 {
-	TRACE_CPUPROFILER_EVENT_SCOPE(ConditionalBeginResolveLinkerLoadImports);
+	TRACE_CPUPROFILER_EVENT_SCOPE(ConditionalBeginResolveImports);
 
-	WaitForAllDependenciesToReachState(ThreadState, &FAsyncPackage2::AllDependenciesImportState, EAsyncPackageLoadingState2::WaitingForLinkerLoadDependencies, AsyncLoadingThread.ConditionalBeginResolveImportsTick,
+	WaitForAllDependenciesToReachState(ThreadState, &FAsyncPackage2::AllDependenciesImportState, EAsyncPackageLoadingState2::WaitingForImportDependencies, AsyncLoadingThread.ConditionalBeginResolveImportsTick,
 		[&ThreadState](FAsyncPackage2* Package)
 		{
-			if (Package->LinkerLoadState.IsSet())
-			{
-				check(Package->AsyncPackageLoadingState == EAsyncPackageLoadingState2::WaitingForLinkerLoadDependencies);
-				Package->AsyncPackageLoadingState = EAsyncPackageLoadingState2::ResolveLinkerLoadImports;
-				Package->GetPackageNode(Package_ResolveLinkerLoadImports).ReleaseBarrier(&ThreadState);
-			}
-			else
-			{
-				// Don't advance state of cooked package, nodes already have dependencies setup.
-			}
+			check(Package->AsyncPackageLoadingState == EAsyncPackageLoadingState2::WaitingForImportDependencies);
+			Package->AsyncPackageLoadingState = EAsyncPackageLoadingState2::ResolveImports;
+			Package->GetPackageNode(Package_ResolveImports).ReleaseBarrier(&ThreadState);
 		}
 	);
 }
@@ -7659,21 +7631,19 @@ void FAsyncPackage2::ConditionalBeginPostLoadPhase(FAsyncLoadingThreadState2& Th
 		{
 			check(Package->AsyncPackageLoadingState == EAsyncPackageLoadingState2::PostLoad);
 
-			for (int32 ExportBundleIndex = 0; ExportBundleIndex < Package->Data.TotalExportBundleCount; ++ExportBundleIndex)
-			{
-				Package->GetExportBundleNode(EEventLoadNode2::ExportBundle_PostLoad, ExportBundleIndex).ReleaseBarrier(&ThreadState);
-			}
+			Package->AsyncPackageLoadingState = EAsyncPackageLoadingState2::PostLoad;
+			Package->GetPackageNode(Package_PostLoad).ReleaseBarrier(&ThreadState);
 		}
 	);
 }
 
-EEventLoadNodeExecutionResult FAsyncPackage2::Event_ResolveLinkerLoadImports(FAsyncLoadingThreadState2& ThreadState, FAsyncPackage2* Package, int32)
+EEventLoadNodeExecutionResult FAsyncPackage2::Event_ResolveImports(FAsyncLoadingThreadState2& ThreadState, FAsyncPackage2* Package, int32)
 {
-	TRACE_CPUPROFILER_EVENT_SCOPE(Event_ResolveLinkerLoadImports);
+	TRACE_CPUPROFILER_EVENT_SCOPE(Event_ResolveImports);
 	UE_ASYNC_PACKAGE_DEBUG(Package->Desc);
 
 	check(FUObjectThreadContext::Get().AsyncVisibilityFilter == EInternalObjectFlags::None);
-	check(Package->AsyncPackageLoadingState == EAsyncPackageLoadingState2::ResolveLinkerLoadImports);
+	check(Package->AsyncPackageLoadingState == EAsyncPackageLoadingState2::ResolveImports);
 
 	FAsyncPackageScope2 Scope(Package);
 #if WITH_EDITOR
@@ -7688,18 +7658,18 @@ EEventLoadNodeExecutionResult FAsyncPackage2::Event_ResolveLinkerLoadImports(FAs
 		}
 	}
 
-	Package->AsyncPackageLoadingState = EAsyncPackageLoadingState2::PreloadLinkerLoadExports;
-	Package->GetPackageNode(Package_PreloadLinkerLoadExports).ReleaseBarrier(&ThreadState);
+	Package->AsyncPackageLoadingState = EAsyncPackageLoadingState2::PreloadExports;
+	Package->GetPackageNode(Package_PreloadExports).ReleaseBarrier(&ThreadState);
 	return EEventLoadNodeExecutionResult::Complete;
 }
 
-EEventLoadNodeExecutionResult FAsyncPackage2::Event_PreloadLinkerLoadExports(FAsyncLoadingThreadState2& ThreadState, FAsyncPackage2* Package, int32)
+EEventLoadNodeExecutionResult FAsyncPackage2::Event_PreloadExports(FAsyncLoadingThreadState2& ThreadState, FAsyncPackage2* Package, int32)
 {
-	TRACE_CPUPROFILER_EVENT_SCOPE(Event_PreloadLinkerLoadExports);
+	TRACE_CPUPROFILER_EVENT_SCOPE(Event_PreloadExports);
 	UE_ASYNC_PACKAGE_DEBUG(Package->Desc);
 
 	check(FUObjectThreadContext::Get().AsyncVisibilityFilter == EInternalObjectFlags::None);
-	check(Package->AsyncPackageLoadingState == EAsyncPackageLoadingState2::PreloadLinkerLoadExports);
+	check(Package->AsyncPackageLoadingState == EAsyncPackageLoadingState2::PreloadExports);
 
 	FAsyncPackageScope2 Scope(Package);
 #if WITH_EDITOR
@@ -7718,6 +7688,13 @@ EEventLoadNodeExecutionResult FAsyncPackage2::Event_PreloadLinkerLoadExports(FAs
 			Package->ExternalReadDependencies.Append(MoveTemp(Package->LinkerLoadState->Linker->ExternalReadDependencies));
 		}
 	}
+	else
+	{
+		if (!Package->PreloadRuntimeExports(ThreadState))
+		{
+			return EEventLoadNodeExecutionResult::Timeout;
+		}
+	}
 
 	if (Package->ExternalReadDependencies.Num() == 0)
 	{
@@ -7739,7 +7716,7 @@ EEventLoadNodeExecutionResult FAsyncPackage2::Event_ProcessPackageSummary(FAsync
 {
 	TRACE_CPUPROFILER_EVENT_SCOPE(Event_ProcessPackageSummary);
 	UE_ASYNC_PACKAGE_DEBUG(Package->Desc);
-	
+
 	check(FUObjectThreadContext::Get().AsyncVisibilityFilter == EInternalObjectFlags::None);
 	check(Package->AsyncPackageLoadingState == EAsyncPackageLoadingState2::WaitingForIo);
 	Package->AsyncPackageLoadingState = EAsyncPackageLoadingState2::ProcessPackageSummary;
@@ -7782,12 +7759,12 @@ EEventLoadNodeExecutionResult FAsyncPackage2::Event_ProcessPackageSummary(FAsync
 			FLoadedPackageRef& PackageRef = Package->AsyncLoadingThread.GlobalImportStore.FindPackageRefChecked(Package->Desc.UPackageId, Package->Desc.UPackageName);
 #if WITH_EDITOR
 			if (Package->OptionalSegmentHeaderData.IsSet())
-			{ 
-				PackageRef.PreInsertPublicExports({ Package->HeaderData.ExportMap, Package->OptionalSegmentHeaderData->ExportMap } );
+			{
+				PackageRef.PreInsertPublicExports({ Package->HeaderData.ExportMap, Package->OptionalSegmentHeaderData->ExportMap });
 			}
 			else
 #endif
-			{ 
+			{
 				PackageRef.PreInsertPublicExports({ Package->HeaderData.ExportMap });
 			}
 		}
@@ -7856,7 +7833,7 @@ EEventLoadNodeExecutionResult FAsyncPackage2::Event_DependenciesReady(FAsyncLoad
 {
 	TRACE_CPUPROFILER_EVENT_SCOPE(Event_DependenciesReady);
 	UE_ASYNC_PACKAGE_DEBUG(Package->Desc);
-	
+
 	check(FUObjectThreadContext::Get().AsyncVisibilityFilter == EInternalObjectFlags::None);
 	check(Package->AsyncPackageLoadingState == EAsyncPackageLoadingState2::WaitingForDependencies);
 
@@ -7932,140 +7909,229 @@ void FAsyncPackage2::InitializeExportArchive(FExportArchive& Ar, bool bIsOptiona
 	Ar.bIsOptionalSegment = bIsOptionalSegment;
 }
 
-EEventLoadNodeExecutionResult FAsyncPackage2::Event_ProcessExportBundle(FAsyncLoadingThreadState2& ThreadState, FAsyncPackage2* Package, int32 InExportBundleIndex)
+// When possible, deferring serialization will enable some nodes to run in parallel
+// As some point, we will also need to defer them if we want to better mix cooked and uncooked loading.
+inline bool CanDeferSerialization(FExportObject& Export)
 {
-	TRACE_CPUPROFILER_EVENT_SCOPE(Event_ProcessExportBundle);
-	UE_ASYNC_PACKAGE_DEBUG(Package->Desc);
-	
+	return GAllowMultithreadedLoading &&
+		// Classes, CDOs and Archetype are serialized during the create export phase since they are required to be loaded by dependencies
+		!Export.Object->GetClass()->IsChildOf(UField::StaticClass()) &&
+		!Export.Object->HasAnyFlags(RF_ClassDefaultObject | RF_ArchetypeObject);
+}
+
+bool FAsyncPackage2::CreateRuntimeExports(FAsyncLoadingThreadState2& ThreadState)
+{
+	TRACE_CPUPROFILER_EVENT_SCOPE(CreateRuntimeExports);
+	UE_ASYNC_PACKAGE_DEBUG(Desc);
+
+	check(!LinkerLoadState.IsSet());
 	check(FUObjectThreadContext::Get().AsyncVisibilityFilter == EInternalObjectFlags::None);
-	check(Package->AsyncPackageLoadingState >= EAsyncPackageLoadingState2::DependenciesReady);
-	Package->AsyncPackageLoadingState = EAsyncPackageLoadingState2::ProcessExportBundles;
+	check(AsyncPackageLoadingState >= EAsyncPackageLoadingState2::DependenciesReady);
 
-	UE_ASYNC_PACKAGE_LOG(VeryVerbose, Package->Desc, TEXT("ProcessExportBundle"), TEXT("Beginning Processing Export Bundle %d"), InExportBundleIndex);
+	if (HasLoadFailed())
+	{
+		return true;
+	}
 
-	FAsyncPackageScope2 Scope(Package);
+	FAsyncPackageScope2 Scope(this);
 #if WITH_EDITOR
-	UE::Core::Private::FPlayInEditorLoadingScope PlayInEditorIDScope(Package->Desc.PIEInstanceID);
+	UE::Core::Private::FPlayInEditorLoadingScope PlayInEditorIDScope(Desc.PIEInstanceID);
 #endif
 
-	// This code path should not be reached for LinkerLoad packages
-	check(!Package->LinkerLoadState.IsSet());
-	check(InExportBundleIndex < Package->Data.TotalExportBundleCount);
-	
-	if (!Package->HasLoadFailed())
+	UE::TUniqueLock ScopeLock(SerializationMutex);
+	while (ProcessedExportBundlesCount < Data.TotalExportBundleCount)
 	{
+		UE_ASYNC_PACKAGE_LOG(VeryVerbose, Desc, TEXT("CreateRuntimeExports"), TEXT("Beginning Creating Exports for runtime bundle %d"), ProcessedExportBundlesCount);
+
 		bool bIsOptionalSegment = false;
 #if WITH_EDITOR
-		const FAsyncPackageHeaderData* HeaderData;
-		FAsyncPackageSerializationState* SerializationState;
-		bIsOptionalSegment = InExportBundleIndex == 1;
+		const FAsyncPackageHeaderData* LocalHeaderData;
+		FAsyncPackageSerializationState* LocalSerializationState;
+		bIsOptionalSegment = ProcessedExportBundlesCount == 1;
 
 		if (bIsOptionalSegment)
 		{
-			HeaderData = Package->OptionalSegmentHeaderData.GetPtrOrNull();
-			check(HeaderData);
-			SerializationState = Package->OptionalSegmentSerializationState.GetPtrOrNull();
-			check(SerializationState);
+			LocalHeaderData = OptionalSegmentHeaderData.GetPtrOrNull();
+			check(LocalHeaderData);
+			LocalSerializationState = OptionalSegmentSerializationState.GetPtrOrNull();
+			check(LocalSerializationState);
 		}
 		else
 		{
-			check(InExportBundleIndex == 0);
-			HeaderData = &Package->HeaderData;
-			SerializationState = &Package->SerializationState;
+			check(ProcessedExportBundlesCount == 0);
+			LocalHeaderData = &HeaderData;
+			LocalSerializationState = &SerializationState;
 		}
 #else
-		const FAsyncPackageHeaderData* HeaderData = &Package->HeaderData;
-		FAsyncPackageSerializationState* SerializationState = &Package->SerializationState;
+		const FAsyncPackageHeaderData* LocalHeaderData = &HeaderData;
+		FAsyncPackageSerializationState* LocalSerializationState = &SerializationState;
 #endif
-		const FIoBuffer& IoBuffer = SerializationState->IoRequest.GetResultOrDie();
+		const FIoBuffer& IoBuffer = LocalSerializationState->IoRequest.GetResultOrDie();
 		FExportArchive Ar(IoBuffer);
-		Package->InitializeExportArchive(Ar, bIsOptionalSegment);
+		InitializeExportArchive(Ar, bIsOptionalSegment);
 
-		while (Package->ExportBundleEntryIndex < HeaderData->ExportBundleEntries.Num())
+		while (ExportBundleEntryIndex < LocalHeaderData->ExportBundleEntries.Num())
 		{
-			const FExportBundleEntry& BundleEntry = HeaderData->ExportBundleEntries[Package->ExportBundleEntryIndex];
-			if (ThreadState.IsTimeLimitExceeded(TEXT("Event_ProcessExportBundle")))
+			const FExportBundleEntry& BundleEntry = LocalHeaderData->ExportBundleEntries[ExportBundleEntryIndex];
+			if (BundleEntry.LocalExportIndex < uint32(LocalHeaderData->ExportMap.Num()))
 			{
-				return EEventLoadNodeExecutionResult::Timeout;
-			}
-			if (BundleEntry.LocalExportIndex < uint32(HeaderData->ExportMap.Num()))
-			{
-				const FExportMapEntry& ExportMapEntry = HeaderData->ExportMap[BundleEntry.LocalExportIndex];
-				FExportObject& Export = HeaderData->ExportsView[BundleEntry.LocalExportIndex];
+				const FExportMapEntry& ExportMapEntry = LocalHeaderData->ExportMap[BundleEntry.LocalExportIndex];
+				FExportObject& Export = LocalHeaderData->ExportsView[BundleEntry.LocalExportIndex];
 
 				if (BundleEntry.CommandType == FExportBundleEntry::ExportCommandType_Create)
 				{
 					if (!Export.Object)
 					{
-						Package->EventDrivenCreateExport(*HeaderData, BundleEntry.LocalExportIndex);
+						EventDrivenCreateExport(*LocalHeaderData, BundleEntry.LocalExportIndex);
 					}
 				}
-				else
+				else if (Export.Object && Export.Object->HasAnyFlags(RF_NeedLoad) && !CanDeferSerialization(Export))
 				{
 					check(BundleEntry.CommandType == FExportBundleEntry::ExportCommandType_Serialize);
-					if (Export.Object && Export.Object->HasAllFlags(RF_NeedLoad))
-					{
-						Package->EventDrivenSerializeExport(*HeaderData, BundleEntry.LocalExportIndex, &Ar);
-					}
+
+					EventDrivenSerializeExport(*LocalHeaderData, BundleEntry.LocalExportIndex, &Ar);
 				}
 			}
 			else
 			{
-				uint32 CellExportIndex = BundleEntry.LocalExportIndex - HeaderData->ExportMap.Num();
-				const FCellExportMapEntry& CellExportMapEntry = HeaderData->CellExportMap[CellExportIndex];
-				FExportCell& CellExport = HeaderData->CellExportsView[CellExportIndex];
+				uint32 CellExportIndex = BundleEntry.LocalExportIndex - LocalHeaderData->ExportMap.Num();
+				const FCellExportMapEntry& CellExportMapEntry = LocalHeaderData->CellExportMap[CellExportIndex];
+				FExportCell& CellExport = LocalHeaderData->CellExportsView[CellExportIndex];
 
 				if (BundleEntry.CommandType == FExportBundleEntry::ExportCommandType_Create)
 				{
 					check(!CellExport.Cell);
-					Package->EventDrivenCreateCellExport(*HeaderData, BundleEntry.LocalExportIndex, &Ar);
+					EventDrivenCreateCellExport(*LocalHeaderData, BundleEntry.LocalExportIndex, &Ar);
 				}
 				else
 				{
 					check(BundleEntry.CommandType == FExportBundleEntry::ExportCommandType_Serialize);
 					check(CellExport.Cell);
-					Package->EventDrivenSerializeCellExport(*HeaderData, BundleEntry.LocalExportIndex, &Ar);
+					EventDrivenSerializeCellExport(*LocalHeaderData, BundleEntry.LocalExportIndex, &Ar);
 				}
 			}
-			++Package->ExportBundleEntryIndex;
+			++ExportBundleEntryIndex;
+		}
+
+		ExportBundleEntryIndex = 0;
+
+		UE_ASYNC_PACKAGE_LOG(VeryVerbose, Desc, TEXT("CreateRuntimeExports"), TEXT("Finished Creating Exports for runtime bundle %d"), ProcessedExportBundlesCount);
+
+		if (++ProcessedExportBundlesCount == Data.TotalExportBundleCount)
+		{
+			ProcessedExportBundlesCount = 0;
+			return true;
 		}
 	}
-	
-	Package->ExportBundleEntryIndex = 0;
 
-	if (++Package->ProcessedExportBundlesCount == Package->Data.TotalExportBundleCount)
+	return false;
+}
+
+bool FAsyncPackage2::PreloadRuntimeExports(FAsyncLoadingThreadState2& ThreadState)
+{
+	TRACE_CPUPROFILER_EVENT_SCOPE(PreloadRuntimeExports);
+	UE_ASYNC_PACKAGE_DEBUG(Desc);
+
+	check(!LinkerLoadState.IsSet());
+	check(FUObjectThreadContext::Get().AsyncVisibilityFilter == EInternalObjectFlags::None);
+	check(AsyncPackageLoadingState >= EAsyncPackageLoadingState2::ResolveImports);
+
+	if (HasLoadFailed())
 	{
-		Package->ProcessedExportBundlesCount = 0;
-		Package->HeaderData.Reset();
-		Package->SerializationState.ReleaseIoRequest();
+		return true;
+	}
+
+	FAsyncPackageScope2 Scope(this);
+#if WITH_EDITOR
+	UE::Core::Private::FPlayInEditorLoadingScope PlayInEditorIDScope(Desc.PIEInstanceID);
+#endif
+
+	UE::TUniqueLock ScopeLock(SerializationMutex);
+	while (ProcessedExportBundlesCount < Data.TotalExportBundleCount)
+	{
+		UE_ASYNC_PACKAGE_LOG(VeryVerbose, Desc, TEXT("PreloadRuntimeExports"), TEXT("Beginning Preloads of runtime export bundle %d"), ProcessedExportBundlesCount);
+
+		bool bIsOptionalSegment = false;
 #if WITH_EDITOR
-		if (Package->OptionalSegmentHeaderData.IsSet())
+		const FAsyncPackageHeaderData* LocalHeaderData;
+		FAsyncPackageSerializationState* LocalSerializationState;
+		bIsOptionalSegment = ProcessedExportBundlesCount == 1;
+
+		if (bIsOptionalSegment)
+		{
+			LocalHeaderData = OptionalSegmentHeaderData.GetPtrOrNull();
+			check(LocalHeaderData);
+			LocalSerializationState = OptionalSegmentSerializationState.GetPtrOrNull();
+			check(LocalSerializationState);
+		}
+		else
 		{
-			Package->OptionalSegmentHeaderData->Reset();
-			Package->OptionalSegmentSerializationState->ReleaseIoRequest();
+			check(ProcessedExportBundlesCount == 0);
+			LocalHeaderData = &HeaderData;
+			LocalSerializationState = &SerializationState;
 		}
+#else
+		const FAsyncPackageHeaderData* LocalHeaderData = &HeaderData;
+		FAsyncPackageSerializationState* LocalSerializationState = &SerializationState;
 #endif
-		check(Package->AsyncPackageLoadingState == EAsyncPackageLoadingState2::ProcessExportBundles);
+		const FIoBuffer& IoBuffer = LocalSerializationState->IoRequest.GetResultOrDie();
+		FExportArchive Ar(IoBuffer);
+		InitializeExportArchive(Ar, bIsOptionalSegment);
 
-		if (Package->ExternalReadDependencies.Num() == 0)
+		while (ExportBundleEntryIndex < LocalHeaderData->ExportBundleEntries.Num())
 		{
-			Package->GetPackageNode(Package_ExportsSerialized).ReleaseBarrier(&ThreadState);
+			const FExportBundleEntry& BundleEntry = LocalHeaderData->ExportBundleEntries[ExportBundleEntryIndex];
+			if (BundleEntry.LocalExportIndex < uint32(LocalHeaderData->ExportMap.Num()))
+			{
+				const FExportMapEntry& ExportMapEntry = LocalHeaderData->ExportMap[BundleEntry.LocalExportIndex];
+				FExportObject& Export = LocalHeaderData->ExportsView[BundleEntry.LocalExportIndex];
+
+				if (BundleEntry.CommandType == FExportBundleEntry::ExportCommandType_Serialize)
+				{
+					if (Export.Object && Export.Object->HasAnyFlags(RF_NeedLoad))
+					{
+						EventDrivenSerializeExport(*LocalHeaderData, BundleEntry.LocalExportIndex, &Ar);
+					}
+				}
+			}
+			else
+			{
+				uint32 CellExportIndex = BundleEntry.LocalExportIndex - LocalHeaderData->ExportMap.Num();
+				const FCellExportMapEntry& CellExportMapEntry = LocalHeaderData->CellExportMap[CellExportIndex];
+				FExportCell& CellExport = LocalHeaderData->CellExportsView[CellExportIndex];
+
+				if (BundleEntry.CommandType == FExportBundleEntry::ExportCommandType_Serialize)
+				{
+					if (!CellExport.bSerialized)
+					{
+						EventDrivenSerializeCellExport(*LocalHeaderData, BundleEntry.LocalExportIndex, &Ar);
+					}
+				}
+			}
+			++ExportBundleEntryIndex;
 		}
-		else
+
+		ExportBundleEntryIndex = 0;
+
+		UE_ASYNC_PACKAGE_LOG(VeryVerbose, Desc, TEXT("PreloadRuntimeExports"), TEXT("Finished Preloads of runtime export bundle %d"), ProcessedExportBundlesCount);
+
+		if (++ProcessedExportBundlesCount == Data.TotalExportBundleCount)
 		{
-			Package->AsyncPackageLoadingState = EAsyncPackageLoadingState2::WaitingForExternalReads;
-			Package->AsyncLoadingThread.ExternalReadQueue.Enqueue(Package);
+			ProcessedExportBundlesCount = 0;
+			HeaderData.Reset();
+			SerializationState.ReleaseIoRequest();
+#if WITH_EDITOR
+			if (OptionalSegmentHeaderData.IsSet())
+			{
+				OptionalSegmentHeaderData->Reset();
+				OptionalSegmentSerializationState->ReleaseIoRequest();
+			}
+#endif
+			return true;
 		}
 	}
-	else
-	{
-		// Release the next bundle now that we've finished.
-		Package->GetExportBundleNode(EEventLoadNode2::ExportBundle_Process, Package->ProcessedExportBundlesCount).ReleaseBarrier(&ThreadState);
-	}
 
-	UE_ASYNC_PACKAGE_LOG(VeryVerbose, Package->Desc, TEXT("ProcessExportBundle"), TEXT("Finished Processing Export Bundle %d"), InExportBundleIndex);
-
-	return EEventLoadNodeExecutionResult::Complete;
+	return false;
 }
 
 UObject* FAsyncPackage2::EventDrivenIndexToObject(const FAsyncPackageHeaderData& Header, FPackageObjectIndex Index, bool bCheckSerialized)
@@ -8082,7 +8148,7 @@ UObject* FAsyncPackage2::EventDrivenIndexToObject(const FAsyncPackageHeaderData&
 		UE_CLOG(!Result, LogStreaming, Warning, TEXT("Missing Dependency, missing export (%s) in package %s. \n\tExport Details - Index: 0x%llX, Super: '%s', Template: '%s', LoadFailed: %s, Filtered: %s (%s), FoundInMemory: %s"),
 			*Header.NameMap.GetName(Header.ExportMap[Index.ToExport()].ObjectName).ToString(),
 			*Desc.PackagePathToLoad.GetPackageFName().ToString(),
-			Index.Value(), 
+			Index.Value(),
 			(Export.SuperObject ? *Export.SuperObject->GetPathName() : TEXT("null")),
 			(Export.TemplateObject ? *Export.TemplateObject->GetPathName() : TEXT("null")),
 			(Export.bExportLoadFailed ? TEXT("true") : TEXT("false")),
@@ -8138,7 +8204,7 @@ void FAsyncPackage2::ProcessExportDependencies(const FAsyncPackageHeaderData& He
 			RunningIndex += DependencyBundle.EntryCount[FExportBundleEntry::ExportCommandType_Create][Index];
 		}
 	}
-	
+
 	for (int32 Index = DependencyBundle.EntryCount[CommandType][FExportBundleEntry::ExportCommandType_Create]; Index > 0; --Index)
 	{
 		const FDependencyBundleEntry& Dep = Header.DependencyBundleEntries[RunningIndex++];
@@ -8233,12 +8299,12 @@ UObject* FAsyncPackage2::ConditionalCreateExport(const FAsyncPackageHeaderData&
 UObject* FAsyncPackage2::ConditionalSerializeExport(const FAsyncPackageHeaderData& Header, int32 LocalExportIndex)
 {
 	FExportObject& Export = Header.ExportsView[LocalExportIndex];
-	
+
 	if (!Export.Object && !(Export.bFiltered || Export.bExportLoadFailed))
 	{
 		ConditionalCreateExport(Header, LocalExportIndex);
 	}
-	
+
 	if (!Export.Object || (Export.bFiltered || Export.bExportLoadFailed))
 	{
 		return nullptr;
@@ -8270,6 +8336,7 @@ UObject* FAsyncPackage2::ConditionalCreateImport(const FAsyncPackageHeaderData&
 	}
 	uint64 ExportHash = Header.ImportedPublicExportHashes[PackageImportRef.GetImportedPublicExportHashIndex()];
 	FAsyncPackageHeaderData* SourcePackageHeader = nullptr;
+	UE::TUniqueLock ScopeLock(SourcePackage->SerializationMutex);
 	int32 ExportIndex = SourcePackage->GetPublicExportIndex(ExportHash, SourcePackageHeader);
 	if (ExportIndex < 0)
 	{
@@ -8299,6 +8366,7 @@ UObject* FAsyncPackage2::ConditionalSerializeImport(const FAsyncPackageHeaderDat
 	}
 	uint64 ExportHash = Header.ImportedPublicExportHashes[PackageImportRef.GetImportedPublicExportHashIndex()];
 	FAsyncPackageHeaderData* SourcePackageHeader = nullptr;
+	UE::TUniqueLock ScopeLock(SourcePackage->SerializationMutex);
 	int32 ExportIndex = SourcePackage->GetPublicExportIndex(ExportHash, SourcePackageHeader);
 	if (ExportIndex < 0)
 	{
@@ -8333,7 +8401,7 @@ Verse::VCell* FAsyncPackage2::ConditionalCreateCellExport(const FAsyncPackageHea
 Verse::VCell* FAsyncPackage2::ConditionalSerializeCellExport(const FAsyncPackageHeaderData& Header, int32 CellExportIndex)
 {
 	ConditionalCreateCellExport(Header, CellExportIndex);
-	
+
 	if (!Header.CellExportsView[CellExportIndex].bSerialized)
 	{
 		FAsyncPackageScope2 Scope(this);
@@ -8359,6 +8427,7 @@ Verse::VCell* FAsyncPackage2::ConditionalCreateCellImport(const FAsyncPackageHea
 	}
 	uint64 CellExportHash = Header.ImportedPublicExportHashes[PackageImportRef.GetImportedPublicExportHashIndex()];
 	FAsyncPackageHeaderData* SourcePackageHeader = nullptr;
+	UE::TUniqueLock ScopeLock(SourcePackage->SerializationMutex);
 	int32 CellExportIndex = SourcePackage->GetPublicCellExportIndex(CellExportHash, SourcePackageHeader);
 	if (CellExportIndex < 0)
 	{
@@ -8380,6 +8449,7 @@ Verse::VCell* FAsyncPackage2::ConditionalSerializeCellImport(const FAsyncPackage
 	}
 	uint64 CellExportHash = Header.ImportedPublicExportHashes[PackageImportRef.GetImportedPublicExportHashIndex()];
 	FAsyncPackageHeaderData* SourcePackageHeader = nullptr;
+	UE::TUniqueLock ScopeLock(SourcePackage->SerializationMutex);
 	int32 CellExportIndex = SourcePackage->GetPublicCellExportIndex(CellExportHash, SourcePackageHeader);
 	if (CellExportIndex < 0)
 	{
@@ -8542,7 +8612,7 @@ public:
 					return true;
 				}
 			}
-		
+
 			return false;
 		}
 	}
@@ -8552,7 +8622,7 @@ public:
 
 bool FAsyncPackage2::CanRunNodeAsync(const FAsyncLoadEventSpec* Spec) const
 {
-	if (Spec->Func == &FAsyncPackage2::Event_PreloadLinkerLoadExports)
+	if (Spec->Func == &FAsyncPackage2::Event_PreloadExports)
 	{
 		TRACE_CPUPROFILER_EVENT_SCOPE(FAsyncPackage2::CanRunNodeAsync);
 
@@ -8646,7 +8716,7 @@ void FAsyncPackage2::EventDrivenCreateExport(const FAsyncPackageHeaderData& Head
 		return;
 	}
 
-	if ((Export.ObjectFlags & RF_ClassDefaultObject) == 0 
+	if ((Export.ObjectFlags & RF_ClassDefaultObject) == 0
 		&& !ExportObject.TemplateObject->IsA(LoadClass))
 	{
 		UE_ASYNC_PACKAGE_LOG(Error, Desc, TEXT("CreateExport"), TEXT("Export class type (%s) differs from the template object type (%s)"),
@@ -8742,8 +8812,8 @@ void FAsyncPackage2::EventDrivenCreateExport(const FAsyncPackageHeaderData& Head
 			}
 			else
 			{
-				checkf(ExportObject.TemplateObject->IsA(LoadClass), 
-					TEXT("ExportObject.TemplateObject: '%s' is not a child of class: '%s'"), 
+				checkf(ExportObject.TemplateObject->IsA(LoadClass),
+					TEXT("ExportObject.TemplateObject: '%s' is not a child of class: '%s'"),
 					*GetNameSafe(ExportObject.TemplateObject),
 					*GetNameSafe(LoadClass));
 			}
@@ -8751,7 +8821,7 @@ void FAsyncPackage2::EventDrivenCreateExport(const FAsyncPackageHeaderData& Head
 #endif
 		checkf(!LoadClass->HasAnyFlags(RF_NeedLoad),
 			TEXT("LoadClass %s had RF_NeedLoad while creating %s"), *LoadClass->GetFullName(), *ObjectName.ToString());
-		checkf(!(LoadClass->GetDefaultObject() && LoadClass->GetDefaultObject()->HasAnyFlags(RF_NeedLoad)), 
+		checkf(!(LoadClass->GetDefaultObject() && LoadClass->GetDefaultObject()->HasAnyFlags(RF_NeedLoad)),
 			TEXT("Class CDO %s had RF_NeedLoad while creating %s"), *LoadClass->GetDefaultObject()->GetFullName(), *ObjectName.ToString());
 		checkf(!ExportObject.TemplateObject->HasAnyFlags(RF_NeedLoad),
 			TEXT("Template %s had RF_NeedLoad while creating %s"), *ExportObject.TemplateObject->GetFullName(), *ObjectName.ToString());
@@ -8783,7 +8853,7 @@ void FAsyncPackage2::EventDrivenCreateExport(const FAsyncPackageHeaderData& Head
 	{
 		Object->SetExternalPackage(LinkerRoot);
 	}
-	
+
 	if (Desc.bCanBeImported && Export.PublicExportHash)
 	{
 		FlagsToSet |= EInternalObjectFlags::LoaderImport;
@@ -8977,7 +9047,7 @@ void FAsyncPackage2::EventDrivenCreateCellExport(const FAsyncPackageHeaderData&
 #endif
 
 	Ar->ExportBufferEnd();
-	
+
 	if (Desc.bCanBeImported && Export.PublicExportHash)
 	{
 		ImportStore.StoreGlobalCell(Desc.UPackageId, Export.PublicExportHash, Cell);
@@ -9103,9 +9173,9 @@ EEventLoadNodeExecutionResult FAsyncPackage2::Event_ExportsDone(FAsyncLoadingThr
 	UE_ASYNC_PACKAGE_DEBUG(Package->Desc);
 
 	check(FUObjectThreadContext::Get().AsyncVisibilityFilter == EInternalObjectFlags::None);
-	check(Package->AsyncPackageLoadingState == EAsyncPackageLoadingState2::ProcessExportBundles || 
-		  Package->AsyncPackageLoadingState == EAsyncPackageLoadingState2::WaitingForExternalReads ||
-		  Package->AsyncPackageLoadingState == EAsyncPackageLoadingState2::PreloadLinkerLoadExports);
+
+	check(Package->AsyncPackageLoadingState == EAsyncPackageLoadingState2::WaitingForExternalReads ||
+		  Package->AsyncPackageLoadingState == EAsyncPackageLoadingState2::PreloadExports);
 
 	Package->AsyncPackageLoadingState = EAsyncPackageLoadingState2::ExportsDone;
 
@@ -9175,14 +9245,14 @@ void FAsyncPackage2::CheckThatAllDependenciesHaveReachedStateDebug(FAsyncLoading
 		Stack.Pop();
 		FAsyncPackage2* Package = PackageAndDependencyChain.Get<0>();
 		DependencyChain = PackageAndDependencyChain.Get<1>();
-		
+
 		for (FAsyncPackage2* ImportedPackage : Package->Data.ImportedAsyncPackages)
 		{
 			if (ImportedPackage && !Visited.Contains(ImportedPackage) && !ThreadState.PackagesOnStack.Contains(Package))
 			{
 				TArray<FAsyncPackage2*> NextDependencyChain = DependencyChain;
 				NextDependencyChain.Add(ImportedPackage);
-				
+
 				check(ImportedPackage->AsyncPackageLoadingState >= PackageStateForCircularDependencies);
 				if (ImportedPackage->AsyncPackageLoadingState < PackageState)
 				{
@@ -9201,9 +9271,9 @@ void FAsyncPackage2::CheckThatAllDependenciesHaveReachedStateDebug(FAsyncLoading
 FAsyncPackage2* FAsyncPackage2::UpdateDependenciesStateRecursive(FAsyncLoadingThreadState2& ThreadState, FUpdateDependenciesStateRecursiveContext& Context)
 {
 	FAllDependenciesState& ThisState = this->*Context.StateMemberPtr;
-	
+
 	check(ThisState.PreOrderNumber < 0);
-	
+
 	if (ThisState.bAllDone)
 	{
 		return nullptr;
@@ -9233,7 +9303,7 @@ FAsyncPackage2* FAsyncPackage2::UpdateDependenciesStateRecursive(FAsyncLoadingTh
 	Context.S.Push(this);
 	Context.P.Push(this);
 
-	auto ProcessImportedPackage = 
+	auto ProcessImportedPackage =
 		[&ThreadState, &Context, &WaitingForPackage](FAsyncPackage2* ImportedPackage) -> bool
 		{
 			if (!ImportedPackage)
@@ -9282,7 +9352,7 @@ FAsyncPackage2* FAsyncPackage2::UpdateDependenciesStateRecursive(FAsyncLoadingTh
 
 			return true;
 		};
-		
+
 	bool bContinueToProcessPackage = true;
 	for (FAsyncPackage2* ImportedPackage : Data.ImportedAsyncPackages)
 	{
@@ -9331,7 +9401,7 @@ FAsyncPackage2* FAsyncPackage2::UpdateDependenciesStateRecursive(FAsyncLoadingTh
 		} while (InStronglyConnectedComponent != this);
 		Context.P.Pop();
 	}
-	
+
 	return WaitingForPackage;
 }
 
@@ -9369,21 +9439,8 @@ void FAsyncPackage2::ConditionalBeginProcessPackageExports(FAsyncLoadingThreadSt
 		{
 			check(Package->AsyncPackageLoadingState == EAsyncPackageLoadingState2::DependenciesReady);
 
-			if (Package->LinkerLoadState.IsSet())
-			{
-				Package->AsyncPackageLoadingState = EAsyncPackageLoadingState2::CreateLinkerLoadExports;
-				Package->GetPackageNode(EEventLoadNode2::Package_CreateLinkerLoadExports).ReleaseBarrier(&ThreadState);
-			}
-			else
-			{
-				Package->AsyncPackageLoadingState = EAsyncPackageLoadingState2::ProcessExportBundles;
-				if (Package->Data.TotalExportBundleCount > 0)
-				{
-					// Release a single export bundle node to avoid them being picked up recursively during a flush.
-					// When a node finishes, it will release another one.
-					Package->GetExportBundleNode(EEventLoadNode2::ExportBundle_Process, 0).ReleaseBarrier(&ThreadState);
-				}
-			}
+			Package->AsyncPackageLoadingState = EAsyncPackageLoadingState2::CreateExports;
+			Package->GetPackageNode(EEventLoadNode2::Package_CreateExports).ReleaseBarrier(&ThreadState);
 		});
 }
 
@@ -9405,7 +9462,7 @@ void FAsyncPackage2::ConditionalFinishLoading(FAsyncLoadingThreadState2& ThreadS
 		});
 }
 
-EEventLoadNodeExecutionResult FAsyncPackage2::Event_PostLoadExportBundle(FAsyncLoadingThreadState2& ThreadState, FAsyncPackage2* Package, int32 InExportBundleIndex)
+EEventLoadNodeExecutionResult FAsyncPackage2::Event_PostLoad(FAsyncLoadingThreadState2& ThreadState, FAsyncPackage2* Package, int32)
 {
 	TRACE_CPUPROFILER_EVENT_SCOPE(Event_PostLoad);
 	UE_ASYNC_PACKAGE_DEBUG(Package->Desc);
@@ -9421,7 +9478,9 @@ EEventLoadNodeExecutionResult FAsyncPackage2::Event_PostLoadExportBundle(FAsyncL
 		return Package->ExecutePostLoadLinkerLoadPackageExports(ThreadState);
 	}
 
-	check(InExportBundleIndex < Package->Data.TotalExportBundleCount);
+#if WITH_EDITOR
+	UE::Core::Private::FPlayInEditorLoadingScope PlayInEditorIDScope(Package->Desc.PIEInstanceID);
+#endif
 
 	EEventLoadNodeExecutionResult LoadingState = EEventLoadNodeExecutionResult::Complete;
 
@@ -9439,78 +9498,83 @@ EEventLoadNodeExecutionResult FAsyncPackage2::Event_PostLoadExportBundle(FAsyncL
 		const bool bAsyncPostLoadEnabled = FAsyncLoadingThreadSettings::Get().bAsyncPostLoadEnabled;
 		const bool bIsMultithreaded = Package->AsyncLoadingThread.IsMultithreaded();
 
+		while (Package->ProcessedExportBundlesCount < Package->Data.TotalExportBundleCount)
 		{
 #if WITH_EDITOR
-		UE::Core::Private::FPlayInEditorLoadingScope PlayInEditorIDScope(Package->Desc.PIEInstanceID);
-
-		const FAsyncPackageHeaderData* HeaderData;
-		if (InExportBundleIndex == 1)
-		{
-			HeaderData = Package->OptionalSegmentHeaderData.GetPtrOrNull();
-			check(HeaderData);
-		}
-		else
-		{
-			check(InExportBundleIndex == 0);
-			HeaderData = &Package->HeaderData;
-		}
+			const FAsyncPackageHeaderData* HeaderData;
+			if (Package->ProcessedExportBundlesCount == 1)
+			{
+				HeaderData = Package->OptionalSegmentHeaderData.GetPtrOrNull();
+				check(HeaderData);
+			}
+			else
+			{
+				check(Package->ProcessedExportBundlesCount == 0);
+				HeaderData = &Package->HeaderData;
+			}
 #else
-		const FAsyncPackageHeaderData* HeaderData = &Package->HeaderData;
+			const FAsyncPackageHeaderData* HeaderData = &Package->HeaderData;
 #endif
 
-		while (Package->ExportBundleEntryIndex < HeaderData->ExportBundleEntriesCopyForPostLoad.Num())
-		{
-			const FExportBundleEntry& BundleEntry = HeaderData->ExportBundleEntriesCopyForPostLoad[Package->ExportBundleEntryIndex];
-			if (ThreadState.IsTimeLimitExceeded(TEXT("Event_PostLoadExportBundle")))
+			while (Package->ExportBundleEntryIndex < HeaderData->ExportBundleEntriesCopyForPostLoad.Num())
 			{
-				LoadingState = EEventLoadNodeExecutionResult::Timeout;
-				break;
-			}
-			
-			if (BundleEntry.CommandType == FExportBundleEntry::ExportCommandType_Serialize)
-			{
-				do
+				const FExportBundleEntry& BundleEntry = HeaderData->ExportBundleEntriesCopyForPostLoad[Package->ExportBundleEntryIndex];
+				if (ThreadState.IsTimeLimitExceeded(TEXT("Event_PostLoadExportBundle")))
 				{
-					if (uint32(HeaderData->ExportsView.Num()) <= BundleEntry.LocalExportIndex)
-					{
-						break;
-					}
+					LoadingState = EEventLoadNodeExecutionResult::Timeout;
+					break;
+				}
 
-					FExportObject& Export = HeaderData->ExportsView[BundleEntry.LocalExportIndex];
-					if (Export.bFiltered || Export.bExportLoadFailed)
+				if (BundleEntry.CommandType == FExportBundleEntry::ExportCommandType_Serialize)
+				{
+					do
 					{
-						break;
-					}
+						if (uint32(HeaderData->ExportsView.Num()) <= BundleEntry.LocalExportIndex)
+						{
+							break;
+						}
 
-					UObject* Object = Export.Object;
-					check(Object);
-					checkObject(Object, !Object->HasAnyFlags(RF_NeedLoad));
-					if (!Object->HasAnyFlags(RF_NeedPostLoad))
-					{
-						break;
-					}
+						FExportObject& Export = HeaderData->ExportsView[BundleEntry.LocalExportIndex];
+						if (Export.bFiltered || Export.bExportLoadFailed)
+						{
+							break;
+						}
 
-					check(Object->IsReadyForAsyncPostLoad());
-					if (!bIsMultithreaded || (bAsyncPostLoadEnabled && CanPostLoadOnAsyncLoadingThread(Object)))
-					{
+						UObject* Object = Export.Object;
+						check(Object);
+						checkObject(Object, !Object->HasAnyFlags(RF_NeedLoad));
+						if (!Object->HasAnyFlags(RF_NeedPostLoad))
+						{
+							break;
+						}
+
+						check(Object->IsReadyForAsyncPostLoad());
+						if (!bIsMultithreaded || (bAsyncPostLoadEnabled && CanPostLoadOnAsyncLoadingThread(Object)))
+						{
 #if WITH_EDITOR
-						SCOPED_LOADTIMER_ASSET_TEXT(*Object->GetPathName());
+							SCOPED_LOADTIMER_ASSET_TEXT(*Object->GetPathName());
 #endif
-						ThreadContext.CurrentlyPostLoadedObjectByALT = Object;
-						Object->ConditionalPostLoad();
-						ThreadContext.CurrentlyPostLoadedObjectByALT = nullptr;
-					}
-				} while (false);
+							ThreadContext.CurrentlyPostLoadedObjectByALT = Object;
+							Object->ConditionalPostLoad();
+							ThreadContext.CurrentlyPostLoadedObjectByALT = nullptr;
+						}
+					} while (false);
+				}
+				++Package->ExportBundleEntryIndex;
+			}
+
+			if (LoadingState == EEventLoadNodeExecutionResult::Timeout)
+			{
+				break;
 			}
-			++Package->ExportBundleEntryIndex;
-		}
 
+			Package->ProcessedExportBundlesCount++;
 		}
 
 		// End async loading, simulates EndLoad
 		Package->EndAsyncLoad();
 	}
-	
+
 	if (LoadingState == EEventLoadNodeExecutionResult::Timeout)
 	{
 		return LoadingState;
@@ -9518,7 +9582,7 @@ EEventLoadNodeExecutionResult FAsyncPackage2::Event_PostLoadExportBundle(FAsyncL
 
 	Package->ExportBundleEntryIndex = 0;
 
-	if (++Package->ProcessedExportBundlesCount == Package->Data.TotalExportBundleCount)
+	if (Package->ProcessedExportBundlesCount == Package->Data.TotalExportBundleCount)
 	{
 		Package->ProcessedExportBundlesCount = 0;
 
@@ -9528,7 +9592,7 @@ EEventLoadNodeExecutionResult FAsyncPackage2::Event_PostLoadExportBundle(FAsyncL
 	return EEventLoadNodeExecutionResult::Complete;
 }
 
-EEventLoadNodeExecutionResult FAsyncPackage2::Event_DeferredPostLoadExportBundle(FAsyncLoadingThreadState2& ThreadState, FAsyncPackage2* Package, int32 InExportBundleIndex)
+EEventLoadNodeExecutionResult FAsyncPackage2::Event_DeferredPostLoad(FAsyncLoadingThreadState2& ThreadState, FAsyncPackage2* Package, int32)
 {
 	SCOPE_CYCLE_COUNTER(STAT_FAsyncPackage_PostLoadObjectsGameThread);
 	TRACE_CPUPROFILER_EVENT_SCOPE(Event_DeferredPostLoad);
@@ -9549,7 +9613,6 @@ EEventLoadNodeExecutionResult FAsyncPackage2::Event_DeferredPostLoadExportBundle
 		return Package->ExecuteDeferredPostLoadLinkerLoadPackageExports(ThreadState);
 	}
 
-	check(InExportBundleIndex < Package->Data.TotalExportBundleCount);
 	EEventLoadNodeExecutionResult LoadingState = EEventLoadNodeExecutionResult::Complete;
 
 	if (!Package->HasLoadFailed())
@@ -9559,71 +9622,81 @@ EEventLoadNodeExecutionResult FAsyncPackage2::Event_DeferredPostLoadExportBundle
 
 		FAsyncLoadingTickScope2 InAsyncLoadingTick(Package->AsyncLoadingThread);
 
+		while (Package->ProcessedExportBundlesCount < Package->Data.TotalExportBundleCount)
+		{
 #if WITH_EDITOR
-		UE::Core::Private::FPlayInEditorLoadingScope PlayInEditorIDScope(Package->Desc.PIEInstanceID);
+			UE::Core::Private::FPlayInEditorLoadingScope PlayInEditorIDScope(Package->Desc.PIEInstanceID);
 
-		const FAsyncPackageHeaderData* HeaderData;
-		if (InExportBundleIndex == 1)
-		{
-			HeaderData = Package->OptionalSegmentHeaderData.GetPtrOrNull();
-			check(HeaderData);
-		}
-		else
-		{
-			check(InExportBundleIndex == 0);
-			HeaderData = &Package->HeaderData;
-		}
+			const FAsyncPackageHeaderData* HeaderData;
+			if (Package->ProcessedExportBundlesCount == 1)
+			{
+				HeaderData = Package->OptionalSegmentHeaderData.GetPtrOrNull();
+				check(HeaderData);
+			}
+			else
+			{
+				check(Package->ProcessedExportBundlesCount == 0);
+				HeaderData = &Package->HeaderData;
+			}
 #else
-		const FAsyncPackageHeaderData* HeaderData = &Package->HeaderData;
+			const FAsyncPackageHeaderData* HeaderData = &Package->HeaderData;
 #endif
 
-		while (Package->ExportBundleEntryIndex < HeaderData->ExportBundleEntriesCopyForPostLoad.Num())
-		{
-			const FExportBundleEntry& BundleEntry = HeaderData->ExportBundleEntriesCopyForPostLoad[Package->ExportBundleEntryIndex];
-			if (ThreadState.IsTimeLimitExceeded(TEXT("Event_DeferredPostLoadExportBundle")))
+			while (Package->ExportBundleEntryIndex < HeaderData->ExportBundleEntriesCopyForPostLoad.Num())
 			{
-				LoadingState = EEventLoadNodeExecutionResult::Timeout;
-				break;
-			}
+				const FExportBundleEntry& BundleEntry = HeaderData->ExportBundleEntriesCopyForPostLoad[Package->ExportBundleEntryIndex];
+				if (ThreadState.IsTimeLimitExceeded(TEXT("Event_DeferredPostLoadExportBundle")))
+				{
+					LoadingState = EEventLoadNodeExecutionResult::Timeout;
+					break;
+				}
 
-			if (BundleEntry.CommandType == FExportBundleEntry::ExportCommandType_Serialize)
-			{
-				do
+				if (BundleEntry.CommandType == FExportBundleEntry::ExportCommandType_Serialize)
 				{
-					if (uint32(HeaderData->ExportsView.Num()) <= BundleEntry.LocalExportIndex)
+					do
 					{
-						break;
-					}
+						if (uint32(HeaderData->ExportsView.Num()) <= BundleEntry.LocalExportIndex)
+						{
+							break;
+						}
 
-					FExportObject& Export = HeaderData->ExportsView[BundleEntry.LocalExportIndex];
-					if (Export.bFiltered || Export.bExportLoadFailed)
-					{
-						break;
-					}
+						FExportObject& Export = HeaderData->ExportsView[BundleEntry.LocalExportIndex];
+						if (Export.bFiltered || Export.bExportLoadFailed)
+						{
+							break;
+						}
 
-					UObject* Object = Export.Object;
-					check(Object);
-					checkObject(Object, !Object->HasAnyFlags(RF_NeedLoad));
-					if (Object->HasAnyFlags(RF_NeedPostLoad))
-					{
+						UObject* Object = Export.Object;
+						check(Object);
+						checkObject(Object, !Object->HasAnyFlags(RF_NeedLoad));
+						if (Object->HasAnyFlags(RF_NeedPostLoad))
+						{
 #if WITH_EDITOR
-						SCOPED_LOADTIMER_ASSET_TEXT(*Object->GetPathName());
+							SCOPED_LOADTIMER_ASSET_TEXT(*Object->GetPathName());
 #endif
-						PackageScope.ThreadContext.CurrentlyPostLoadedObjectByALT = Object;
-						{
-							FScopeCycleCounterUObject ConstructorScope(Object, GET_STATID(STAT_FAsyncPackage_PostLoadObjectsGameThread));
-							LLM_SCOPE(ELLMTag::UObject);
-							LLM_SCOPE_BYTAG(UObject_FAsyncPackage2);
-							LLM_SCOPE_DYNAMIC_STAT_OBJECTPATH(Object->GetPackage(), ELLMTagSet::Assets);
-							LLM_SCOPE_DYNAMIC_STAT_OBJECTPATH(Object->GetClass(), ELLMTagSet::AssetClasses);
-							UE_TRACE_METADATA_SCOPE_ASSET(Object, Object->GetClass());
-							Object->ConditionalPostLoad();
+							PackageScope.ThreadContext.CurrentlyPostLoadedObjectByALT = Object;
+							{
+								FScopeCycleCounterUObject ConstructorScope(Object, GET_STATID(STAT_FAsyncPackage_PostLoadObjectsGameThread));
+								LLM_SCOPE(ELLMTag::UObject);
+								LLM_SCOPE_BYTAG(UObject_FAsyncPackage2);
+								LLM_SCOPE_DYNAMIC_STAT_OBJECTPATH(Object->GetPackage(), ELLMTagSet::Assets);
+								LLM_SCOPE_DYNAMIC_STAT_OBJECTPATH(Object->GetClass(), ELLMTagSet::AssetClasses);
+								UE_TRACE_METADATA_SCOPE_ASSET(Object, Object->GetClass());
+								Object->ConditionalPostLoad();
+							}
+							PackageScope.ThreadContext.CurrentlyPostLoadedObjectByALT = nullptr;
 						}
-						PackageScope.ThreadContext.CurrentlyPostLoadedObjectByALT = nullptr;
-					}
-				} while (false);
+					} while (false);
+				}
+				++Package->ExportBundleEntryIndex;
 			}
-			++Package->ExportBundleEntryIndex;
+
+			if (LoadingState == EEventLoadNodeExecutionResult::Timeout)
+			{
+				break;
+			}
+
+			++Package->ProcessedExportBundlesCount;
 		}
 	}
 
@@ -9634,7 +9707,7 @@ EEventLoadNodeExecutionResult FAsyncPackage2::Event_DeferredPostLoadExportBundle
 
 	Package->ExportBundleEntryIndex = 0;
 
-	if (++Package->ProcessedExportBundlesCount == Package->Data.TotalExportBundleCount)
+	if (Package->ProcessedExportBundlesCount == Package->Data.TotalExportBundleCount)
 	{
 		Package->ProcessedExportBundlesCount = 0;
 		check(Package->AsyncPackageLoadingState == EAsyncPackageLoadingState2::DeferredPostLoad);
@@ -9651,13 +9724,6 @@ FEventLoadNode2& FAsyncPackage2::GetPackageNode(EEventLoadNode2 Phase)
 	return PackageNodes[Phase];
 }
 
-FEventLoadNode2& FAsyncPackage2::GetExportBundleNode(EEventLoadNode2 Phase, uint32 ExportBundleIndex)
-{
-	check(ExportBundleIndex < uint32(Data.TotalExportBundleCount));
-	uint32 ExportBundleNodeIndex = ExportBundleIndex * EEventLoadNode2::ExportBundle_NumPhases + Phase;
-	return Data.ExportBundleNodes[ExportBundleNodeIndex];
-}
-
 void FAsyncLoadingThread2::CleanupSyncLoadContext(FAsyncLoadingThreadState2& ThreadState)
 {
 	if (ThreadState.bIsAsyncLoadingThread)
@@ -9730,7 +9796,7 @@ void FAsyncLoadingThread2::UpdateSyncLoadContext(FAsyncLoadingThreadState2& Thre
 		// Ensure that we've created the package we're waiting for
 		CreateAsyncPackagesFromQueue(ThreadState);
 		int32 FoundPackages = 0;
-		for (int32 i=0; i < SyncLoadContext->RequestIDs.Num(); ++i)
+		for (int32 i = 0; i < SyncLoadContext->RequestIDs.Num(); ++i)
 		{
 			int32 RequestID = SyncLoadContext->RequestIDs[i];
 			if (SyncLoadContext->IsRequestedPackageValid(i))
@@ -9761,7 +9827,7 @@ void FAsyncLoadingThread2::UpdateSyncLoadContext(FAsyncLoadingThreadState2& Thre
 				++FoundPackages;
 			}
 		}
-		
+
 		// Only set when full list is available 
 		if (FoundPackages == SyncLoadContext->RequestIDs.Num())
 		{
@@ -9780,7 +9846,7 @@ void FAsyncLoadingThread2::UpdateSyncLoadContext(FAsyncLoadingThreadState2& Thre
 				// Flushing a package while it's already being processed on the stack, if we're done preloading we let it pass and remove the request id
 				bool bPreloadIsDone = RequestedPackage->AsyncPackageLoadingState >= EAsyncPackageLoadingState2::DeferredPostLoad;
 				UE_CLOG(!bPreloadIsDone, LogStreaming, Warning, TEXT("Flushing package %s while it's being preloaded in the same callstack is not possible. Releasing request %d to unblock."), *RequestedPackage->Desc.UPackageName.ToString(), RequestID);
-				RemovePendingRequests(ThreadState, {RequestID});
+				RemovePendingRequests(ThreadState, { RequestID });
 				bMayHaveFulfilledRequest = true;
 			}
 		}
@@ -9807,7 +9873,7 @@ EAsyncPackageState::Type FAsyncLoadingThread2::ProcessAsyncLoadingFromGameThread
 
 	while (true)
 	{
-		do 
+		do
 		{
 			if ((++LoopIterations) % 32 == 31)
 			{
@@ -10091,7 +10157,7 @@ EAsyncPackageState::Type FAsyncLoadingThread2::ProcessLoadedPackagesFromGameThre
 		for (int32 CompletedPackageRequestIndex = CompletedPackageRequests.Num() - 1; CompletedPackageRequestIndex >= 0; --CompletedPackageRequestIndex)
 		{
 			FCompletedPackageRequest& CompletedPackageRequest = CompletedPackageRequests[CompletedPackageRequestIndex];
-			if (FlushRequestIDs.Num() == 0 
+			if (FlushRequestIDs.Num() == 0
 				|| Algo::AnyOf(FlushRequestIDs, [&CompletedPackageRequest](int32 FlushRequestID) { return CompletedPackageRequest.RequestIDs.Contains(FlushRequestID); }))
 			{
 				RemovePendingRequests(ThreadState, CompletedPackageRequest.RequestIDs);
@@ -10124,7 +10190,7 @@ EAsyncPackageState::Type FAsyncLoadingThread2::ProcessLoadedPackagesFromGameThre
 		}
 
 		bDidSomething = true;
-		
+
 		if (FlushRequestIDs.Num() != 0 && !ContainsAnyRequestID(FlushRequestIDs))
 		{
 			// The only packages we care about have finished loading, so we're good to exit
@@ -10224,18 +10290,16 @@ FAsyncLoadingThread2::FAsyncLoadingThread2(FIoDispatcher& InIoDispatcher)
 	PendingPackages.Reserve(DefaultAsyncPackagesReserveCount);
 	RequestIdToPackageMap.Reserve(DefaultAsyncPackagesReserveCount);
 
-	EventSpecs.AddDefaulted(EEventLoadNode2::Package_NumPhases + EEventLoadNode2::ExportBundle_NumPhases);
+	EventSpecs.AddDefaulted(EEventLoadNode2::Package_NumPhases);
 	EventSpecs[EEventLoadNode2::Package_ProcessSummary] = { &FAsyncPackage2::Event_ProcessPackageSummary, &EventQueue, EExecutionType::Parallel, TEXT("ProcessSummary") };
 	EventSpecs[EEventLoadNode2::Package_ProcessDependencies] = { &FAsyncPackage2::Event_ProcessDependencies, &EventQueue, EExecutionType::Normal, TEXT("ProcessDependencies") };
-	EventSpecs[EEventLoadNode2::Package_DependenciesReady] = { &FAsyncPackage2::Event_DependenciesReady, &EventQueue, EExecutionType::Normal, TEXT("DependenciesReady")};
-	EventSpecs[EEventLoadNode2::Package_CreateLinkerLoadExports] = { &FAsyncPackage2::Event_CreateLinkerLoadExports, &EventQueue, EExecutionType::Normal, TEXT("CreateLinkerLoadExports") };
-	EventSpecs[EEventLoadNode2::Package_ResolveLinkerLoadImports] = { &FAsyncPackage2::Event_ResolveLinkerLoadImports, &EventQueue, EExecutionType::Normal, TEXT("ResolveLinkerLoadImports") };
-	EventSpecs[EEventLoadNode2::Package_PreloadLinkerLoadExports] = { &FAsyncPackage2::Event_PreloadLinkerLoadExports, &EventQueue, EExecutionType::Parallel, TEXT("PreloadLinkerLoadExports") };
+	EventSpecs[EEventLoadNode2::Package_DependenciesReady] = { &FAsyncPackage2::Event_DependenciesReady, &EventQueue, EExecutionType::Normal, TEXT("DependenciesReady") };
+	EventSpecs[EEventLoadNode2::Package_CreateExports] = { &FAsyncPackage2::Event_CreateExports, &EventQueue, EExecutionType::Normal, TEXT("CreateExports") };
+	EventSpecs[EEventLoadNode2::Package_ResolveImports] = { &FAsyncPackage2::Event_ResolveImports, &EventQueue, EExecutionType::Normal, TEXT("ResolveImports") };
+	EventSpecs[EEventLoadNode2::Package_PreloadExports] = { &FAsyncPackage2::Event_PreloadExports, &EventQueue, EExecutionType::Parallel, TEXT("PreloadExports") };
 	EventSpecs[EEventLoadNode2::Package_ExportsSerialized] = { &FAsyncPackage2::Event_ExportsDone, &EventQueue, EExecutionType::Normal, TEXT("ExportsSerialized") };
-
-	EventSpecs[EEventLoadNode2::Package_NumPhases + EEventLoadNode2::ExportBundle_Process] = { &FAsyncPackage2::Event_ProcessExportBundle, &EventQueue, EExecutionType::Normal, TEXT("ExportBundle_Process") };
-	EventSpecs[EEventLoadNode2::Package_NumPhases + EEventLoadNode2::ExportBundle_PostLoad] = { &FAsyncPackage2::Event_PostLoadExportBundle, &EventQueue, EExecutionType::Normal, TEXT("ExportBundle_PostLoad") };
-	EventSpecs[EEventLoadNode2::Package_NumPhases + EEventLoadNode2::ExportBundle_DeferredPostLoad] = { &FAsyncPackage2::Event_DeferredPostLoadExportBundle, &MainThreadEventQueue, EExecutionType::Normal, TEXT("ExportBundle_DeferredPostLoad") };
+	EventSpecs[EEventLoadNode2::Package_PostLoad] = { &FAsyncPackage2::Event_PostLoad, &EventQueue, EExecutionType::Normal, TEXT("PostLoad") };
+	EventSpecs[EEventLoadNode2::Package_DeferredPostLoad] = { &FAsyncPackage2::Event_DeferredPostLoad, &MainThreadEventQueue, EExecutionType::Normal, TEXT("DeferredPostLoad") };
 
 	CancelLoadingEvent = FPlatformProcess::GetSynchEventFromPool();
 	ThreadSuspendedEvent = FPlatformProcess::GetSynchEventFromPool();
@@ -10351,7 +10415,7 @@ uint32 FAsyncLoadingThread2::Run()
 	FMemory::SetupTLSCachesOnCurrentThread();
 
 	FAsyncLoadingThreadState2& ThreadState = *FAsyncLoadingThreadState2::Get();
-	
+
 #if WITH_VERSE_VM || defined(__INTELLISENSE__)
 	Verse::FIOContext Context = Verse::FIOContext::CreateForManualStackScanning();
 #endif
@@ -10463,7 +10527,7 @@ uint32 FAsyncLoadingThread2::Run()
 					{
 						bool bDidSomething = false;
 						bool bPopped = false;
-						do 
+						do
 						{
 							bPopped = false;
 							UpdateSyncLoadContext(ThreadState);
@@ -10550,7 +10614,7 @@ EAsyncPackageState::Type FAsyncLoadingThread2::TickAsyncThreadFromGameThread(FAs
 {
 	check(IsInGameThread());
 	EAsyncPackageState::Type Result = EAsyncPackageState::Complete;
-	
+
 	if (AsyncThreadReady)
 	{
 		if (ThreadState.IsTimeLimitExceeded(TEXT("TickAsyncThreadFromGameThread")))
@@ -10641,12 +10705,12 @@ static void VerifyObjectLoadFlagsWhenFinishedLoading()
 	const EInternalObjectFlags AsyncFlags =
 		EInternalObjectFlags_AsyncLoading;
 
-	const EObjectFlags LoadIntermediateFlags = 
+	const EObjectFlags LoadIntermediateFlags =
 		EObjectFlags::RF_NeedLoad | EObjectFlags::RF_WillBeLoaded |
 		EObjectFlags::RF_NeedPostLoad | RF_NeedPostLoadSubobjects;
 
 	ParallelFor(TEXT("VerifyObjectLoadFlagsDebugTask"), GUObjectArray.GetObjectArrayNum(), 512,
-		[AsyncFlags,LoadIntermediateFlags](int32 ObjectIndex)
+		[AsyncFlags, LoadIntermediateFlags](int32 ObjectIndex)
 	{
 		FUObjectItem* ObjectItem = &GUObjectArray.GetObjectItemArrayUnsafe()[ObjectIndex];
 		if (UObject* Obj = static_cast<UObject*>(ObjectItem->GetObject()))
@@ -11030,10 +11094,10 @@ FAsyncPackage2::FAsyncPackage2(
 	FAsyncLoadingThread2& InAsyncLoadingThread,
 	FAsyncLoadEventGraphAllocator& InGraphAllocator,
 	const FAsyncLoadEventSpec* EventSpecs)
-: Desc(InDesc)
-, AsyncLoadingThread(InAsyncLoadingThread)
-, GraphAllocator(InGraphAllocator)
-, ImportStore(AsyncLoadingThread.GlobalImportStore)
+	: Desc(InDesc)
+	, AsyncLoadingThread(InAsyncLoadingThread)
+	, GraphAllocator(InGraphAllocator)
+	, ImportStore(AsyncLoadingThread.GlobalImportStore)
 {
 	TRACE_LOADTIME_NEW_ASYNC_PACKAGE(this);
 	AddRequestID(ThreadState, Desc.RequestID);
@@ -11046,7 +11110,7 @@ FAsyncPackage2::FAsyncPackage2(
 void FAsyncPackage2::CreatePackageNodes(const FAsyncLoadEventSpec* EventSpecs)
 {
 	const int32 BarrierCount = 1;
-	
+
 	FEventLoadNode2* Node = reinterpret_cast<FEventLoadNode2*>(PackageNodesMemory);
 	for (int32 Phase = 0; Phase < EEventLoadNode2::Package_NumPhases; ++Phase)
 	{
@@ -11055,19 +11119,6 @@ void FAsyncPackage2::CreatePackageNodes(const FAsyncLoadEventSpec* EventSpecs)
 	PackageNodes = MakeArrayView(Node, EEventLoadNode2::Package_NumPhases);
 }
 
-void FAsyncPackage2::CreateExportBundleNodes(const FAsyncLoadEventSpec* EventSpecs)
-{
-	const int32 BarrierCount = 1;
-	for (int32 ExportBundleIndex = 0; ExportBundleIndex < Data.TotalExportBundleCount; ++ExportBundleIndex)
-	{
-		uint32 NodeIndex = EEventLoadNode2::ExportBundle_NumPhases * ExportBundleIndex;
-		for (int32 Phase = 0; Phase < EEventLoadNode2::ExportBundle_NumPhases; ++Phase)
-		{
-			new (&Data.ExportBundleNodes[NodeIndex + Phase]) FEventLoadNode2(EventSpecs + EEventLoadNode2::Package_NumPhases + Phase, this, ExportBundleIndex, BarrierCount);
-		}
-	}
-}
-
 FAsyncPackage2::~FAsyncPackage2()
 {
 	TRACE_LOADTIME_DESTROY_ASYNC_PACKAGE(this);
@@ -11270,7 +11321,7 @@ void FAsyncPackage2::BeginAsyncLoad()
 }
 
 /**
- * End async loading process. Simulates parts of EndLoad(). 
+ * End async loading process. Simulates parts of EndLoad().
  */
 void FAsyncPackage2::EndAsyncLoad()
 {
@@ -11505,7 +11556,7 @@ void FAsyncPackage2::AddProgressCallback(TSharedPtr<FLoadPackageAsyncProgressDel
 }
 
 UE_TRACE_EVENT_BEGIN(CUSTOM_LOADTIMER_LOG, LoadAsyncPackageInternal, NoSync)
-	UE_TRACE_EVENT_FIELD(UE::Trace::WideString, PackageName)
+UE_TRACE_EVENT_FIELD(UE::Trace::WideString, PackageName)
 UE_TRACE_EVENT_END()
 
 int32 FAsyncLoadingThread2::LoadPackageInternal(const FPackagePath& InPackagePath, FName InCustomName, TUniquePtr<FLoadPackageAsyncDelegate>&& InCompletionDelegate, TSharedPtr<FLoadPackageAsyncProgressDelegate> InProgressDelegate, EPackageFlags InPackageFlags, int32 InPIEInstanceID, int32 InPackagePriority, const FLinkerInstancingContext* InInstancingContext, uint32 InLoadFlags)
@@ -11759,12 +11810,12 @@ void FAsyncLoadingThread2::FlushLoading(TConstArrayView<int32> RequestIDs)
 			// Don't forget to adjust unittest in FLoadingTests_InvalidFlush_FromWorker if changing this.
 #if !UE_BUILD_SHIPPING
 			const FString& ThreadName = FThreadManager::GetThreadName(FPlatformTLS::GetCurrentThreadId());
-			const FString Heading = 
+			const FString Heading =
 				FString::Printf(
-					TEXT("The current loader '%s' is unable to FlushAsyncLoading from the current thread '%s'. Flush will be ignored."), 
+					TEXT("The current loader '%s' is unable to FlushAsyncLoading from the current thread '%s'. Flush will be ignored."),
 					LexToString(LoaderType),
 					*ThreadName
-			);
+				);
 			FDebug::DumpStackTraceToLog(*Heading, ELogVerbosity::Error);
 #endif
 			return;
@@ -12121,7 +12172,7 @@ void FAsyncLoadingThread2::FLoaderStats::RecordPartialFlush(const FAsyncPackage2
 void FAsyncLoadingThread2::FLoaderStats::Log(FStringBuilderBase& Builder, const FPartialFlushInfo& PartialFlush)
 {
 #if !UE_BUILD_SHIPPING
-	FUniqueCallstack* Callstack = PartialFlush.UniqueCallstackIndex < 0 
+	FUniqueCallstack* Callstack = PartialFlush.UniqueCallstackIndex < 0
 		? nullptr : &UniqueCallstacks[PartialFlush.UniqueCallstackIndex];
 
 	//
@@ -12136,7 +12187,7 @@ void FAsyncLoadingThread2::FLoaderStats::Log(FStringBuilderBase& Builder, const
 		LexToString(PartialFlush.RequestingPackageState),
 		PartialFlush.UniqueCallstackIndex,
 		(!Callstack || Callstack->bHasLogged) ? TEXT(".") : TEXT(" - first occurrence - Callstack:\n")
-		);
+	);
 
 	if (Callstack && !Callstack->bHasLogged)
 	{
@@ -12179,68 +12230,68 @@ void FAsyncPackage2::DumpStateImpl(TSet<FAsyncPackage2*>& Set, int32 Indent, TMu
 	Set.Add(this);
 
 	auto GetPackageType = [](FAsyncPackage2* Package)
-	{
-		return LexToString(Package->Desc.Loader);
-	};
+		{
+			return LexToString(Package->Desc.Loader);
+		};
 
 	auto FormatPackage = [&GetPackageType](FAsyncPackage2* Package)
-	{
-		return FString::Printf(TEXT("%s package %s state %s"), GetPackageType(Package), *Package->Desc.UPackageName.ToString(), LexToString(Package->AsyncPackageLoadingState));
-	};
+		{
+			return FString::Printf(TEXT("%s package %s state %s"), GetPackageType(Package), *Package->Desc.UPackageName.ToString(), LexToString(Package->AsyncPackageLoadingState));
+		};
 
 	UE_LOG(LogStreaming, Warning, TEXT("%s%s:"), FCString::Spc(Indent), *FormatPackage(this));
 
 	auto VisitNodes = [](FEventLoadNode2& Node, TFunction<void(FEventLoadNode2*)> Visitor)
-	{
-		if (Node.DependenciesCount == 1)
 		{
-			Visitor(Node.SingleDependent);
-		}
-		else if (Node.DependenciesCount != 0)
-		{
-			FEventLoadNode2** Current = Node.MultipleDependents;
-			FEventLoadNode2** End = Node.MultipleDependents + Node.DependenciesCount;
-			for (; Current < End; ++Current)
+			if (Node.DependenciesCount == 1)
 			{
-				Visitor(*Current);
+				Visitor(Node.SingleDependent);
 			}
-		}
-	};
+			else if (Node.DependenciesCount != 0)
+			{
+				FEventLoadNode2** Current = Node.MultipleDependents;
+				FEventLoadNode2** End = Node.MultipleDependents + Node.DependenciesCount;
+				for (; Current < End; ++Current)
+				{
+					Visitor(*Current);
+				}
+			}
+		};
 
 	auto DumpNode = [&FormatPackage, &VisitNodes](TMultiMap<FEventLoadNode2*, FEventLoadNode2*>& MappedNodes, const TCHAR* Header, FEventLoadNode2& Node, int32 Indent)
-	{
-		bool bIsEmpty = Node.Spec->EventQueue->IsEmptyForDebug();
-
-		if (Node.DependenciesCount == 0 && Node.BarrierCount == 0 && bIsEmpty)
 		{
-			return false;
-		}
+			bool bIsEmpty = Node.Spec->EventQueue->IsEmptyForDebug();
 
-		if (Header)
-		{
-			UE_LOG(LogStreaming, Warning, TEXT("%s %s"), FCString::Spc(Indent), Header);
-		}
-		Indent++;
+			if (Node.DependenciesCount == 0 && Node.BarrierCount == 0 && bIsEmpty)
+			{
+				return false;
+			}
 
-		UE_LOG(LogStreaming, Warning, TEXT("%s Node %s (BarrierCount %d, EventQueue: %s)"), FCString::Spc(Indent), Node.GetSpec()->Name, Node.BarrierCount.load(), bIsEmpty ? TEXT("Empty") : TEXT("NonEmpty"));
+			if (Header)
+			{
+				UE_LOG(LogStreaming, Warning, TEXT("%s %s"), FCString::Spc(Indent), Header);
+			}
+			Indent++;
 
-		TArray<FEventLoadNode2*> WaitingOn;
-		MappedNodes.MultiFind(&Node, WaitingOn);
+			UE_LOG(LogStreaming, Warning, TEXT("%s Node %s (BarrierCount %d, EventQueue: %s)"), FCString::Spc(Indent), Node.GetSpec()->Name, Node.BarrierCount.load(), bIsEmpty ? TEXT("Empty") : TEXT("NonEmpty"));
 
-		for (FEventLoadNode2* WaitingNode : WaitingOn)
-		{
-			UE_LOG(LogStreaming, Warning, TEXT("%s Waiting on node %s from %s"), FCString::Spc(Indent + 1), WaitingNode->GetSpec()->Name, *FormatPackage(WaitingNode->GetPackage()));
-		}
+			TArray<FEventLoadNode2*> WaitingOn;
+			MappedNodes.MultiFind(&Node, WaitingOn);
 
-		VisitNodes(Node,
-			[Indent, &FormatPackage](FEventLoadNode2* Node)
+			for (FEventLoadNode2* WaitingNode : WaitingOn)
 			{
-				UE_LOG(LogStreaming, Warning, TEXT("%s Will trigger node %s for %s"), FCString::Spc(Indent + 1), Node->GetSpec()->Name, *FormatPackage(Node->GetPackage()));
+				UE_LOG(LogStreaming, Warning, TEXT("%s Waiting on node %s from %s"), FCString::Spc(Indent + 1), WaitingNode->GetSpec()->Name, *FormatPackage(WaitingNode->GetPackage()));
 			}
-		);
 
-		return true;
-	};
+			VisitNodes(Node,
+				[Indent, &FormatPackage](FEventLoadNode2* Node)
+				{
+					UE_LOG(LogStreaming, Warning, TEXT("%s Will trigger node %s for %s"), FCString::Spc(Indent + 1), Node->GetSpec()->Name, *FormatPackage(Node->GetPackage()));
+				}
+			);
+
+			return true;
+		};
 
 	TMultiMap<FEventLoadNode2*, FEventLoadNode2*> LocalMappedNode;
 	if (MappedNodes == nullptr)
@@ -12257,16 +12308,6 @@ void FAsyncPackage2::DumpStateImpl(TSet<FAsyncPackage2*>& Set, int32 Indent, TMu
 						FEventLoadNode2& Node = Package->GetPackageNode((EEventLoadNode2)PhaseIndex);
 						VisitNodes(Node, [&](FEventLoadNode2* Dependent) { LocalMappedNode.AddUnique(Dependent, &Node); });
 					}
-
-					for (int32 ExportBundleIndex = 0; ExportBundleIndex < Package->Data.TotalExportBundleCount; ++ExportBundleIndex)
-					{
-						for (uint8 PhaseIndex = 0; PhaseIndex < ExportBundle_NumPhases; ++PhaseIndex)
-						{
-							EEventLoadNode2 EventNode = (EEventLoadNode2)PhaseIndex;
-							FEventLoadNode2& Node = Package->GetExportBundleNode(EventNode, ExportBundleIndex);
-							VisitNodes(Node, [&](FEventLoadNode2* Dependent) { LocalMappedNode.AddUnique(Dependent, &Node); });
-						}
-					}
 				}
 			}
 		);
@@ -12293,20 +12334,6 @@ void FAsyncPackage2::DumpStateImpl(TSet<FAsyncPackage2*>& Set, int32 Indent, TMu
 		DumpNode(*MappedNodes, nullptr, GetPackageNode(EventNode), Indent + 1);
 	}
 
-	for (int32 ExportBundleIndex = 0; ExportBundleIndex < Data.TotalExportBundleCount; ++ExportBundleIndex)
-	{
-		FString HeaderStr = FString::Printf(TEXT("ExportBundle %d"), ExportBundleIndex);
-		const TCHAR* Header = *HeaderStr;
-		for (uint8 PhaseIndex = 0; PhaseIndex < ExportBundle_NumPhases; ++PhaseIndex)
-		{
-			EEventLoadNode2 EventNode = (EEventLoadNode2)PhaseIndex;
-			if (DumpNode(*MappedNodes, Header, GetExportBundleNode(EventNode, ExportBundleIndex), Indent + 1))
-			{
-				Header = nullptr;
-			}
-		}
-	}
-
 	for (FAsyncPackage2* Import : Data.ImportedAsyncPackages)
 	{
 		if (Import && Import->AsyncPackageLoadingState < EAsyncPackageLoadingState2::DeferredDelete)
@@ -12424,7 +12451,7 @@ bool FPublicExportMapIterationTest::RunTest(const FString& Parameters)
 	{
 		Map.Store(Keys[i], Values[i]);
 	}
-		
+
 	TSet<uint64> FoundKeys;
 	for (uint64 Key : Map.GetKeys())
 	{
@@ -12502,7 +12529,7 @@ bool FPublicExportMapPreInsertPublicExportsTest::RunTest(const FString& Paramete
 	// Empty map, populate a single key
 	{
 		FPublicExportMap Map;
-		TArray<FExportMapEntry> Exports = { {.PublicExportHash = 1 }};
+		TArray<FExportMapEntry> Exports = { {.PublicExportHash = 1 } };
 		int32 OriginaMapCount = Map.Num();
 
 		TestEqual(TEXT("Empty map has no count"), OriginaMapCount, 0);
@@ -12597,7 +12624,7 @@ bool FPublicExportMapPreInsertPublicExportsTest::RunTest(const FString& Paramete
 	// One existing entry (preinsert a single entry overlapping)
 	{
 		FPublicExportMap Map;
-		TArray<FExportMapEntry> Exports = { {.PublicExportHash = 1 }};
+		TArray<FExportMapEntry> Exports = { {.PublicExportHash = 1 } };
 
 		uint64 PreExistingKey = 1;
 		int32 PreExistingValue = 1;
