Always populate InitialNativeReferences on the first reachability pass

StartReachabilityAnalysis populates both InitialReferences (from FGCObjects with AddStableNativeReferencesOnly) and InitialObjects (from GRoots and KeepFlags). The first PerformReachabilityAnalysisPass then populates Context->InitialNativeReferences from InitialReferences, and Context->InitialObjects from InitialObjects.

However, PerformReachabilityAnalysis calls StartVerseGC after StartReachabilityAnalysis and before PerformReachabilityAnalysisPass, which means the Verse GC thread may start adding to GReachableObjects. If this happens early enough, PerformReachabilityAnalysisPass assumes that `!GReachableObjects.IsEmpty()` means that the first pass must have already occurred, and skips InitialNativeReferences, causing these objects to go unmarked.

Populating InitialNativeReferences when `!GReachabilityState.IsSuspended()` means it will always happen on the first pass regardless of GReachableObjects. The existing checks for `GetNumIterations` or `bFoundGarbageRef` then become redundant.

#rnx

diff --git a/Engine/Source/Runtime/CoreUObject/Private/UObject/GarbageCollection.cpp b/Engine/Source/Runtime/CoreUObject/Private/UObject/GarbageCollection.cpp
--- a/Engine/Source/Runtime/CoreUObject/Private/UObject/GarbageCollection.cpp
+++ b/Engine/Source/Runtime/CoreUObject/Private/UObject/GarbageCollection.cpp
@@ -4522,6 +4522,7 @@ private:
 		if (!GReachabilityState.IsSuspended())
 		{
 			Context = Pool.AllocateFromPool();
+			Context->InitialNativeReferences = GetInitialReferences(Options);
 		}
 		else
 		{
@@ -4533,15 +4534,6 @@ private:
 #endif
 		}
 
-#if WITH_VERSE_VM || defined(__INTELLISENSE__)
-		// Process native structs first as this may modify GReachableObjects used below
-		if (!Private::GReachableNativeStructs.IsEmpty())
-		{
-			V_DIE_UNLESS(GIsFrankenGCCollecting);
-			Private::GReachableNativeStructs.PopAllAndEmpty(InitialNativeStructs);
-		}
-#endif
-
 		if (!Private::GReachableObjects.IsEmpty())
 		{
 			// Add objects marked with the GC barrier to the initial set of objects for the next iteration of incremental reachability
@@ -4550,11 +4542,6 @@ private:
 			UE_LOG(LogGarbage, Verbose, TEXT("Adding %d object(s) marker by GC barrier to the list of objects to process"), InitialObjects.Num());
 			ConditionallyAddBarrierReferencesToHistory(*Context);
 		}
-		else if (GReachabilityState.GetNumIterations() == 0 || (Stats.bFoundGarbageRef && !GReachabilityState.IsSuspended()))
-		{
-			Context->InitialNativeReferences = GetInitialReferences(Options);
-		}
-
 		if (!Private::GReachableClusters.IsEmpty())
 		{
 			// Process cluster roots that were marked as reachable by the GC barrier
@@ -4566,6 +4553,14 @@ private:
 				MarkReferencedClustersAsReachable<EGCOptions::None>(ObjectItem->GetClusterIndex(), InitialObjects);
 			}
 		}
+#if WITH_VERSE_VM || defined(__INTELLISENSE__)
+		if (!Private::GReachableNativeStructs.IsEmpty())
+		{
+			// Add VNativeStructs marked by the Verse GC or the GC barrier
+			V_DIE_UNLESS(GIsFrankenGCCollecting);
+			Private::GReachableNativeStructs.PopAllAndEmpty(InitialNativeStructs);
+		}
+#endif
 
 		Context->SetInitialObjectsUnpadded(InitialObjects);
 #if WITH_VERSE_VM || defined(__INTELLISENSE__)
