[UObjectGlobals]
- Object replacement and reinstanciation delegates both need to be thread-safe since its called from ALT

diff --git a/Engine/Source/Runtime/CoreUObject/Private/UObject/UObjectGlobals.cpp b/Engine/Source/Runtime/CoreUObject/Private/UObject/UObjectGlobals.cpp
--- a/Engine/Source/Runtime/CoreUObject/Private/UObject/UObjectGlobals.cpp
+++ b/Engine/Source/Runtime/CoreUObject/Private/UObject/UObjectGlobals.cpp
@@ -200,7 +200,7 @@ FCoreUObjectDelegates::FOnObjectModified FCoreUObjectDelegates::OnObjectModified
 FCoreUObjectDelegates::FOnObjectRenamed FCoreUObjectDelegates::OnObjectRenamed;
 FCoreUObjectDelegates::FOnObjectTransacted FCoreUObjectDelegates::OnObjectTransacted;
 FCoreUObjectDelegates::FOnObjectsReplaced FCoreUObjectDelegates::OnObjectsReplaced;
-FCoreUObjectDelegates::FOnObjectsReplaced FCoreUObjectDelegates::OnObjectsReinstanced;
+FCoreUObjectDelegates::FOnObjectsReinstanced FCoreUObjectDelegates::OnObjectsReinstanced;
 FCoreUObjectDelegates::FOnVerseDeadReferences FCoreUObjectDelegates::OnVerseDeadObjectReferences;
 FCoreUObjectDelegates::FOnObjectPostCDOCompiled FCoreUObjectDelegates::OnObjectPostCDOCompiled;
 FCoreUObjectDelegates::FOnAssetLoaded FCoreUObjectDelegates::OnAssetLoaded;
diff --git a/Engine/Source/Runtime/CoreUObject/Public/UObject/UObjectGlobals.h b/Engine/Source/Runtime/CoreUObject/Public/UObject/UObjectGlobals.h
--- a/Engine/Source/Runtime/CoreUObject/Public/UObject/UObjectGlobals.h
+++ b/Engine/Source/Runtime/CoreUObject/Public/UObject/UObjectGlobals.h
@@ -3307,7 +3307,7 @@ struct FCoreUObjectDelegates
 	 * objects are fixed up in other objects (i.e. other objects can still be pointing to old data)
 	 */
 	using FReplacementObjectMap = TMap<UObject*, UObject*>; // Alias for use in the macro
-	DECLARE_MULTICAST_DELEGATE_OneParam(FOnObjectsReplaced, const FReplacementObjectMap&);
+	DECLARE_TS_MULTICAST_DELEGATE_OneParam(FOnObjectsReplaced, const FReplacementObjectMap&);
 	static COREUOBJECT_API FOnObjectsReplaced OnObjectsReplaced;
 
 	/**
@@ -3315,7 +3315,7 @@ struct FCoreUObjectDelegates
 	 * Note that this is called after references to replacement objects are fixed up in other objects (i.e. all object
 	 * references should be self-consistent).
 	 */
-	DECLARE_MULTICAST_DELEGATE_OneParam(FOnObjectsReinstanced, const FReplacementObjectMap&);
+	DECLARE_TS_MULTICAST_DELEGATE_OneParam(FOnObjectsReinstanced, const FReplacementObjectMap&);
 	static COREUOBJECT_API FOnObjectsReinstanced OnObjectsReinstanced;
 
 	/**
