[ZenLoader]
- Move ClearImportedPackages to be cleaned on the async loading thread
- This fixes a race condition and reduce game thread work

diff --git a/Engine/Source/Runtime/CoreUObject/Private/Serialization/AsyncLoading2.cpp b/Engine/Source/Runtime/CoreUObject/Private/Serialization/AsyncLoading2.cpp
--- a/Engine/Source/Runtime/CoreUObject/Private/Serialization/AsyncLoading2.cpp
+++ b/Engine/Source/Runtime/CoreUObject/Private/Serialization/AsyncLoading2.cpp
@@ -1673,8 +1673,8 @@ public:
 private:
 #endif
 
-	// Process all the deferred deletion for packages that are finished loading
-	void FlushDeferredDeletePackagesQueue();
+	// Process all the deferred cleanups for packages that are finished loading
+	void FlushDeferredCleanupPackagesQueue();
 public:
 	FGlobalImportStore(FAsyncLoadingThread2& InAsyncLoadingThread)
 		: AsyncLoadingThread(InAsyncLoadingThread)
@@ -1730,7 +1730,7 @@ public:
 	{
 		LLM_SCOPE_BYNAME(TEXT("AsyncLoadPackageStore"));
 
-		FlushDeferredDeletePackagesQueue();
+		FlushDeferredCleanupPackagesQueue();
 
 		FWriteScopeLock ScopeLock(PackagesLock);
 		FLoadedPackageRef& PackageRef = Packages.FindOrAdd(PackageId);
@@ -2969,6 +2969,7 @@ enum class EAsyncPackageLoadingState2 : uint8
 	PostLoadInstances,
 	CreateClusters,
 	Complete,
+	DeferredClearImportedPackages,
 	DeferredDelete,
 };
 
@@ -2995,6 +2996,7 @@ const TCHAR* LexToString(EAsyncPackageLoadingState2 AsyncPackageLoadingState)
 	case EAsyncPackageLoadingState2::PostLoadInstances: return TEXT("PostLoadInstances");
 	case EAsyncPackageLoadingState2::CreateClusters: return TEXT("CreateClusters");
 	case EAsyncPackageLoadingState2::Complete: return TEXT("Complete");
+	case EAsyncPackageLoadingState2::DeferredClearImportedPackages: return TEXT("DeferredClearImportedPackages");
 	case EAsyncPackageLoadingState2::DeferredDelete: return TEXT("DeferredDelete");
 	default:
 	checkNoEntry();
@@ -4188,10 +4190,10 @@ private:
 	TArray<UObject*> EditorLoadedAssets;
 	TArray<UPackage*> EditorCompletedUPackages;
 #endif
-	/** [ASYNC/GAME THREAD] Packages to be deleted from async thread */
-	TMpscQueue<FAsyncPackage2*> DeferredDeletePackages;
-	/** Used to protect DeferredDeletePackages when dequeuing since its a single consumer queue. */
-	UE::FMutex DeferredDeletePackagesLock;
+	/** [ASYNC/GAME THREAD] Packages to be cleaned up from async thread */
+	TMpscQueue<FAsyncPackage2*> DeferredCleanupPackages;
+	/** Used to protect DeferredCleanupPackages when dequeuing since its a single consumer queue. */
+	UE::FMutex DeferredCleanupPackagesLock;
 
 	struct FCompletedPackageRequest
 	{
@@ -4792,19 +4794,28 @@ private:
 #endif
 
 public:
-	bool ProcessDeferredDeletePackagesQueue(int32 MaxCount = MAX_int32)
+	bool ProcessDeferredCleanupPackagesQueue(int32 MaxCount = MAX_int32)
 	{
 		bool bDidSomething = false;
-		if (!DeferredDeletePackages.IsEmpty())
+		if (!DeferredCleanupPackages.IsEmpty())
 		{
-			UE::TUniqueLock ScopeLock(DeferredDeletePackagesLock);
-			TRACE_CPUPROFILER_EVENT_SCOPE(ProcessDeferredDeletePackagesQueue);
+			UE::TUniqueLock ScopeLock(DeferredCleanupPackagesLock);
+			TRACE_CPUPROFILER_EVENT_SCOPE(ProcessDeferredCleanupPackagesQueue);
 
 			FAsyncPackage2* Package = nullptr;
 			int32 Count = 0;
-			while (++Count <= MaxCount && DeferredDeletePackages.Dequeue(Package))
+			while (++Count <= MaxCount && DeferredCleanupPackages.Dequeue(Package))
 			{
-				DeleteAsyncPackage(Package);
+				if (Package->AsyncPackageLoadingState == EAsyncPackageLoadingState2::DeferredClearImportedPackages)
+				{
+					Package->AsyncPackageLoadingState = EAsyncPackageLoadingState2::DeferredDelete;
+					Package->ClearImportedPackages();
+				}
+				else
+				{
+					check(Package->AsyncPackageLoadingState == EAsyncPackageLoadingState2::DeferredDelete);
+					DeleteAsyncPackage(Package);
+				}
 				bDidSomething = true;
 			}
 		}
@@ -9934,7 +9945,7 @@ EAsyncPackageState::Type FAsyncLoadingThread2::ProcessAsyncLoadingFromGameThread
 				break;
 			}
 
-			if (ProcessDeferredDeletePackagesQueue(1))
+			if (ProcessDeferredCleanupPackagesQueue(1))
 			{
 				bDidSomething = true;
 				break;
@@ -10132,10 +10143,9 @@ EAsyncPackageState::Type FAsyncLoadingThread2::ProcessLoadedPackagesFromGameThre
 			UE_ASYNC_PACKAGE_LOG(Verbose, Package->Desc, TEXT("GameThread: LoadCompleted"),
 				TEXT("All loading of package is done, and the async package and load request will be deleted."));
 
-
 			check(Package->AsyncPackageLoadingState == EAsyncPackageLoadingState2::Complete);
-			Package->AsyncPackageLoadingState = EAsyncPackageLoadingState2::DeferredDelete;
-			Package->ClearImportedPackages();
+			Package->AsyncPackageLoadingState = EAsyncPackageLoadingState2::DeferredClearImportedPackages;
+			DeferredCleanupPackages.Enqueue(Package);
 
 			if (Package->CompletionCallbacks.IsEmpty() && Package->ProgressCallbacks.IsEmpty())
 			{
@@ -10198,10 +10208,6 @@ EAsyncPackageState::Type FAsyncLoadingThread2::ProcessLoadedPackagesFromGameThre
 		}
 	}
 
-	if (Result == EAsyncPackageState::Complete)
-	{
-	}
-
 	return Result;
 }
 
@@ -10577,7 +10583,7 @@ uint32 FAsyncLoadingThread2::Run()
 						}
 					}
 
-					if (ProcessDeferredDeletePackagesQueue(100))
+					if (ProcessDeferredCleanupPackagesQueue(100))
 					{
 						continue;
 					}
@@ -10830,8 +10836,8 @@ void FAsyncLoadingThread2::CollectUnreachableObjects(
 void FAsyncLoadingThread2::OnGarbageCollectStarted()
 {
 	TRACE_CPUPROFILER_EVENT_SCOPE(FAsyncLoadingThread2::OnGarbageCollectStarted);
-	// Flush the delete queue so that we don't prevent packages from being garbage collected if we're done with them
-	ProcessDeferredDeletePackagesQueue();
+	// Flush the cleanup queue so that we don't prevent packages from being garbage collected if we're done with them
+	ProcessDeferredCleanupPackagesQueue();
 }
 
 void FAsyncLoadingThread2::RemoveUnreachableObjects(FUnreachableObjects& ObjectsToRemove)
@@ -11150,7 +11156,7 @@ void FAsyncPackage2::Release()
 	if (OldRefCount == 1)
 	{
 		FAsyncLoadingThread2& AsyncLoadingThreadLocal = AsyncLoadingThread;
-		AsyncLoadingThreadLocal.DeferredDeletePackages.Enqueue(this);
+		AsyncLoadingThreadLocal.DeferredCleanupPackages.Enqueue(this);
 		AsyncLoadingThreadLocal.AltZenaphore.NotifyOne();
 	}
 }
@@ -12092,9 +12098,9 @@ EAsyncPackageState::Type FAsyncLoadingThread2::ProcessLoadingUntilCompleteFromGa
 	return bLoadingComplete ? EAsyncPackageState::Complete : EAsyncPackageState::TimeOut;
 }
 
-void FGlobalImportStore::FlushDeferredDeletePackagesQueue()
+void FGlobalImportStore::FlushDeferredCleanupPackagesQueue()
 {
-	AsyncLoadingThread.ProcessDeferredDeletePackagesQueue();
+	AsyncLoadingThread.ProcessDeferredCleanupPackagesQueue();
 }
 
 void FAsyncLoadingThread2::FLoaderStats::Clear()
