Add more info to IsUnlocked assertion when detaching

diff --git a/Engine/Source/Runtime/CoreUObject/Private/Serialization/BulkData.cpp b/Engine/Source/Runtime/CoreUObject/Private/Serialization/BulkData.cpp
--- a/Engine/Source/Runtime/CoreUObject/Private/Serialization/BulkData.cpp
+++ b/Engine/Source/Runtime/CoreUObject/Private/Serialization/BulkData.cpp
@@ -1422,7 +1422,15 @@ void FBulkData::DetachFromArchive( FArchive* Ar, bool bEnsureBulkDataIsLoaded )
 {
 	check( Ar );
 	check( Ar == AttachedAr || AttachedAr == nullptr || AttachedAr->IsProxyOf(Ar) );
+#if DO_CHECK
+	if (!IsUnlocked())
+	{
+		FLinker* FromLinker = Ar->GetLinker();
+		checkf(false, TEXT("Bulk data cannot be detached while still locked. BulkData: %s, Linker: %s, Archive: %s"), 
+			*GetDebugName(), FromLinker ? *FromLinker->GetDebugName() : TEXT("nullptr"), *Ar->GetArchiveName() );
+	}
 	check( IsUnlocked() );
+#endif
 
 	// Make sure bulk data is loaded.
 	if( bEnsureBulkDataIsLoaded )
