Allow unknown enum tracking to reliably FindOriginalType

#rnx

diff --git a/Engine/Source/Runtime/CoreUObject/Private/UObject/EnumProperty.cpp b/Engine/Source/Runtime/CoreUObject/Private/UObject/EnumProperty.cpp
--- a/Engine/Source/Runtime/CoreUObject/Private/UObject/EnumProperty.cpp
+++ b/Engine/Source/Runtime/CoreUObject/Private/UObject/EnumProperty.cpp
@@ -302,6 +302,7 @@ void FEnumProperty::SerializeItem(FStructuredArchive::FSlot Slot, void* Value, v
 			FUObjectSerializeContext* SerializeContext = FUObjectThreadContext::Get().GetSerializeContext();
 			if (UNLIKELY(SerializeContext->bTrackUnknownEnumNames))
 			{
+				TGuardValue<bool> ScopedImpersonateProperties(SerializeContext->bImpersonateProperties, true);
 				UE::FUnknownEnumNames(SerializeContext->SerializedObject).Add(Enum, UE::FindOriginalType(this), EnumValueName);
 			}
 		#endif
diff --git a/Engine/Source/Runtime/CoreUObject/Private/UObject/PropertyByte.cpp b/Engine/Source/Runtime/CoreUObject/Private/UObject/PropertyByte.cpp
--- a/Engine/Source/Runtime/CoreUObject/Private/UObject/PropertyByte.cpp
+++ b/Engine/Source/Runtime/CoreUObject/Private/UObject/PropertyByte.cpp
@@ -77,6 +77,7 @@ void FByteProperty::SerializeItem( FStructuredArchive::FSlot Slot, void* Value,
 			FUObjectSerializeContext* SerializeContext = FUObjectThreadContext::Get().GetSerializeContext();
 			if (UNLIKELY(SerializeContext->bTrackUnknownEnumNames))
 			{
+				TGuardValue<bool> ScopedImpersonateProperties(SerializeContext->bImpersonateProperties, true);
 				UE::FUnknownEnumNames(SerializeContext->SerializedObject).Add(Enum, UE::FindOriginalType(this), EnumValueName);
 			}
 		#endif
