[ZenLoader]
- Fix deadlock if taking the mutex recursively

diff --git a/Engine/Source/Runtime/CoreUObject/Private/Serialization/AsyncLoading2.cpp b/Engine/Source/Runtime/CoreUObject/Private/Serialization/AsyncLoading2.cpp
--- a/Engine/Source/Runtime/CoreUObject/Private/Serialization/AsyncLoading2.cpp
+++ b/Engine/Source/Runtime/CoreUObject/Private/Serialization/AsyncLoading2.cpp
@@ -78,6 +78,7 @@
 #include "Templates/GuardValueAccessors.h"
 #include "Async/ParallelFor.h"
 #include "Async/ManualResetEvent.h"
+#include "Async/RecursiveMutex.h"
 #include "HAL/LowLevelMemStats.h"
 #include "HAL/IPlatformFileOpenLogWrapper.h"
 #include "Modules/ModuleManager.h"
@@ -3770,7 +3771,7 @@ private:
 	/** Whether the constructed objects have been moved to the second loading phase. */
 	TAsyncAtomic<bool> ObjectsNowInPhase2 = false;
 	/** Protects against the same package being serialized from different threads. */
-	UE::FMutex                  SerializationMutex;
+	UE::FRecursiveMutex         SerializationMutex;
 
 	struct FAllDependenciesState
 	{
