Add source location to FBulkData::Lock

diff --git a/Engine/Source/Runtime/CoreUObject/Private/Serialization/BulkData.cpp b/Engine/Source/Runtime/CoreUObject/Private/Serialization/BulkData.cpp
--- a/Engine/Source/Runtime/CoreUObject/Private/Serialization/BulkData.cpp
+++ b/Engine/Source/Runtime/CoreUObject/Private/Serialization/BulkData.cpp
@@ -8,6 +8,7 @@
 #include "Math/GuardedInt.h"
 #include "Misc/ConfigCacheIni.h"
 #include "Misc/Optional.h"
+#include "Misc/SourceLocation.h"
 #include "ProfilingDebugging/LoadTimeTracker.h"
 #include "Serialization/LargeMemoryWriter.h"
 #include "Serialization/LargeMemoryReader.h"
@@ -911,12 +912,16 @@ void FBulkData::GetCopy( void** Dest, bool bDiscardInternalCopy )
  *
  * @param	LockFlags	Flags determining lock behavior
  */
-void* FBulkData::Lock( uint32 LockFlags )
+void* FBulkData::Lock(uint32 LockFlags, UE::FSourceLocation Location)
 {
 	check(IsUnlocked());
 	
 	// Make sure bulk data is loaded.
 	MakeSureBulkDataIsLoaded();
+
+#if WITH_EDITOR && DO_CHECK
+	LockSourceLocation.Emplace(Location);
+#endif
 		
 	// Read-write operations are allowed on returned memory.
 	if( LockFlags & LOCK_READ_WRITE )
@@ -948,7 +953,7 @@ void* FBulkData::Lock( uint32 LockFlags )
 	}
 }
 
-const void* FBulkData::LockReadOnly() const
+const void* FBulkData::LockReadOnly(UE::FSourceLocation Location) const
 {
 	check(IsUnlocked());
 	
@@ -959,6 +964,9 @@ const void* FBulkData::LockReadOnly() const
 
 	// Only read operations are allowed on returned memory.
 	mutable_this->BulkMeta.SetLockStatus(LOCKSTATUS_ReadOnlyLock);
+#if WITH_EDITOR && DO_CHECK
+	mutable_this->LockSourceLocation.Emplace(Location);
+#endif
 
 	return GetDataBufferReadOnly();
 }
@@ -991,6 +999,9 @@ void FBulkData::Unlock() const
 	FBulkData* mutable_this = const_cast<FBulkData*>(this);
 
 	mutable_this->BulkMeta.SetLockStatus(LOCKSTATUS_Unlocked);
+#if WITH_EDITOR && DO_CHECK
+	mutable_this->LockSourceLocation.Reset();
+#endif
 
 	// Free pointer if we're guaranteed to only to access the data once.
 	if (BulkMeta.HasAnyFlags(BULKDATA_SingleUse))
@@ -1426,11 +1437,19 @@ void FBulkData::DetachFromArchive( FArchive* Ar, bool bEnsureBulkDataIsLoaded )
 	if (!IsUnlocked())
 	{
 		FLinker* FromLinker = Ar->GetLinker();
-		checkf(false, TEXT("Bulk data cannot be detached while still locked. BulkData: %s, Linker: %s, Archive: %s"), 
-			*GetDebugName(), FromLinker ? *FromLinker->GetDebugName() : TEXT("nullptr"), *Ar->GetArchiveName() );
-	}
-	check( IsUnlocked() );
-#endif
+#if WITH_EDITOR
+		UE::FSourceLocation* Location = LockSourceLocation.GetPtrOrNull();
+#else
+		UE::FSourceLocation* Location = nullptr;
+#endif // WITH_EDITOR
+		checkf(false, TEXT("Bulk data cannot be detached while still locked. BulkData: %s, Linker: %s, Archive: %s. LockLocation: %s %s(%d)"), 
+			*GetDebugName(), FromLinker ? *FromLinker->GetDebugName() : TEXT("nullptr"), *Ar->GetArchiveName(),
+			Location ? StringCast<TCHAR>(Location->GetFunctionName()).Get() : TEXT("null"),
+			Location ? StringCast<TCHAR>(Location->GetFileName()).Get() : TEXT("null"),
+			Location ? Location->GetLine() : 0
+		);
+	}
+#endif // DO_CHECK
 
 	// Make sure bulk data is loaded.
 	if( bEnsureBulkDataIsLoaded )
diff --git a/Engine/Source/Runtime/CoreUObject/Public/Serialization/BulkData.h b/Engine/Source/Runtime/CoreUObject/Public/Serialization/BulkData.h
--- a/Engine/Source/Runtime/CoreUObject/Public/Serialization/BulkData.h
+++ b/Engine/Source/Runtime/CoreUObject/Public/Serialization/BulkData.h
@@ -25,6 +25,7 @@
 #include "Misc/OptionalFwd.h"
 #include "Misc/PackagePath.h"
 #include "Misc/PackageSegment.h"
+#include "Misc/SourceLocation.h"
 #include "Serialization/Archive.h"
 #include "Serialization/BulkDataCookedIndex.h"
 #include "Serialization/CustomVersion.h"
@@ -764,13 +765,13 @@ public:
 	 *
 	 * @param	LockFlags	Flags determining lock behavior
 	 */
-	COREUOBJECT_API void* Lock( uint32 LockFlags );
+	COREUOBJECT_API void* Lock(uint32 LockFlags, UE::FSourceLocation Location = UE::FSourceLocation::Current());
 
 	/**
 	 * Locks the bulk data and returns a read-only pointer to it.
 	 * This variant can be called on a const bulkdata
 	 */
-	COREUOBJECT_API const void* LockReadOnly() const;
+	COREUOBJECT_API const void* LockReadOnly(UE::FSourceLocation Location = UE::FSourceLocation::Current()) const;
 
 	/**
 	 * Change size of locked bulk data. Only valid if locked via read-write lock.
@@ -1018,6 +1019,9 @@ protected:
 	FLinkerLoad*			Linker = nullptr;
 	/** The index to which the cooked version of the bulkdata payload should be assigned*/
 	FBulkDataCookedIndex	CookedIndex;
+#if DO_CHECK
+	TOptional<UE::FSourceLocation> LockSourceLocation;
+#endif
 #endif // WITH_EDITOR
 	   //
 #if !USE_RUNTIME_BULKDATA
