Pass target object to GetRootPatih in GetPathToStaleObjectReferencer to avoid printing out the same path for every leaked object.

diff --git a/Engine/Source/Runtime/CoreUObject/Private/UObject/ReferenceChainSearch.cpp b/Engine/Source/Runtime/CoreUObject/Private/UObject/ReferenceChainSearch.cpp
--- a/Engine/Source/Runtime/CoreUObject/Private/UObject/ReferenceChainSearch.cpp
+++ b/Engine/Source/Runtime/CoreUObject/Private/UObject/ReferenceChainSearch.cpp
@@ -1896,7 +1896,7 @@ static FString GetPathToStaleObjectReferencer(UObject* ObjectToFindReferencesTo,
 				Params.Out->Logf(ELogVerbosity::Log, TEXT("%s    ^ This reference is preventing the old %s from being GC'd ^"), FCString::Spc(Params.Indent), *ObjectToFindReferencesTo->GetClass()->GetName());
 			}
 			return true;
-		});
+		}, ObjectToFindReferencesTo);
 	return PathToCulprit;
 }
 
