Degrating an assert to a log warning when changing a serial number of a weak object pointer after resolving a remote object

diff --git a/Engine/Source/Runtime/CoreUObject/Private/UObject/WeakObjectPtr.cpp b/Engine/Source/Runtime/CoreUObject/Private/UObject/WeakObjectPtr.cpp
--- a/Engine/Source/Runtime/CoreUObject/Private/UObject/WeakObjectPtr.cpp
+++ b/Engine/Source/Runtime/CoreUObject/Private/UObject/WeakObjectPtr.cpp
@@ -273,7 +273,9 @@ FUObjectItem* FWeakObjectPtr::TryResolveRemoteObject() const
 		}
 		if (ObjectSerialNumber != ObjectItem->GetSerialNumber())
 		{
-			checkf(ObjectSerialNumber == 0, TEXT("Attempting to change existing and valid serial number %d to %d when resolving remote object %s (%s)"),
+			// ObjectSerialNumber doesn't play an important role when remote object support is enabled because the remote id is a globally unique
+			// identifier of an object but we keep it as an optimization (no need to resolve the remote object if the serial numbers match).
+			UE_CLOG(ObjectSerialNumber != 0, LogRemoteObject, Warning, TEXT("Attempting to change existing and valid serial number %d to %d when resolving remote object %s (%s)"),
 				ObjectSerialNumber, ObjectItem->GetSerialNumber(), *ResolvedObject->GetPathName(), *RemoteId.ToString());
 			FPlatformAtomics::AtomicStore(&const_cast<FWeakObjectPtr*>(this)->ObjectSerialNumber, ObjectItem->GetSerialNumber());
 		}
