[ZenLoader]
- Fix deadlock on failed runtime packages

diff --git a/Engine/Source/Runtime/CoreUObject/Private/Serialization/AsyncLoading2.cpp b/Engine/Source/Runtime/CoreUObject/Private/Serialization/AsyncLoading2.cpp
--- a/Engine/Source/Runtime/CoreUObject/Private/Serialization/AsyncLoading2.cpp
+++ b/Engine/Source/Runtime/CoreUObject/Private/Serialization/AsyncLoading2.cpp
@@ -9586,6 +9586,11 @@ EEventLoadNodeExecutionResult FAsyncPackage2::Event_PostLoad(FAsyncLoadingThread
 		// End async loading, simulates EndLoad
 		Package->EndAsyncLoad();
 	}
+	else
+	{
+		// On failure, just move to the next step
+		Package->ProcessedExportBundlesCount = Package->Data.TotalExportBundleCount;
+	}
 
 	if (LoadingState == EEventLoadNodeExecutionResult::Timeout)
 	{
@@ -9711,6 +9716,11 @@ EEventLoadNodeExecutionResult FAsyncPackage2::Event_DeferredPostLoad(FAsyncLoadi
 			++Package->ProcessedExportBundlesCount;
 		}
 	}
+	else
+	{
+		// On failure, just move to the next step
+		Package->ProcessedExportBundlesCount = Package->Data.TotalExportBundleCount;
+	}
 
 	if (LoadingState == EEventLoadNodeExecutionResult::Timeout)
 	{
