Remove deprecated ChunkIDs property from FAssetData.

- A small improvement to memory usage given the amount of these structures in a large project.

diff --git a/Engine/Source/Runtime/CoreUObject/Private/AssetRegistry/AssetData.cpp b/Engine/Source/Runtime/CoreUObject/Private/AssetRegistry/AssetData.cpp
--- a/Engine/Source/Runtime/CoreUObject/Private/AssetRegistry/AssetData.cpp
+++ b/Engine/Source/Runtime/CoreUObject/Private/AssetRegistry/AssetData.cpp
@@ -592,16 +592,6 @@ UClass* FAssetData::GetClass(EResolveClass ResolveClass) const
 
 FAssetData::FChunkArrayView FAssetData::GetChunkIDs() const
 {
-#if !UE_STRIP_DEPRECATED_PROPERTIES
-	PRAGMA_DISABLE_DEPRECATION_WARNINGS
-	// Modifying the chunk IDs directly is no longer supported; use AddChunkID/SetChunkIDs/ClearChunkIDs instead
-	if (!ChunkIDs.IsEmpty())
-	{
-		UE_LOG(LogAssetData, Error, TEXT("Modifying FAssetData::ChunkIDs directly is no longer supported; use AddChunkID/SetChunkIDs/ClearChunkIDs instead."));
-	}
-	PRAGMA_ENABLE_DEPRECATION_WARNINGS
-#endif
-
 	return UE::AssetRegistry::GChunkArrayRegistry.GetChunkIDs(ChunkArrayRegistryHandle);
 }
 
diff --git a/Engine/Source/Runtime/CoreUObject/Public/AssetRegistry/AssetData.h b/Engine/Source/Runtime/CoreUObject/Public/AssetRegistry/AssetData.h
--- a/Engine/Source/Runtime/CoreUObject/Public/AssetRegistry/AssetData.h
+++ b/Engine/Source/Runtime/CoreUObject/Public/AssetRegistry/AssetData.h
@@ -229,12 +229,6 @@ public:
 	 */
 	TSharedPtr<FAssetBundleData, ESPMode::ThreadSafe> TaggedAssetBundles;
 
-	/** The IDs of the pakchunks this asset is located in for streaming install.  Empty if not assigned to a chunk */
-#if !UE_STRIP_DEPRECATED_PROPERTIES
-	UE_DEPRECATED(5.1, "Use SetChunkIDs/GetChunkIDs/AddChunkID instead.")
-	TArray<int32, TInlineAllocator<2>> ChunkIDs;
-#endif
-
 	COREUOBJECT_API void SetTagsAndAssetBundles(FAssetDataTagMap&& Tags);
 
 	// These are usually very small arrays and we can preallocate two elements for the same cost as one on 64-bit systems
