Fix checkfSlow for not reallocating enum names/values

diff --git a/Engine/Source/Runtime/CoreUObject/Private/UObject/Enum.cpp b/Engine/Source/Runtime/CoreUObject/Private/UObject/Enum.cpp
--- a/Engine/Source/Runtime/CoreUObject/Private/UObject/Enum.cpp
+++ b/Engine/Source/Runtime/CoreUObject/Private/UObject/Enum.cpp
@@ -51,14 +51,14 @@ void UEnum::FNameData::DeleteValues()
 
 FName* UEnum::FNameData::AllocateNames(int32 InNum)
 {
-	checkfSlow(TaggedNames == 0 || ((TaggedNames & MaskPointer) == MaskPointer), TEXT("AllocateNames should only be called when names are unallocated or statically allocated"));
+	checkfSlow(TaggedNames == 0 || ((TaggedNames & MaskPointer) != TagDynamic), TEXT("AllocateNames should only be called when names are unallocated or statically allocated"));
 	FName* TmpNames = (FName*)FMemory::Malloc(sizeof(FName) * InNum);
 	TaggedNames = (UPTRINT)TmpNames | TagDynamic;
 	return TmpNames;
 }
 int64* UEnum::FNameData::AllocateValues(int32 InNum)
 {
-	checkfSlow(TaggedValues == 0 || ((TaggedValues & MaskPointer) == MaskPointer), TEXT("AllocateValues should only be called when values are unallocated or statically allocated"));
+	checkfSlow(TaggedValues == 0 || ((TaggedValues & MaskPointer) != TagDynamic), TEXT("AllocateValues should only be called when values are unallocated or statically allocated"));
 	int64* Values = (int64*)FMemory::Malloc(sizeof(int64) * InNum);
 	TaggedValues = (UPTRINT)Values | TagDynamic;
 	return Values;
