Silence V1078 warnings caused by a PVS parse bug with Swap
#rnx

diff --git a/Engine/Plugins/Animation/ACLPlugin/Source/ACLPlugin/Private/EditorDatabaseMonitor.cpp b/Engine/Plugins/Animation/ACLPlugin/Source/ACLPlugin/Private/EditorDatabaseMonitor.cpp
--- a/Engine/Plugins/Animation/ACLPlugin/Source/ACLPlugin/Private/EditorDatabaseMonitor.cpp
+++ b/Engine/Plugins/Animation/ACLPlugin/Source/ACLPlugin/Private/EditorDatabaseMonitor.cpp
@@ -26,7 +26,7 @@ namespace EditorDatabaseMonitor
 		}
 
 		// Iterate over our dirty databases and refresh any stale mappings we might have
-		for (TWeakObjectPtr<UAnimationCompressionLibraryDatabase>& DatabasePtr : DirtyDatabasesTmp)
+		for (TWeakObjectPtr<UAnimationCompressionLibraryDatabase>& DatabasePtr : DirtyDatabasesTmp) //-V1078 7.39 Swap Bugsamp
 		{
 			if (UAnimationCompressionLibraryDatabase* Database = DatabasePtr.Get())
 			{
diff --git a/Engine/Plugins/Animation/ControlRig/Source/ControlRig/Public/Tracks/SampleTrack.h b/Engine/Plugins/Animation/ControlRig/Source/ControlRig/Public/Tracks/SampleTrack.h
--- a/Engine/Plugins/Animation/ControlRig/Source/ControlRig/Public/Tracks/SampleTrack.h
+++ b/Engine/Plugins/Animation/ControlRig/Source/ControlRig/Public/Tracks/SampleTrack.h
@@ -538,7 +538,7 @@ public:
 		UpdateArrayViews();
 
 		// add the values back
-		for(const T& Value : CopyOfValues)
+		for(const T& Value : CopyOfValues) //-V1078 7.39 Swap Bug
 		{
 			AddSample(Value, InTolerance);
 		}
diff --git a/Engine/Plugins/FX/Niagara/Source/Niagara/Private/NiagaraWorldManager.cpp b/Engine/Plugins/FX/Niagara/Source/Niagara/Private/NiagaraWorldManager.cpp
--- a/Engine/Plugins/FX/Niagara/Source/Niagara/Private/NiagaraWorldManager.cpp
+++ b/Engine/Plugins/FX/Niagara/Source/Niagara/Private/NiagaraWorldManager.cpp
@@ -1412,7 +1412,7 @@ void FNiagaraWorldManager::PostActorTick(float DeltaSeconds)
 		// Execute spawning
 		TArray<FNiagaraSystemSimulationRef> LocalSimulationsWithPostActorWork;
 		Swap(SimulationsWithPostActorWork, LocalSimulationsWithPostActorWork);
-		for (FNiagaraSystemSimulationRef& Simulation : LocalSimulationsWithPostActorWork)
+		for (FNiagaraSystemSimulationRef& Simulation : LocalSimulationsWithPostActorWork) //-V1078 7.39 Swap Bug
 		{
 			if (Simulation->IsValid())
 			{
diff --git a/Engine/Plugins/Media/ElectraUtil/Source/ElectraHTTPStream/Private/Generic/GenericElectraHTTPStream.cpp b/Engine/Plugins/Media/ElectraUtil/Source/ElectraHTTPStream/Private/Generic/GenericElectraHTTPStream.cpp
--- a/Engine/Plugins/Media/ElectraUtil/Source/ElectraHTTPStream/Private/Generic/GenericElectraHTTPStream.cpp
+++ b/Engine/Plugins/Media/ElectraUtil/Source/ElectraHTTPStream/Private/Generic/GenericElectraHTTPStream.cpp
@@ -753,7 +753,7 @@ void FElectraHTTPStreamGeneric::SetupNewRequests()
 	TArray<TSharedPtr<FElectraHTTPStreamRequestGeneric, ESPMode::ThreadSafe>> NewReqs;
 	Swap(NewReqs, NewRequests);
 	RequestLock.Unlock();
-	for(auto &Request : NewReqs)
+	for(auto &Request : NewReqs) //-V1078 7.39 Swap Bug
 	{
 		if (Request->Setup())
 		{
@@ -824,7 +824,7 @@ void FElectraHTTPStreamGeneric::HandleCompletedRequests()
 	{
 		TArray<TSharedPtr<FElectraHTTPStreamRequestGeneric, ESPMode::ThreadSafe>> TempRequests;
 		Swap(CompletedRequests, TempRequests);
-		for(auto &Finished : TempRequests)
+		for(auto &Finished : TempRequests) //-V1078 7.39 Swap Bug
 		{
 			Finished->SetFinished();
 			Finished->Close();
diff --git a/Engine/Plugins/Runtime/RigVM/Source/RigVM/Public/RigVMCore/RigVMNativized.h b/Engine/Plugins/Runtime/RigVM/Source/RigVM/Public/RigVMCore/RigVMNativized.h
--- a/Engine/Plugins/Runtime/RigVM/Source/RigVM/Public/RigVMCore/RigVMNativized.h
+++ b/Engine/Plugins/Runtime/RigVM/Source/RigVM/Public/RigVMCore/RigVMNativized.h
@@ -295,7 +295,7 @@ protected:
 		TArray<T> Temp;
 		Swap(InOutA, Temp);
 
-		for(const T& Element : Temp)
+		for(const T& Element : Temp) //-V1078 7.39 Swap Bug
 		{
 			if(!InB.Contains(Element))
 			{
@@ -317,7 +317,7 @@ protected:
 		TArray<T> Temp;
 		Swap(InOutA, Temp);
 
-		for(const T& Element : Temp)
+		for(const T& Element : Temp) //-V1078 7.39 Swap Bug
 		{
 			if(InB.Contains(Element))
 			{
diff --git a/Engine/Plugins/Runtime/RigVM/Source/RigVMDeveloper/Private/RigVMAsset.cpp b/Engine/Plugins/Runtime/RigVM/Source/RigVMDeveloper/Private/RigVMAsset.cpp
--- a/Engine/Plugins/Runtime/RigVM/Source/RigVMDeveloper/Private/RigVMAsset.cpp
+++ b/Engine/Plugins/Runtime/RigVM/Source/RigVMDeveloper/Private/RigVMAsset.cpp
@@ -192,7 +192,7 @@ IRigVMAssetInterface::IRigVMAssetInterface(const FObjectInitializer& ObjectIniti
 		Swap(QueuedCompilerMessageDelegates, DelegatesForReportFromCompiler);
 	}
 
-	for(const FOnRigVMReportCompilerMessage::FDelegate& Delegate : DelegatesForReportFromCompiler)
+	for(const FOnRigVMReportCompilerMessage::FDelegate& Delegate : DelegatesForReportFromCompiler) //-V1078 7.39 Swap Bug
 	{
 		ReportCompilerMessageEvent.Add(Delegate);
 	}
diff --git a/Engine/Plugins/Runtime/RigVM/Source/RigVMDeveloper/Private/RigVMModel/RigVMController.cpp b/Engine/Plugins/Runtime/RigVM/Source/RigVMDeveloper/Private/RigVMModel/RigVMController.cpp
--- a/Engine/Plugins/Runtime/RigVM/Source/RigVMDeveloper/Private/RigVMModel/RigVMController.cpp
+++ b/Engine/Plugins/Runtime/RigVM/Source/RigVMDeveloper/Private/RigVMModel/RigVMController.cpp
@@ -22998,7 +22998,7 @@ void URigVMController::ProcessDetachedLinks(const FRestoreLinkedPathSettings& In
 	Swap(GetGraph()->DetachedLinks, DetachedLinks);
 
 	// destroy the links finally
-	for(URigVMLink* DetachedLink : DetachedLinks)
+	for(URigVMLink* DetachedLink : DetachedLinks) //-V1078 7.39 Swap Bug
 	{
 		ReportRemovedLink(DetachedLink->GetSourcePinPath(), DetachedLink->GetTargetPinPath());
 		Notify(ERigVMGraphNotifType::LinkRemoved, DetachedLink);
diff --git a/Engine/Source/Developer/Windows/LiveCoding/Private/LiveCodingModule.cpp b/Engine/Source/Developer/Windows/LiveCoding/Private/LiveCodingModule.cpp
--- a/Engine/Source/Developer/Windows/LiveCoding/Private/LiveCodingModule.cpp
+++ b/Engine/Source/Developer/Windows/LiveCoding/Private/LiveCodingModule.cpp
@@ -1350,7 +1350,7 @@ void FLiveCodingModule::UpdateModules(bool bAllowStarting)
 
 		TArray<FString> EnableModules;
 		TArray<FString> LazyLoadModules;
-		for (const ModuleChange& Change : Changes)
+		for (const ModuleChange& Change : Changes) //-V1078 7.39 Swap Bug
 		{
 			if (Change.bLoaded)
 			{
diff --git a/Engine/Source/Editor/Sequencer/Private/SequencerKeyParams.cpp b/Engine/Source/Editor/Sequencer/Private/SequencerKeyParams.cpp
--- a/Engine/Source/Editor/Sequencer/Private/SequencerKeyParams.cpp
+++ b/Engine/Source/Editor/Sequencer/Private/SequencerKeyParams.cpp
@@ -145,7 +145,7 @@ void FKeyOperation::FSectionCandidates::FilterOperations(UMovieSceneTrack* Track
 		Swap(OldOperations, Operations);
 
 		TArray<UMovieSceneSection*, TInlineAllocator<1>> Sections;
-		for (const FKeySectionOperation& Operation : OldOperations)
+		for (const FKeySectionOperation& Operation : OldOperations) //-V1078 7.39 Swap Bug
 		{
 			UMovieSceneSection* Section = Operation.Section->GetSectionObject();
 			Sections.Add(Section);
diff --git a/Engine/Source/Editor/UnrealEd/Private/Cooker/CookWorkerClient.cpp b/Engine/Source/Editor/UnrealEd/Private/Cooker/CookWorkerClient.cpp
--- a/Engine/Source/Editor/UnrealEd/Private/Cooker/CookWorkerClient.cpp
+++ b/Engine/Source/Editor/UnrealEd/Private/Cooker/CookWorkerClient.cpp
@@ -699,7 +699,7 @@ void FCookWorkerClient::SendPendingMessages()
 		FScopeLock FileTransferScopeLock(&FileTransferLock);
 		Swap(LocalPendingFileTransfers, PendingFileTransfers);
 	}
-	for (FFileTransferMessage& FileTransferMessage : LocalPendingFileTransfers)
+	for (FFileTransferMessage& FileTransferMessage : LocalPendingFileTransfers) //-V1078 7.39 Swap Bug
 	{
 		SendMessage(FileTransferMessage);
 	}
diff --git a/Engine/Source/Runtime/Core/Private/IO/IoDispatcher.cpp b/Engine/Source/Runtime/Core/Private/IO/IoDispatcher.cpp
--- a/Engine/Source/Runtime/Core/Private/IO/IoDispatcher.cpp
+++ b/Engine/Source/Runtime/Core/Private/IO/IoDispatcher.cpp
@@ -852,7 +852,7 @@ private:
 				Swap(LocalRequestsToCancel, RequestsToCancel);
 				Swap(LocalRequestsToReprioritize, RequestsToReprioritize);
 			}
-			for (FIoRequestImpl* RequestToCancel : LocalRequestsToCancel)
+			for (FIoRequestImpl* RequestToCancel : LocalRequestsToCancel) //-V1078 7.39 Swap Bug
 			{
 				if (!RequestToCancel->IsCancelled())
 				{
@@ -864,7 +864,7 @@ private:
 				}
 				RequestToCancel->ReleaseRef();
 			}
-			for (FIoRequestImpl* RequestToRePrioritize : LocalRequestsToReprioritize)
+			for (FIoRequestImpl* RequestToRePrioritize : LocalRequestsToReprioritize) //-V1078 7.39 Swap Bug
 			{
 				if (RequestToRePrioritize->Backend)
 				{
diff --git a/Engine/Source/Runtime/CoreUObject/Private/Serialization/AsyncLoading.cpp b/Engine/Source/Runtime/CoreUObject/Private/Serialization/AsyncLoading.cpp
--- a/Engine/Source/Runtime/CoreUObject/Private/Serialization/AsyncLoading.cpp
+++ b/Engine/Source/Runtime/CoreUObject/Private/Serialization/AsyncLoading.cpp
@@ -1567,7 +1567,7 @@ struct FPrecacheCallbackHandler
 			Swap(LocalIncoming, Incoming);
 			Swap(LocalIncomingSummaries, IncomingSummaries);
 		}
-		for (IAsyncReadRequest * Req : LocalIncoming)
+		for (IAsyncReadRequest * Req : LocalIncoming) //-V1078 7.39 Swap Bug
 		{
 			check(Req);
 			FWeakAsyncPackagePtr Found = WaitingPackages.FindAndRemoveChecked(Req);
@@ -1575,7 +1575,7 @@ struct FPrecacheCallbackHandler
 			check(Pkg);
 			UnprocessedMemUsed += Pkg->PrecacheRequestReady(Req);
 		}
-		for (FWeakAsyncPackagePtr& Sum : LocalIncomingSummaries)
+		for (FWeakAsyncPackagePtr& Sum : LocalIncomingSummaries) //-V1078 7.39 Swap Bug
 		{
 			FAsyncLoadingThread* LocalAsyncLoadingThread = &FAsyncLoadingThread::Get();
 			LocalAsyncLoadingThread->QueueEvent_FinishLinker(Sum, FAsyncLoadEvent::EventSystemPriority_MAX);
diff --git a/Engine/Source/Runtime/CoreUObject/Private/Serialization/AsyncLoading2.cpp b/Engine/Source/Runtime/CoreUObject/Private/Serialization/AsyncLoading2.cpp
--- a/Engine/Source/Runtime/CoreUObject/Private/Serialization/AsyncLoading2.cpp
+++ b/Engine/Source/Runtime/CoreUObject/Private/Serialization/AsyncLoading2.cpp
@@ -11540,7 +11540,7 @@ void FAsyncLoadingThread2::ConditionalProcessEditorCallbacks()
 
 		// Call the global delegate for package endloads and set the bHasBeenLoaded flag that is used to
 		// check which packages have reached this state
-		for (UPackage* CompletedUPackage : LocalEditorCompletedUPackages)
+		for (UPackage* CompletedUPackage : LocalEditorCompletedUPackages) //-V1078 7.39 Swap Bug
 		{
 			CompletedUPackage->SetHasBeenEndLoaded(true);
 		}
@@ -11548,7 +11548,7 @@ void FAsyncLoadingThread2::ConditionalProcessEditorCallbacks()
 			FEndLoadPackageContext{ LocalEditorCompletedUPackages, 0, false /* bSynchronous */ });
 
 		// In editor builds, call the asset load callback. This happens in both editor and standalone to match EndLoad
-		for (UObject* LoadedObject : LocalEditorLoadedAssets)
+		for (UObject* LoadedObject : LocalEditorLoadedAssets) //-V1078 7.39 Swap Bug
 		{
 			if (LoadedObject)
 			{
diff --git a/Engine/Source/Runtime/Engine/Private/EdGraph/EdGraphPin.cpp b/Engine/Source/Runtime/Engine/Private/EdGraph/EdGraphPin.cpp
--- a/Engine/Source/Runtime/Engine/Private/EdGraph/EdGraphPin.cpp
+++ b/Engine/Source/Runtime/Engine/Private/EdGraph/EdGraphPin.cpp
@@ -58,7 +58,7 @@ public:
 			Swap(LocalPinsToDelete, PinsToDelete);
 		}
 
-		for (UEdGraphPin* Pin : LocalPinsToDelete)
+		for (UEdGraphPin* Pin : LocalPinsToDelete) //-V1078 7.39 Swap Bug
 		{
 			delete Pin;
 		}
diff --git a/Engine/Source/Runtime/Experimental/Chaos/Private/Chaos/Island/IslandManager.cpp b/Engine/Source/Runtime/Experimental/Chaos/Private/Chaos/Island/IslandManager.cpp
--- a/Engine/Source/Runtime/Experimental/Chaos/Private/Chaos/Island/IslandManager.cpp
+++ b/Engine/Source/Runtime/Experimental/Chaos/Private/Chaos/Island/IslandManager.cpp
@@ -2127,7 +2127,7 @@ namespace Chaos::Private
 		// rebuilding the islands below (see AddNodeToIsland)
 		const bool bCheckSleep = Island->Flags.bCheckSleep;
 
-		for (FPBDIslandParticle* RootNode : IslandNodes)
+		for (FPBDIslandParticle* RootNode : IslandNodes) //-V1078 7.39 Swap Bug
 		{
 			// If we haven't seen this node yet and it is dynamic, start forming a new island
 			if ((RootNode->VisitEpoch != VisitEpoch) && RootNode->Flags.bIsDynamic)
diff --git a/Engine/Source/Runtime/Landscape/Private/LandscapeSplines.cpp b/Engine/Source/Runtime/Landscape/Private/LandscapeSplines.cpp
--- a/Engine/Source/Runtime/Landscape/Private/LandscapeSplines.cpp
+++ b/Engine/Source/Runtime/Landscape/Private/LandscapeSplines.cpp
@@ -3166,7 +3166,7 @@ void ULandscapeSplineSegment::UpdateSplinePoints(bool bUpdateCollision, bool bUp
 	TMap<ULandscapeSplinesComponent*, TArray<USplineMeshComponent*>> ForeignMeshComponentsMap = GetForeignMeshComponents();
 	
 	// Unregister components, Remove Foreign/Local Associations
-	for (TObjectPtr<USplineMeshComponent>& LocalMeshComponent : OldLocalMeshComponents)
+	for (TObjectPtr<USplineMeshComponent>& LocalMeshComponent : OldLocalMeshComponents) //-V1078 7.39 Swap Bug
 	{
 		USplineMeshComponent* Comp = LocalMeshComponent.Get();
 		checkSlow(OuterSplines->MeshComponentLocalOwnersMap.FindRef(Comp) == this);
diff --git a/Engine/Source/Runtime/MovieScene/Private/EntitySystem/MovieSceneEntitySystemRunner.cpp b/Engine/Source/Runtime/MovieScene/Private/EntitySystem/MovieSceneEntitySystemRunner.cpp
--- a/Engine/Source/Runtime/MovieScene/Private/EntitySystem/MovieSceneEntitySystemRunner.cpp
+++ b/Engine/Source/Runtime/MovieScene/Private/EntitySystem/MovieSceneEntitySystemRunner.cpp
@@ -1327,7 +1327,7 @@ void FMovieSceneEntitySystemRunner::GameThread_PostEvaluationPhase(UMovieSceneEn
 	{
 		FMovieSceneEntitySystemEvaluationReentrancyWindow Window(this, Linker);
 
-		for (const FQueuedUpdateParams& UpdateParams : TmpCurrentInstances)
+		for (const FQueuedUpdateParams& UpdateParams : TmpCurrentInstances) //-V1078 7.39 Swap Bug
 		{
 			// We must check for validity here because the cache handles may have become invalid
 			// during this iteration (since there is a re-entrancy window open)
@@ -1354,7 +1354,7 @@ void FMovieSceneEntitySystemRunner::GameThread_PostEvaluationPhase(UMovieSceneEn
 			}
 		}
 
-		for (const FSimpleDelegate& OnFlushed : TmpOnFlushedDelegates)
+		for (const FSimpleDelegate& OnFlushed : TmpOnFlushedDelegates) //-V1078 7.39 Swap Bug
 		{
 			OnFlushed.Execute();
 		}
diff --git a/Engine/Source/Runtime/MovieScene/Private/MovieScene.cpp b/Engine/Source/Runtime/MovieScene/Private/MovieScene.cpp
--- a/Engine/Source/Runtime/MovieScene/Private/MovieScene.cpp
+++ b/Engine/Source/Runtime/MovieScene/Private/MovieScene.cpp
@@ -692,7 +692,7 @@ void UMovieScene::EmptyRootFolders()
 
 	RootFolders.Empty();
 
-	for (UMovieSceneFolder* Folder : OldFolders)
+	for (UMovieSceneFolder* Folder : OldFolders) //-V1078 7.39 Swap Bug
 	{
 		EventHandlers.Trigger(&UE::MovieScene::ISequenceDataEventHandler::OnRootFolderRemoved, Folder);
 	}
diff --git a/Engine/Source/Runtime/MovieSceneCapture/Private/ActiveMovieSceneCaptures.cpp b/Engine/Source/Runtime/MovieSceneCapture/Private/ActiveMovieSceneCaptures.cpp
--- a/Engine/Source/Runtime/MovieSceneCapture/Private/ActiveMovieSceneCaptures.cpp
+++ b/Engine/Source/Runtime/MovieSceneCapture/Private/ActiveMovieSceneCaptures.cpp
@@ -45,7 +45,7 @@ void FActiveMovieSceneCaptures::Shutdown()
 	TArray<TObjectPtr<UMovieSceneCapture>> ActiveCapturesCopy;
 	Swap(ActiveCaptures, ActiveCapturesCopy);
 
-	for (auto& Obj : ActiveCapturesCopy)
+	for (auto& Obj : ActiveCapturesCopy) //-V1078 7.39 Swap Bug
 	{
 		Obj->Finalize();
 	}
diff --git a/Engine/Source/Runtime/MovieSceneCapture/Private/FrameGrabberProtocol.cpp b/Engine/Source/Runtime/MovieSceneCapture/Private/FrameGrabberProtocol.cpp
--- a/Engine/Source/Runtime/MovieSceneCapture/Private/FrameGrabberProtocol.cpp
+++ b/Engine/Source/Runtime/MovieSceneCapture/Private/FrameGrabberProtocol.cpp
@@ -37,7 +37,7 @@ void UFrameGrabberProtocol::TickImpl()
 	{
 		TArray<FCapturedFrameData> CapturedFrames = FrameGrabber->GetCapturedFrames();
 
-		for (FCapturedFrameData& Frame : CapturedFrames)
+		for (FCapturedFrameData& Frame : CapturedFrames) //-V1078 7.39 Swap Bug
 		{
 			ProcessFrame(MoveTemp(Frame));
 		}
diff --git a/Engine/Source/Runtime/MovieSceneCapture/Private/UserDefinedCaptureProtocol.cpp b/Engine/Source/Runtime/MovieSceneCapture/Private/UserDefinedCaptureProtocol.cpp
--- a/Engine/Source/Runtime/MovieSceneCapture/Private/UserDefinedCaptureProtocol.cpp
+++ b/Engine/Source/Runtime/MovieSceneCapture/Private/UserDefinedCaptureProtocol.cpp
@@ -115,7 +115,7 @@ void UUserDefinedCaptureProtocol::TickImpl()
 	{
 		TArray<FCapturedFrameData> CapturedFrames = FinalPixelsFrameGrabber->GetCapturedFrames();
 
-		for (FCapturedFrameData& Frame : CapturedFrames)
+		for (FCapturedFrameData& Frame : CapturedFrames) //-V1078 7.39 Swap Bug
 		{
 			// Steal the frame and make it shareable
 			FCapturedPixels CapturedPixels { MakeShared<TImagePixelData<FColor>, ESPMode::ThreadSafe>( Frame.BufferSize, TArray64<FColor>(MoveTemp(Frame.ColorBuffer) ) ) };
diff --git a/Engine/Source/Runtime/UMG/Private/Animation/UMGSequenceTickManager.cpp b/Engine/Source/Runtime/UMG/Private/Animation/UMGSequenceTickManager.cpp
--- a/Engine/Source/Runtime/UMG/Private/Animation/UMGSequenceTickManager.cpp
+++ b/Engine/Source/Runtime/UMG/Private/Animation/UMGSequenceTickManager.cpp
@@ -178,7 +178,7 @@ void UUMGSequenceTickManager::TickWidgetAnimations(float DeltaSeconds)
 		TArray<TObjectPtr<UUserWidget>> NewUserWidgets;
 		Swap(PendingUserWidgets, NewUserWidgets);
 
-		for (TObjectPtr<UUserWidget> NewWidget : NewUserWidgets)
+		for (TObjectPtr<UUserWidget> NewWidget : NewUserWidgets) //-V1078 7.39 Swap Bug
 		{
 			FSequenceTickManagerWidgetData NewWidgetData;
 
