Add property name/path to disregard for GC warnings

diff --git a/Engine/Source/Runtime/CoreUObject/Private/UObject/GarbageCollectionVerification.cpp b/Engine/Source/Runtime/CoreUObject/Private/UObject/GarbageCollectionVerification.cpp
--- a/Engine/Source/Runtime/CoreUObject/Private/UObject/GarbageCollectionVerification.cpp
+++ b/Engine/Source/Runtime/CoreUObject/Private/UObject/GarbageCollectionVerification.cpp
@@ -77,7 +77,7 @@ public:
 	{
 		return NumErrors.GetValue();
 	}
-	FORCEINLINE_DEBUGGABLE void HandleTokenStreamObjectReference(FGCArrayStruct& ObjectsToSerializeStruct, UObject* ReferencingObject, UObject*& Object, FMemberId MemberId, EOrigin Origin, bool bAllowReferenceElimination)
+	FORCEINLINE_DEBUGGABLE void HandleTokenStreamObjectReference(FWorkerContext& ObjectsToSerializeStruct, UObject* ReferencingObject, UObject*& Object, FMemberId MemberId, EOrigin Origin, bool bAllowReferenceElimination)
 	{
 		if (Object)
 		{
@@ -112,9 +112,10 @@ public:
 					  GUObjectArray.ObjectToObjectItem(Object)->GetOwnerIndex() > 0 ||
 					  GUObjectArray.ObjectToObjectItem(Object)->HasAnyFlags(EInternalObjectFlags::ClusterRoot)))
 			{
-				UE_LOG(LogGarbage, Warning, TEXT("Disregard for GC object %s referencing %s which is not part of root set"),
-					*ReferencingObject->GetFullName(),
-					*Object->GetFullName());
+				UE_LOG(LogGarbage, Warning, TEXT("Disregard for GC object %s referencing %s which is not part of root set. Property: %s"),
+					*GetFullNameSafe(ReferencingObject),
+					*Object->GetFullName(),
+					ObjectsToSerializeStruct.SchemaStack ? *ObjectsToSerializeStruct.SchemaStack->ToString() : TEXT("unknown"));
 				NumErrors.Increment();
 			}
 		}
