Add a standalone program to quickly print a .uasset's object graph and tagged properties to json without linking native code or loading imported assets
New functionality is gated behind UE_INTERNAL_UNIVERSAL_PLACEHOLDERS to limit scope of change

diff --git a/Engine/Source/Programs/UnrealAssetStringify/Private/UnrealAssetStringify.cpp b/Engine/Source/Programs/UnrealAssetStringify/Private/UnrealAssetStringify.cpp
new file mode 100644
--- /dev/null
+++ b/Engine/Source/Programs/UnrealAssetStringify/Private/UnrealAssetStringify.cpp
@@ -0,0 +1,327 @@
+// Copyright Epic Games, Inc. All Rights Reserved.
+
+#include "JsonObjectGraph/Stringify.h"
+#include "Logging/StructuredLog.h"
+#include "UObject/UObjectGlobals.h"
+#include "UObject/UObjectThreadContext.h"
+
+#include "RequiredProgramMainCPPInclude.h"
+
+DEFINE_LOG_CATEGORY_STATIC(LogUnrealAssetStringify, Log, All);
+
+IMPLEMENT_APPLICATION(UnrealAssetStringify, "UnrealAssetStringify");
+
+static void PrintUsage()
+{
+const TCHAR* const Usage = TEXT(R"(Usage:
+	UnrealAssetStringify [OPTION]... FILE
+
+Options:
+	--save	Save the stringified output to Engine\Programs\UnrealAssetStringify\Saved
+			instead of simply printing it
+
+Example:
+	UnrealAssetStringify.exe "D:\Path\To\A.uasset"
+	UnrealAssetStringify.exe --save "D:\Path\To\A.uasset")");
+	UE_LOGFMT(LogUnrealAssetStringify, Display, "{0}", Usage);
+}
+
+namespace UE::Private
+{
+	bool TryInferPackageName(FStringView Filename, FString& OutResult);
+
+	enum class EReturnCodes : int32
+	{
+		Success,
+		BadCommandlineArgs,
+		MissingFile,
+		BadAsset,
+		BadPackageName,
+		CouldNotCreateResultFile,
+		CouldNotWriteResultFile
+	};
+}
+
+// Returns a best guess package name given a filename that may or may not be
+// in an actual project folder. Doesn't bother looking for conventional
+// project structure as that would be slow. This means that assets in /Game/
+// will be somewhat mangled, but that seems fine to me. Most packages
+// now record their location in the project summary:
+bool UE::Private::TryInferPackageName(FStringView Filename, FString& OutResult)
+{
+	OutResult = FPaths::ChangeExtension(FString(Filename), TEXT(""));
+	FPaths::NormalizeFilename(OutResult);
+	int32 ColonIndex;
+	if (OutResult.FindChar(TEXT(':'), ColonIndex))
+	{
+		OutResult = OutResult.RightChop(ColonIndex + 1);
+	}
+	// if the filename has '/Content/' in it then assume the preceding folder
+	// is the project name, and the sub folder follows the content folder:
+	const FString ContentToken = TEXT("/Content/");
+	int32 ContentIndex = OutResult.Find(ContentToken, ESearchCase::IgnoreCase, ESearchDir::FromStart);
+	if (ContentIndex != INDEX_NONE)
+	{
+		// Extract the "root folder" before /Content/
+		FString RootFolder = OutResult.Left(ContentIndex);
+		FPaths::NormalizeDirectoryName(RootFolder);
+
+		// Take the subpath *after* /Content/
+		FString SubPath = OutResult.Mid(ContentIndex + ContentToken.Len());
+
+		// Build virtual package root (e.g. /ProjectName/), aka shorten the root:
+		FString RootName = FPaths::GetCleanFilename(RootFolder);
+		// Recombine the root and subpath:
+		FString PackagePath = FString::Printf(TEXT("/%s/%s"), *RootName, *SubPath);
+
+		// Assign result and validate
+		OutResult = MoveTemp(PackagePath);
+		FPackageName::RegisterMountPoint(OutResult, OutResult);
+		check(FPackageName::IsValidLongPackageName(OutResult, true));
+		return true;
+	}
+	else
+	{
+		// If the filename does not have a /Content/ folder then the best
+		// we can do is drop the drive letter:
+		FString ErrorMessage;
+		if(FPackageName::TryConvertFilenameToLongPackageName(OutResult, OutResult, &ErrorMessage))
+		{
+			FPackageName::RegisterMountPoint(OutResult, OutResult);
+			check(FPackageName::IsValidLongPackageName(OutResult, true));
+			return true;
+		}
+		UE_LOGFMT(LogUnrealAssetStringify, Error, "Failed to infer package name: {0}", ErrorMessage);
+	}
+	return false;
+}
+
+struct FUnrealAssetStringifyArgs
+{
+	// options:
+	bool bSave = false;
+	// required inputs:
+	FString Filename;
+	
+	static bool ParseArgs(int32 ArgC, TCHAR* const ArgV[], FUnrealAssetStringifyArgs& OutArgs);
+	static int32 StringifyFile(FStringView Filename, const FUnrealAssetStringifyArgs& InArgs);
+};
+
+bool FUnrealAssetStringifyArgs::ParseArgs(int32 ArgC, TCHAR* const ArgV[], FUnrealAssetStringifyArgs& OutArgs)
+{
+	int32 I = 1;
+	for(I; I < ArgC - 1; ++I)
+	{
+		if(FStringView(ArgV[I]) == TEXT("--save"))
+		{
+			OutArgs.bSave = true;
+		}
+	}
+
+	if(I < ArgC && !FStringView(ArgV[I]).StartsWith(TEXT("--")))
+	{
+		OutArgs.Filename = ArgV[I];
+	}
+	return !OutArgs.Filename.IsEmpty();
+}
+
+int32 FUnrealAssetStringifyArgs::StringifyFile(FStringView Filename, const FUnrealAssetStringifyArgs& InArgs)
+{
+	FPackagePath PackagePath = FPackagePath::FromLocalPath(Filename);
+	UPackage* Package = CreatePackage(*PackagePath.GetPackageName());
+	check(Package);
+	{
+		if(!IFileManager::Get().FileExists(*FString(Filename)))
+		{
+			UE_LOGFMT(LogUnrealAssetStringify, Error, "Could not find file: \"{0}\"", Filename);
+			return (int32)UE::Private::EReturnCodes::MissingFile;
+		}
+
+		TRefCountPtr<FUObjectSerializeContext> LoadContext(FUObjectThreadContext::Get().GetSerializeContext());
+		BeginLoad(LoadContext);
+		FLinkerLoad* Linker = FLinkerLoad::CreateLinker(LoadContext, Package, PackagePath, LOAD_NoVerify|LOAD_DisableEngineVersionChecks);
+		if(!Linker)
+		{
+			UE_LOGFMT(LogUnrealAssetStringify, Error, "Could not read file as a uasset: {0}", Filename);
+			return (int32)UE::Private::EReturnCodes::BadAsset;
+		}
+		if((Linker->Summary.GetPackageFlags() & PKG_Cooked) != 0)
+		{
+			// we can't load cooked packages yet because the cooked format is not interpretable
+			// without the native code:
+			UE_LOGFMT(LogUnrealAssetStringify, Warning, "Cannot load {0} because it has been cooked and does not have tagged data", Filename);
+			return (int32)UE::Private::EReturnCodes::Success;
+		}
+		// GetPackageName isn't going to work for us, but we can rename the package to the one
+		// stored in the asset, which will match user expectation:
+		if(Linker->Summary.PackageName != FName())
+		{
+			Package->Rename(*Linker->Summary.PackageName);
+		}
+		else
+		{
+			// The package predates the inclusion of the packagename in the summary
+			// lets try to infer a reasonable package name. It cannot efficiently
+			// determine /Game/, so there will be quirks.
+			FString InferredPackageName;
+			if(UE::Private::TryInferPackageName(Filename, InferredPackageName))
+			{
+				Package->Rename(*InferredPackageName);
+			}
+			else
+			{
+				UE_LOGFMT(LogUnrealAssetStringify, Error, 
+					"Cannot infer package name for {0} consider resaving the asset to add an explicit package name", Filename);
+				return (int32)UE::Private::EReturnCodes::BadPackageName;
+			}
+		}
+		check(Linker);
+		Linker->LoadAllObjects(false);
+		EndLoad(LoadContext);
+	}
+
+	FUtf8String Json = UE::JsonObjectGraph::Stringify(
+		{Package}
+	);
+
+	// output, either to console or to a file:
+	if(!InArgs.bSave)
+	{
+		UE_LOGFMT(LogUnrealAssetStringify, Display, "Asset printed by UnrealAssetStringify\n{0}", Json);
+	}
+	else
+	{
+		// filename to use, _uas is just a fingerprint, short for
+		// UnrealAssetStringify, while .json is a reasonable file extension:
+		const auto GetIntermediateAssetName = [](const UObject* Object)
+		{
+			return
+				FPaths::ProjectSavedDir() +
+				Object->GetPathName() +
+				FString(TEXT("_uas.json"));
+		};
+		FString DestinationFile = GetIntermediateAssetName(Package);
+		TUniquePtr<FArchive> FileArchive(IFileManager::Get().CreateDebugFileWriter(*DestinationFile));
+		if(!FileArchive) // failed to create a file
+		{
+			UE_LOGFMT(LogUnrealAssetStringify, Error, 
+				"Failed to create result file {0}", DestinationFile);
+			return (int32)UE::Private::EReturnCodes::CouldNotCreateResultFile; 
+		}
+
+		FileArchive->Serialize((void*)Json.GetCharArray().GetData(), Json.GetCharArray().Num() - 1);
+		if(FileArchive->IsError()) // error during writing, likely disk space
+		{
+			UE_LOGFMT(LogUnrealAssetStringify, Error, 
+				"Failed to write result file {0}", DestinationFile);
+			return (int32)UE::Private::EReturnCodes::CouldNotWriteResultFile; 
+		}
+		UE_LOGFMT(LogUnrealAssetStringify, Display, 
+			"Wrote result file: {0}", DestinationFile);
+	}
+	return (int32)UE::Private::EReturnCodes::Success;
+}
+
+FString GetIniBootstrapFilename()
+{
+	FCommandLine::Set(TEXT(""));
+	static FString VersionString = TEXT("uas_delme");
+	static bool bVersionStringInialized = false;
+	if(!bVersionStringInialized)
+	{
+		// We use the exe to invalidate cached inis:
+		TUniquePtr<FArchive> FileArchive(IFileManager::Get().CreateFileReader(
+			FPlatformProcess::ExecutablePath()));
+		if (FileArchive)
+		{
+			TArray<uint8> FileData;
+			FileData.SetNum(FileArchive->TotalSize());
+			FileArchive->Serialize(FileData.GetData(), FileData.Num());
+
+			FBlake3 Hash;
+			Hash.Update(FileData.GetData(), FileData.Num());
+			VersionString = LexToString(Hash.Finalize());
+		}
+	}
+	return FString::Printf(TEXT("%s/Bootstrap-UnrealAssetStringify-%s.inis"), *FPaths::ProjectSavedDir(), *VersionString);
+}
+
+FString GetIniBootstrap()
+{
+	// if a cached version of the init exists return that:
+	FString BootstrapFilename = GetIniBootstrapFilename();
+	if(IFileManager::Get().FileExists(*BootstrapFilename))
+	{
+		return FString::Printf(TEXT(" -IniBootstrap=\"%s\""), *BootstrapFilename);
+	}
+	return "";
+}
+
+INT32_MAIN_INT32_ARGC_TCHAR_ARGV()
+{
+	FTaskTagScope Scope(ETaskTag::EGameThread);
+	ON_SCOPE_EXIT
+	{ 
+		FIoDispatcher::Shutdown(); 
+		LLM(FLowLevelMemTracker::Get().UpdateStatsPerFrame());
+		RequestEngineExit(TEXT("Exiting"));
+		FEngineLoop::AppPreExit();
+		FModuleManager::Get().UnloadModulesAtShutdown();
+		FEngineLoop::AppExit();
+	};
+
+	// Ini boot strap is a fairly critical optimization for standalone programs
+	// without this large multiprocess workloads will be much slower, which 
+	// inhibits testing:
+	FString IniBootStrap = GetIniBootstrap();
+
+	// silence some logs - LogEnum support could be improved,
+	// NoPreviewPlatforms is a config optimization
+	FString InitCommandLine = 
+		TEXT(" -NODEFAULTLOG -LogCmds=\"LogConfig warning, LogLinker error, LogStreaming error, LogUObjectGlobals error, LogEnum error, LogExit error\" -NoPreviewPlatforms") + 
+		IniBootStrap;
+	FCommandLine::Set(*InitCommandLine);
+	FLogSuppressionInterface::Get().ProcessConfigAndCommandLine();
+	// disable localization display string support, it is expensive to initialize:
+	IConsoleManager::Get().ProcessUserConsoleInput(TEXT("Localization.DisplayStringSupport 2"), *GLog, nullptr);
+	IConsoleManager::Get().ProcessUserConsoleInput(TEXT("IDO.Placeholder.Feature.ReplaceMissingTypeImportsOnLoad 1"), *GLog, nullptr);
+	IConsoleManager::Get().ProcessUserConsoleInput(TEXT("IDO.Placeholder.Feature.SerializeExportReferencesOnLoad 1"), *GLog, nullptr);
+	IConsoleManager::Get().ProcessUserConsoleInput(TEXT("IDO.Placeholder.Feature.ReplaceMissingReinstancedTypes 1"), *GLog, nullptr);
+	IConsoleManager::Get().ProcessUserConsoleInput(TEXT("IDO.Placeholder.Feature.ReplaceDeadClassInstanceTypes 1"), *GLog, nullptr);
+	GAllowUnversionedContentInEditor = true;
+	GAllowCookedDataInEditorBuilds = true;
+	FScopedAllowAbstractClassAllocation AllowAbstract; // mostly for UMulticastDelegateProperty
+
+	if (int32 Ret = GEngineLoop.PreInit(*InitCommandLine))
+	{
+		UE_LOGFMT(LogUnrealAssetStringify, Error, "Failed to PreInit");
+		return Ret;
+	}
+
+	if(IniBootStrap.IsEmpty())
+	{
+		GConfig->SaveCurrentStateForBootstrap(*GetIniBootstrapFilename());
+	}
+
+	FUnrealAssetStringifyArgs Args;
+	if(!FUnrealAssetStringifyArgs::ParseArgs(ArgC, ArgV, Args))
+	{
+		PrintUsage();
+		UE_LOGFMT(LogUnrealAssetStringify, Error, "Failed to parse arguments");
+		return (int32)UE::Private::EReturnCodes::BadCommandlineArgs;
+	}
+
+	// These are initialized lazily, but with no other modules loaded 
+	// they won't be initialized at all
+	FDelegateProperty::StaticClass();
+	FSoftClassProperty::StaticClass();
+	FInt8Property::StaticClass();
+	FInt16Property::StaticClass();
+	FLazyObjectProperty::StaticClass();
+	FMulticastDelegateProperty::StaticClass();
+	FMulticastInlineDelegateProperty::StaticClass();
+	FMulticastSparseDelegateProperty::StaticClass();
+	FFieldPathProperty::StaticClass();
+	
+	return FUnrealAssetStringifyArgs::StringifyFile(Args.Filename, Args);
+}
diff --git a/Engine/Source/Programs/UnrealAssetStringify/UnrealAssetStringify.Build.cs b/Engine/Source/Programs/UnrealAssetStringify/UnrealAssetStringify.Build.cs
new file mode 100644
--- /dev/null
+++ b/Engine/Source/Programs/UnrealAssetStringify/UnrealAssetStringify.Build.cs
@@ -0,0 +1,21 @@
+// Copyright Epic Games, Inc. All Rights Reserved.
+
+using UnrealBuildTool;
+
+public class UnrealAssetStringify : ModuleRules
+{
+	public UnrealAssetStringify(ReadOnlyTargetRules Target) : base(Target)
+	{
+		PublicIncludePathModuleNames.Add("Launch");
+		PrivateDependencyModuleNames.AddRange(
+			new string[] {
+				"Core",
+				"CoreUObject",
+				"Projects", // required for project manager..
+				"JsonObjectGraph"
+			}
+		);
+		
+		bEnableExceptions = true;
+	}
+}
diff --git a/Engine/Source/Programs/UnrealAssetStringify/UnrealAssetStringify.Target.cs b/Engine/Source/Programs/UnrealAssetStringify/UnrealAssetStringify.Target.cs
new file mode 100644
--- /dev/null
+++ b/Engine/Source/Programs/UnrealAssetStringify/UnrealAssetStringify.Target.cs
@@ -0,0 +1,41 @@
+// Copyright Epic Games, Inc. All Rights Reserved.
+
+using UnrealBuildTool;
+
+[SupportedPlatforms(UnrealPlatformClass.Desktop)]
+public class UnrealAssetStringifyTarget : TargetRules
+{
+	public UnrealAssetStringifyTarget(TargetInfo Target) : base(Target)
+	{
+		Type = TargetType.Program;
+		LinkType = TargetLinkType.Monolithic;
+		IncludeOrderVersion = EngineIncludeOrderVersion.Latest;
+		LaunchModuleName = "UnrealAssetStringify";
+		DefaultBuildSettings = BuildSettingsVersion.Latest;
+
+		// Linking Engine is not viable - would bring in 
+		// a huge amount of code:
+		bCompileAgainstEngine = false;
+		bBuildDeveloperTools = false;
+		bCompileAgainstApplicationCore = false;
+		bBuildWithEditorOnlyData = true;
+		bCompileAgainstEditor = true; // this defines WITH_EDITOR, allowing us to load editor data
+		bCompileAgainstCoreUObject = true;
+
+		bIsBuildingConsoleApplication = true;
+		bEnableTrace = true;
+
+		// We need full package paths.. making these explicit...
+		GlobalDefinitions.Add("UE_SUPPORT_FULL_PACKAGEPATH=1");
+		GlobalDefinitions.Add("ALLOW_OTHER_PLATFORM_CONFIG=1"); // required by IConsoleManager.h when using bCompileAgainstEditor
+		GlobalDefinitions.Add("UE_INTERNAL_UNIVERSAL_PLACEHOLDERS=1");
+		// Disable crashreporter to improve startup time and prevent ini pollution
+		GlobalDefinitions.Add("NOINITCRASHREPORTER=1");
+		// Another optimization, limit the number of threads:
+		GlobalDefinitions.Add("UE_TASKGRAPH_THREAD_LIMIT=5");
+
+		// Late resolve doesn't quite work atm due to CDOs being created
+		// immediately by the linker if there's a Class in the asset:
+		GlobalDefinitions.Add("UE_WITH_OBJECT_HANDLE_LATE_RESOLVE=0");
+	}
+}
diff --git a/Engine/Source/Runtime/Core/Private/Serialization/Archive.cpp b/Engine/Source/Runtime/Core/Private/Serialization/Archive.cpp
--- a/Engine/Source/Runtime/Core/Private/Serialization/Archive.cpp
+++ b/Engine/Source/Runtime/Core/Private/Serialization/Archive.cpp
@@ -54,6 +54,36 @@ FArchive& SerializeByteOrderSwapped(FArchive& Ar, T& Value)
 
 } // namespace ArchiveUtil
 
+#if UE_INTERNAL_UNIVERSAL_PLACEHOLDERS
+FArchiveBoundsScope::FArchiveBoundsScope(FArchive* InOwner, FArchiveBounds InOriginalBounds)
+	: Owner(InOwner)
+	, OriginalBounds(InOriginalBounds)
+{
+	check(Owner);
+}
+
+FArchiveBoundsScope::~FArchiveBoundsScope()
+{
+	check(Owner);
+	Owner->SetBoundsImpl(OriginalBounds);
+}
+
+FArchiveBoundsScope FArchive::SetBounds(int64_t Size)
+{
+	FArchiveBounds PrevBounds = GetBoundsImpl();
+	SetBoundsImpl({ Tell(), Size});
+	return FArchiveBoundsScope(this, PrevBounds);
+}
+
+void FArchive::SetBoundsImpl(FArchiveBounds NewBounds)
+{
+}
+
+FArchiveBounds FArchive::GetBoundsImpl() const
+{
+	return {};
+}
+#endif // UE_INTERNAL_UNIVERSAL_PLACEHOLDERS
 
 /*-----------------------------------------------------------------------------
 	FArchive implementation.
@@ -321,6 +351,13 @@ void FArchiveState::SetCriticalError()
 	ForEachState([](FArchiveState& State) { State.ArIsError = State.ArIsCriticalError = true; });
 }
 
+#if UE_INTERNAL_UNIVERSAL_PLACEHOLDERS
+void FArchiveState::ClearCriticalError()
+{
+	ForEachState([](FArchiveState& State) { State.ArIsError = State.ArIsCriticalError = false; });
+}
+#endif // UE_INTERNAL_UNIVERSAL_PLACEHOLDERS
+
 void FArchiveState::ClearError()
 {
 	ForEachState([](FArchiveState& State) { State.ArIsError = false; });
diff --git a/Engine/Source/Runtime/Core/Public/Misc/Build.h b/Engine/Source/Runtime/Core/Public/Misc/Build.h
--- a/Engine/Source/Runtime/Core/Public/Misc/Build.h
+++ b/Engine/Source/Runtime/Core/Public/Misc/Build.h
@@ -535,3 +535,42 @@
 #ifndef UE_SUPPORT_FOR_ACTOR_TICK_DISABLE
 	#define UE_SUPPORT_FOR_ACTOR_TICK_DISABLE 0
 #endif
+
+// Enables the universal placeholders feature, which allows all imported objects to be replaced with 
+// placeholder objects. This flag is internal and will dispear in the future without warning.
+// It is used to facilitate fast visualization of tagged property .uasset data by satisfying the 
+// following needs:
+// 
+//	1.	Assume no imported packages need to actually be loaded, and can therefore be universally mocked, 
+//		hence 'universal placeholders' - see FLinkerLoad::CreateUniversalPlaceholders
+//	2.	Accomodate missing types with native stream serializers by adding extra bounds checking to serialization
+//		This extra bounds checking would waste critical main thread time in other configurations, but in the 
+//		the long run the workloads can be moved to other threads - see Archive.h, PropertyTag.cpp
+//	3.	Provide safe behavior to placeholders which would need their type to 'emerge' from usage. 
+//		Currently placeholders are mocked as plain objects, but that cannot accomodate native types with 
+//		native invariants - see ObjectPtr.h
+//	4.	Relatedly avoids discarding placeholders that don't conform to emergent non-native type 
+//		constraints see FLinkerLoad::operator<<(FObjectPtr&)
+//	5.	Relatedly avoids hashing of placeholder types that are not hashable - see PropertySet.cpp,
+//		PropertyMap.cpp
+//  6.	Assumes a read operation, avoiding expensive mutators - see PropertyBagRepository.cpp
+//	7.	Handles some unexpected seeking in complicated native asset loads - see Class.cpp
+//	8.	Allows serialization routines to expect unreachable objects (in the long run all loaded 
+//		data should be made reachable) - see JsonStringifyImpl.cpp
+//
+// All of these usages will be accomodated more elegantly by future work. Do not use this define.
+#ifndef UE_INTERNAL_UNIVERSAL_PLACEHOLDERS
+	#define UE_INTERNAL_UNIVERSAL_PLACEHOLDERS 0
+#endif
+
+// Universal placeholders are incompatible with full editor functionality, preventing use in that
+// configuration for safety:
+#if UE_EDITOR && UE_INTERNAL_UNIVERSAL_PLACEHOLDERS
+	#error UE_INTERNAL_UNIVERSAL_PLACEHOLDERS cannot be used in the full UE_EDITOR
+#endif
+
+// Universal placeholders are also incompatible with pure cooked configurations, as universal placeholders
+// rely on tagged serialization, IDOs, placeholders and other editor only FLinkerLoad functionality:
+#if !WITH_EDITOR && UE_INTERNAL_UNIVERSAL_PLACEHOLDERS
+	#error UE_INTERNAL_UNIVERSAL_PLACEHOLDERS cannot be used in cooked configurations
+#endif
diff --git a/Engine/Source/Runtime/Core/Public/Serialization/Archive.h b/Engine/Source/Runtime/Core/Public/Serialization/Archive.h
--- a/Engine/Source/Runtime/Core/Public/Serialization/Archive.h
+++ b/Engine/Source/Runtime/Core/Public/Serialization/Archive.h
@@ -66,6 +66,35 @@ struct FArchiveCookContext;
 //       This will also make TSAN happy when multiple threads are hitting the same bitfield during concurrent saves.
 #define FArchive_Serialize_BitfieldBool(ARCHIVE, BITFIELD_BOOL) { bool TEMP_BITFIELD_BOOL = BITFIELD_BOOL; ARCHIVE << TEMP_BITFIELD_BOOL; if (BITFIELD_BOOL != TEMP_BITFIELD_BOOL) { BITFIELD_BOOL = TEMP_BITFIELD_BOOL; } }
 
+#if UE_INTERNAL_UNIVERSAL_PLACEHOLDERS
+// Named pair to identify bounds, used by some FArchive implementations
+struct FArchiveBounds
+{
+	int64_t BoundsBasis = 0;
+	int64_t BoundsSize = 0;
+};
+
+// Helper for scoping bounds of a read from buffer aka load operation
+// meant to be created via FArchive::SetBounds. This is only used in
+// high placeholder environments but could be appropriate for other
+// unreliable reads as well.
+struct FArchiveBoundsScope
+{
+	CORE_API ~FArchiveBoundsScope();
+private:
+	friend class FArchive;
+
+	FArchiveBoundsScope(FArchive* InOwner, FArchiveBounds Bounds);
+	FArchive* Owner;
+	FArchiveBounds OriginalBounds;
+
+	FArchiveBoundsScope& operator=(const FArchiveBoundsScope&) = delete;
+	FArchiveBoundsScope(const FArchiveBoundsScope&) = delete;
+	FArchiveBoundsScope& operator=(FArchiveBoundsScope&&) = delete;
+	FArchiveBoundsScope(FArchiveBoundsScope&&) = delete;
+};
+#endif // UE_INTERNAL_UNIVERSAL_PLACEHOLDERS
+
 struct FArchiveState
 {
 private:
@@ -118,6 +147,13 @@ public:
 	 * Sets the archiver IsCriticalError and IsError to true. Also sets CriticalError in the proxy archiver if one is wrapping this.
 	 */
 	CORE_API void SetCriticalError();
+	
+#if UE_INTERNAL_UNIVERSAL_PLACEHOLDERS
+	/**
+	 * Sets ArIsError and ArIsCriticalError to false, useful only if recoverable bounds are available
+	 */
+	CORE_API void ClearCriticalError();
+#endif // UE_INTERNAL_UNIVERSAL_PLACEHOLDERS
 
 	/**
 	 * Called to get the computed size from a size-detecting archive after it has finished serializing.
@@ -1741,6 +1777,9 @@ public:
 	using FArchiveState::SetError;
 	using FArchiveState::ClearError;
 	using FArchiveState::SetCriticalError;
+#if UE_INTERNAL_UNIVERSAL_PLACEHOLDERS
+	using FArchiveState::ClearCriticalError;
+#endif // UE_INTERNAL_UNIVERSAL_PLACEHOLDERS
 	using FArchiveState::GetInnermostState;
 	using FArchiveState::CountBytes;
 	using FArchiveState::GetArchiveName;
@@ -1938,6 +1977,18 @@ public:
 	 */
 	virtual void MarkScriptSerializationEnd(const UObject* Obj) { }
 
+#if UE_INTERNAL_UNIVERSAL_PLACEHOLDERS
+	/**
+	 * Update the bounds the archive should be reading from, this is meant to be a number
+	 * of expected bytes from the current archive position (established by Tell()).
+	 * Bounds are expected to be recursively narrowed by the FArchiveBoundsScope, with outer
+	 * bounds restored by inner scopes upon scope exit. This is currently only used in 
+	 * a UE_INTERNAL_UNIVERSAL_PLACEHOLDERS environment and implemented only by
+	 * FLinkerLoad, but it could be generally useful.
+	 */
+	CORE_API [[nodiscard]] FArchiveBoundsScope SetBounds(int64_t Size);
+#endif // UE_INTERNAL_UNIVERSAL_PLACEHOLDERS
+
 	/**
 	 * Called to register a reference to a specific name value, of type TypeObject (UEnum or UStruct normally). Const so it can be called from PostSerialize
 	 */
@@ -1947,6 +1998,14 @@ public:
 
 private:
 	CORE_API void VARARGS LogfImpl(const TCHAR* Fmt, ...);
+	
+#if UE_INTERNAL_UNIVERSAL_PLACEHOLDERS
+	// impl functions for archives that want to support bounds checking
+	// this default implementation does nothing:
+	CORE_API virtual void SetBoundsImpl(FArchiveBounds NewBounds);
+	CORE_API virtual FArchiveBounds GetBoundsImpl() const;
+	friend struct FArchiveBoundsScope;
+#endif // UE_INTERNAL_UNIVERSAL_PLACEHOLDERS
 
 public:
 	// Logf implementation for convenience.
diff --git a/Engine/Source/Runtime/CoreUObject/Private/UObject/Class.cpp b/Engine/Source/Runtime/CoreUObject/Private/UObject/Class.cpp
--- a/Engine/Source/Runtime/CoreUObject/Private/UObject/Class.cpp
+++ b/Engine/Source/Runtime/CoreUObject/Private/UObject/Class.cpp
@@ -1937,6 +1937,12 @@ void UStruct::SerializeVersionedTaggedProperties(FStructuredArchive::FSlot Slot,
 								// Serialize the value if it contains a struct or enum because only those may contain an unknown property.
 								if (bSerializeValue)
 								{
+#if UE_INTERNAL_UNIVERSAL_PLACEHOLDERS
+									// The Link call above may have preloaded an object, which will clear
+									// bIsSerializingScriptProperties and cause the linker to stop reading
+									// tagged properties, force setting it here:
+									UnderlyingArchive.MarkScriptSerializationStart(nullptr);
+#endif // UE_INTERNAL_UNIVERSAL_PLACEHOLDERS
 									UnderlyingArchive.Seek(StartOfProperty);
 									FStructuredArchive::FSlot ValueSlotCopy = PropertyRecord.EnterField(TEXT("Value"));
 									void* TempData = TempProperty->AllocateAndInitializeValue();
diff --git a/Engine/Source/Runtime/CoreUObject/Private/UObject/LinkerLoad.cpp b/Engine/Source/Runtime/CoreUObject/Private/UObject/LinkerLoad.cpp
--- a/Engine/Source/Runtime/CoreUObject/Private/UObject/LinkerLoad.cpp
+++ b/Engine/Source/Runtime/CoreUObject/Private/UObject/LinkerLoad.cpp
@@ -1,6 +1,10 @@
 // Copyright Epic Games, Inc. All Rights Reserved.
 
 #include "UObject/LinkerLoad.h"
+#if UE_INTERNAL_UNIVERSAL_PLACEHOLDERS
+#include "Algo/ForEach.h"
+#include "Algo/TopologicalSort.h"
+#endif// UE_INTERNAL_UNIVERSAL_PLACEHOLDERS
 #include "AssetRegistry/AssetData.h"
 #include "HAL/FileManager.h"
 #include "Misc/Paths.h"
@@ -586,6 +590,10 @@ void FLinkerLoad::SetLoader(FArchive* InLoader, bool bInLoaderNeedsEngineVersion
 	Loader = InLoader;
 	bLoaderNeedsEngineVersionChecks = bInLoaderNeedsEngineVersionChecks;
 
+#if UE_INTERNAL_UNIVERSAL_PLACEHOLDERS
+	CurrentBounds.BoundsSize = InLoader->TotalSize();
+#endif // UE_INTERNAL_UNIVERSAL_PLACEHOLDERS
+
 	check(StructuredArchive == nullptr);
 	check(!StructuredArchiveRootRecord.IsSet());
 
@@ -4436,6 +4444,374 @@ UClass* FLinkerLoad::TryCreatePlaceholderClassImport(int32 ImportIndex)
 	return ClassObject;
 }
 
+#if UE_INTERNAL_UNIVERSAL_PLACEHOLDERS
+void FLinkerLoad::CreateUniversalPlaceholders()
+{
+	// Pass 0: Create necessary packages:
+	Algo::ForEach( ImportMap, [this](FObjectImport& Import)
+	{
+		TryCreateImportedPackage(Import);
+	});
+	// Pass 1: Create placeholder classes for classes that are declared in this 
+	// package:
+	Algo::ForEach( ExportMap, [this](const FObjectExport& Export)
+	{
+		TryCreatePlaceholderClassExport(Export);
+	});
+	// Pass 2: Create placeholder classes for exports that import their types.
+	// If you do this in pass 1 you'll not give classes within the package a chance
+	// to create themselves reasonably...
+	Algo::ForEach( ExportMap, [this](const FObjectExport& Export)
+	{
+		TryCreateImportedClassForExport(Export);
+	});
+	// Pass 3: Create remaining placeholder classes for imports:
+	TryCreateImportedClassesForImports();
+	// The required classes are now mostly created, but more placeholders may
+	// be created as we read properties from the objects (either from 
+	// LoadAllObjects or EndLoad)
+}
+
+void FLinkerLoad::TryCreateImportedPackage(FObjectImport& Import)
+{
+	if( !Import.OuterIndex.IsNull() )
+	{
+		return;
+	}
+	
+	FName ObjectName = InstancingContext.RemapPackage(Import.ObjectName);
+	UPackage* Pkg = CreatePackage(*ObjectName.ToString());
+	Pkg->SetPackageFlags(PKG_InMemoryOnly); // tag as in memory so we don't check for a corresponding FLinker
+	Import.XObject = Pkg;
+}
+
+void FLinkerLoad::TryCreatePlaceholderClassExport(const FObjectExport& Export)
+{
+	if(!Export.ClassIndex.IsExport())
+	{
+		return;
+	}
+	
+	const int32 ExportIndex = Export.ClassIndex.ToExport();
+	FObjectExport& ExportClass = ExportMap[ExportIndex];
+	if (ExportClass.Object)
+	{
+		check(ExportClass.Object->IsA<UClass>());
+		return;
+	}
+
+	// We need two pieces of external information to create the class defined
+	// in this package:
+	//	1. The class of the class (e.g. UClass::StaticClass(), or 
+	//		UBlueprintGeneratedClass::StaticClass())
+	//	2. The class that the class inherits from (e.g. AActor::StaticClass())
+	// Find or create both of them - at the moment these are always import so no
+	// topo sort is needed for exports:
+	check(ExportClass.ClassIndex.IsImport());
+	UClass* const ClassOfClass = ForceCreatePlaceholderClassImport(
+		ExportClass.ClassIndex.ToImport(), UClass::StaticClass());
+		
+	check(ExportClass.SuperIndex.IsImport());
+	UClass* const SuperClass = ForceCreatePlaceholderClassImport(
+		ExportClass.SuperIndex.ToImport());
+
+	// With required imports resolved, we can now create the export object (which is
+	// a class) and its cdo:
+	UClass* const ClassObject = 
+		Cast<UClass>(UE::FPropertyBagRepository::CreatePropertyBagPlaceholderType(
+			LinkerRoot, ClassOfClass, ExportClass.ObjectName, RF_Public, SuperClass));
+	ClassObject->SetLinker(this, ExportIndex);
+	ClassObject->SetFlags(
+		ExportClass.ObjectFlags|
+		RF_NeedLoad|RF_NeedPostLoad|RF_NeedPostLoadSubobjects|RF_WasLoaded);
+	FUObjectThreadContext::Get().GetSerializeContext()->AddLoadedObject(ClassObject);
+	ExportClass.Object = ClassObject;
+
+	UObject* const ClassCDO = ClassObject->GetDefaultObject(false);
+	// find export index for the cdo, if any:
+	const FName CDOName = ClassCDO->GetFName();
+	const int32 CDOIndex = ExportMap.IndexOfByPredicate(
+		[CDOName, ExportIndex](const FObjectExport& Export)
+		{
+			return Export.ObjectName == CDOName &&
+				Export.ClassIndex.IsExport() &&
+				Export.ClassIndex.ToExport() == ExportIndex;
+		}
+	);
+	if(CDOIndex != INDEX_NONE)
+	{
+		ClassCDO->SetFlags(
+			RF_NeedLoad|RF_NeedPostLoad|RF_NeedPostLoadSubobjects|RF_WasLoaded);
+		FObjectExport& CDOExport = ExportMap[CDOIndex];
+		FUObjectThreadContext::Get().GetSerializeContext()->AddLoadedObject(ClassCDO);
+		CDOExport.Object = ClassCDO;
+		ClassCDO->SetLinker(this, CDOIndex);
+	}
+}
+
+void FLinkerLoad::TryCreateImportedClassForExport(const FObjectExport& Export)
+{
+	if(!Export.ClassIndex.IsImport())
+	{
+		return;
+	}
+
+	FObjectImport& Import = ImportMap[ Export.ClassIndex.ToImport() ];
+	if(Import.XObject != nullptr) // already resolved, living with whatever we've found
+	{
+		return;
+	}
+
+	// Look for an object that already exists, even in our reduced condition - 
+	// this allows coreuobject to function without any placeholders for coreuobject
+	// level types:
+	FObjectImport& OuterImport = ImportMap[ Import.OuterIndex.ToImport() ];
+	if(!OuterImport.XObject)
+	{
+		OuterImport.XObject = FindObjectFast<UPackage>(
+			nullptr, OuterImport.ObjectName, EFindObjectFlags::ExactClass);
+	}
+
+	if(OuterImport.XObject)
+	{
+		Import.XObject = StaticFindObjectFastInternal(
+			nullptr, OuterImport.XObject, Import.ObjectName, EFindObjectFlags::ExactClass);
+	}
+
+	// failed to resolve, create the placeholder:
+	if(!Import.XObject)
+	{
+		ForceCreatePlaceholderClassImport(Export.ClassIndex.ToImport());
+	}
+}
+
+void FLinkerLoad::TryCreateImportedClassesForImports()
+{
+	// The purpose of this function is to build a complete and correct view of
+	// *imported* classes, including consistent use of custom class types that
+	// are not knowably derived from uclass because their associated native
+	// code is not loaded.
+	
+	// Helper function to build a FSoftObjectPath from a not yet loaded FObjectImport:
+	const auto ObjectImportToSoftObjectPath = [this](const FObjectImport& Import) -> FSoftObjectPath
+	{
+		if(Import.OuterIndex.IsNull())
+		{
+			return FSoftObjectPath(FTopLevelAssetPath(Import.ObjectName, FName()));
+		}
+
+		// common case is root level import, so preallocating for that:
+		TArray<FName, TInlineAllocator<2>> Names;
+		const FObjectResource* Current = &Import;
+
+		while (Current)
+		{
+			Names.Add(Current->ObjectName);
+			if (Current->OuterIndex.IsNull())
+			{
+				break;
+			}
+			Current = &ImpExp(Current->OuterIndex);
+		}
+		check(Names.Num() >= 2);
+		FTopLevelAssetPath Path(Names[Names.Num() - 1], Names[Names.Num() - 2]);
+		FUtf8String SubPath;
+		if(Names.Num() > 2 )
+		{
+			for( int32 Iter = Names.Num() - 3; Iter >= 0; --Iter )
+			{
+				SubPath.Append(Names[Iter].ToString());
+				if(Iter >= 1)
+				{
+					SubPath.Append(".");
+				}
+			}
+		}
+		
+		return FSoftObjectPath(Path, SubPath);
+	};
+
+	// Gather all identities of all objects used as a class in
+	// the import table:
+	TSet<FSoftObjectPath> ImportedClasses;
+	Algo::Transform( ImportMap, ImportedClasses, [&](const FObjectImport& Import)
+	{
+		return FSoftObjectPath(
+			FTopLevelAssetPath(Import.ClassPackage, Import.ClassName)
+		);
+	});
+
+	// root classes identified by usage as a class in an import
+	// now look for classes that are instances of classes:
+	TMap<FSoftObjectPath, int32> ClassToImportIndex;
+	TMap<FSoftObjectPath, FSoftObjectPath> AllClassesWithClassTypes;
+	TSet<FSoftObjectPath> DiscoveredClassTypes;
+	const FSoftObjectPath ClassPath(UClass::StaticClass());
+	const FSoftObjectPath ObjectPath(UObject::StaticClass());
+
+	for( int32 ImportIndex = 0; ImportIndex < ImportMap.Num(); ++ImportIndex )
+	{
+		const FObjectImport& Import = ImportMap[ImportIndex];
+		// This import may be a UClass, test it by seeing if it's used as
+		// an imported class:
+		FSoftObjectPath ImportPath = ObjectImportToSoftObjectPath(Import);
+		if( ImportedClasses.Contains(ImportPath) )
+		{
+			ClassToImportIndex.Add(ImportPath, ImportIndex);
+			check(!AllClassesWithClassTypes.Contains(ImportPath)); // duplicate import map?
+			FSoftObjectPath ImportClassPath(
+				FTopLevelAssetPath(Import.ClassPackage, Import.ClassName)
+			);
+			if(ImportClassPath == ClassPath)
+			{
+				// a normal class type:
+				AllClassesWithClassTypes.Add( 
+					ImportPath, 
+					ObjectPath
+				);
+			}
+			else
+			{
+				// a custom class type:
+				AllClassesWithClassTypes.Add( 
+					ImportPath, 
+					ImportClassPath
+				);
+				DiscoveredClassTypes.Add(ImportClassPath);
+			}
+		}
+	}
+
+	// All imports categorized, finally categorize implicitly used
+	// classes that are themselves not an import:
+	// would prevent toposort from organizing our classes:
+	AllClassesWithClassTypes.Add(ObjectPath, FSoftObjectPath());
+	AllClassesWithClassTypes.Add(ClassPath, ObjectPath);
+	for(const FSoftObjectPath& ImplicitClass : ImportedClasses)
+	{
+		if(!AllClassesWithClassTypes.Contains(ImplicitClass))
+		{
+			ClassToImportIndex.Add(ImplicitClass, INDEX_NONE);
+			if(!DiscoveredClassTypes.Contains(ImplicitClass))
+			{
+				AllClassesWithClassTypes.Add(
+					ImplicitClass, ObjectPath
+				);
+			}
+			else
+			{
+				AllClassesWithClassTypes.Add(
+					ImplicitClass, ClassPath
+				);
+			}
+		}
+	}
+
+	// topological sort of the classes - check for cycles
+	TArray<FSoftObjectPath> AllClassesOrdered;
+	AllClassesWithClassTypes.GenerateKeyArray(AllClassesOrdered);
+
+	const bool bOrderedClassNodes = Algo::TopologicalSort(AllClassesOrdered,
+		[&AllClassesWithClassTypes](const FSoftObjectPath& Node)
+		{
+			TArray<FSoftObjectPath> Deps;
+			// This lambda defines dependencies for each node
+			// i.e. which nodes must come before this one.
+			if (const FSoftObjectPath* Parent = AllClassesWithClassTypes.Find(Node))
+			{
+				Deps.Add(*Parent);
+			}
+			return Deps;
+		}
+	);
+	check(bOrderedClassNodes); // cycles in class hierarchy, this should never have happened
+
+	// Classes have been ordered and discovered, we can now create remaining placeholder
+	// classes:
+	for(const FSoftObjectPath& Node : AllClassesOrdered)
+	{
+		const UObject* Existing = Node.ResolveObject();
+		if(Existing)
+		{
+			continue;
+		}
+
+		const int32 ImportIndex = ClassToImportIndex.FindChecked(Node);
+		if(ImportIndex != INDEX_NONE)
+		{
+			ForceCreatePlaceholderClassImport(ImportIndex);
+		}
+		else
+		{
+			// implicit import, just create the class (with a package if necessary):
+			const FName PackageName = Node.GetLongPackageFName();
+			UPackage* ClassPackage = FindObjectFast<UPackage>( 
+				nullptr, PackageName, EFindObjectFlags::None );
+			if(!ClassPackage)
+			{
+				ClassPackage = CreatePackage(*PackageName.ToString());
+				// tag as in memory so no code checks for a corresponding FLinker:
+				ClassPackage->SetPackageFlags(PKG_InMemoryOnly);
+			}
+			const FSoftObjectPath& CustomClass = AllClassesWithClassTypes.FindChecked(Node);
+			UE::FPropertyBagRepository::CreatePropertyBagPlaceholderType(
+				ClassPackage, 
+				UClass::StaticClass(), 
+				Node.GetAssetPath().GetAssetName(), 
+				RF_Public, 
+				CastChecked<UClass>(CustomClass.ResolveObject()));
+		}
+	}
+}
+
+UClass* FLinkerLoad::ForceCreatePlaceholderClassImport(int32 ImportIndex, UClass* SuperClassOverride)
+{
+	FObjectImport& Import = ImportMap[ImportIndex];
+
+	// Placeholder already created:
+	if (Import.XObject)
+	{
+		return CastChecked<UClass>(Import.XObject);
+	}
+
+	if(!SuperClassOverride)
+	{
+		SuperClassOverride = UObject::StaticClass();
+	}
+
+	UObject* ImportClassPackage = FindObjectFast<UPackage>(nullptr, Import.ClassPackage, EFindObjectFlags::ExactClass);
+	UClass* ImportClass = FindObjectFast<UClass>(ImportClassPackage, Import.ClassName, EFindObjectFlags::None);
+	if(!ImportClass)
+	{
+		ImportClass = UClass::StaticClass();
+	}
+
+	// outer should have already been created by pass 0:
+	UPackage* ClassObjectPackage = Cast<UPackage>(IndexToObject(Import.OuterIndex));
+	check(ClassObjectPackage);
+	
+	UClass* PlaceholderClass = FindObjectFast<UClass>( ClassObjectPackage, Import.ObjectName, EFindObjectFlags::None );
+	if(!PlaceholderClass)
+	{
+		PlaceholderClass = CastChecked<UClass>(UE::FPropertyBagRepository::CreatePropertyBagPlaceholderType(
+			ClassObjectPackage, ImportClass, Import.ObjectName, RF_Public, SuperClassOverride));
+	}
+
+	Import.XObject = PlaceholderClass;
+	return PlaceholderClass;
+}
+
+void FLinkerLoad::SetBoundsImpl(FArchiveBounds NewBounds)
+{
+	CurrentBounds = NewBounds;
+}
+
+FArchiveBounds FLinkerLoad::GetBoundsImpl() const
+{
+	return CurrentBounds;
+}
+#endif // UE_INTERNAL_UNIVERSAL_PLACEHOLDERS
+
 UClass* FLinkerLoad::TryCreatePlaceholderClassForExport(int32 ExportIndex)
 {
 	UClass* ClassObject = nullptr;
@@ -4550,7 +4926,10 @@ void FLinkerLoad::LoadAllObjects(bool bForcePreload)
 
 	// Tick the heartbeat if we're loading on the game thread
 	const bool bShouldTickHeartBeat = IsInGameThread();
-
+	
+#if UE_INTERNAL_UNIVERSAL_PLACEHOLDERS
+	CreateUniversalPlaceholders();
+#endif // UE_INTERNAL_UNIVERSAL_PLACEHOLDERS
 	for(int32 ExportIndex = 0; ExportIndex < ExportMap.Num(); ++ExportIndex)
 	{
 #if WITH_EDITOR
@@ -4813,6 +5192,7 @@ void FLinkerLoad::Preload( UObject* Object )
 				const bool bDoesSavedClassMatchActualClass = DoesSavedClassMatchActualClass(ExportIndex);
 				bool& bIsLoadingToPropertyBagObject = FUObjectThreadContext::Get().GetSerializeContext()->bIsLoadingToPropertyBagObject;
 				TGuardValue<bool> ScopeGuard(bIsLoadingToPropertyBagObject, UE::FPropertyBagRepository::IsPropertyBagPlaceholderObject(Object));
+				bool bCanRunNativeSerialize = !bIsLoadingToPropertyBagObject;
 				if (UEVer() >= EUnrealEngineObjectUE5Version::SCRIPT_SERIALIZATION_OFFSET)
 				{
 					if (bIsLoadingToPropertyBagObject || !bDoesSavedClassMatchActualClass)
@@ -4820,7 +5200,12 @@ void FLinkerLoad::Preload( UObject* Object )
 						// note: script start/end offsets are relative to the export's offset in the file
 						StartPos += Export.ScriptSerializationStartOffset;
 						ExpectedSerialSize = Export.ScriptSerializationEndOffset - Export.ScriptSerializationStartOffset;
-						bSerializeOnlyScriptProperties = true;	// signals that we can safely narrow the load to SerializeScriptProperties()
+						// Some UObjects (e.g. URigHierarchy) don't invoke base serialize, and thus won't write out any tagged data to read, even the
+						// none terminator. We can detect those by checking for 0 bytes expected and skip tagged serialization
+						bSerializeOnlyScriptProperties = ExpectedSerialSize != 0;	// signals that we can safely narrow the load to SerializeScriptProperties()
+						// even if we can't serialize the script properties, when
+						// the above conditions are met we don't want to run the native serializer:
+						bCanRunNativeSerialize = false;
 					}
 				}
 #endif
@@ -4964,7 +5349,7 @@ void FLinkerLoad::Preload( UObject* Object )
 							{
 								Object->SerializeScriptProperties(*this);
 							}
-							else
+							else if (bCanRunNativeSerialize)
 #endif
 							{
 								Object->Serialize(*this);
@@ -5052,7 +5437,9 @@ void FLinkerLoad::Preload( UObject* Object )
 					}
 					else
 					{
+#if !UE_INTERNAL_UNIVERSAL_PLACEHOLDERS
 						UE_ASSET_LOG(LogLinker, Fatal, PackagePath, TEXT("%s: Serial size mismatch: Got %d, Expected %" INT64_FMT), *Object->GetFullName(), (int32)SizeSerialized, ExpectedSerialSize);
+#endif //!UE_INTERNAL_UNIVERSAL_PLACEHOLDERS
 					}
 				}
 
@@ -5986,6 +6373,29 @@ UObject* FLinkerLoad::CreateImport( int32 Index )
 
 						// Find object now that we know it's class, outer and name.
 						FindObject = FindImportFast(FindClass, FindOuter, Import.ObjectName);
+#if UE_INTERNAL_UNIVERSAL_PLACEHOLDERS
+						if(!FindObject)
+						{
+							if(FindClass->IsChildOf(UClass::StaticClass()))
+							{
+								ForceCreatePlaceholderClassImport(Index);
+							}
+							else
+							{
+								FindObject = NewObject<UObject>(FindOuter, FindClass, Import.ObjectName, RF_Public);
+								if(UScriptStruct* Struct = Cast<UScriptStruct>(FindObject))
+								{
+									// add a single property to avoid creating a 0 sized struct. Zero sized allocations
+									// are supported by some systems, but containers are generally unhappy with zero stride
+									// container elements. We can't just use the fallback struct as we want a proper
+									// object name:
+									FBoolProperty* Property = new FBoolProperty(Struct, "__PlaceholderSentinel", RF_NoFlags);
+									Struct->ChildProperties = Property;
+									Struct->StaticLink(true); // make the struct roughly usable, it's just a dummy type
+								}
+							}
+						}
+#endif // UE_INTERNAL_UNIVERSAL_PLACEHOLDERS
 					}
 
 					if( FindObject )
@@ -6383,7 +6793,9 @@ FArchive& FLinkerLoad::operator<<( UObject*& Object )
 	Ar << Index;
 
 	Object = ResolveResource(Index);
-#if WITH_EDITOR
+	// UE_INTERNAL_UNIVERSAL_PLACEHOLDERS disables this because some native facilities like ObjectRedirectors
+	// do not provide a ObjectProperty and we want to allow them to resolve objects:
+#if WITH_EDITOR && !UE_INTERNAL_UNIVERSAL_PLACEHOLDERS
 	if (Object && UE::FPropertyBagRepository::IsPropertyBagPlaceholderObject(Object))
 	{
 		if (UE::FPropertyBagRepository::IsPropertyBagPlaceholderObjectFeatureEnabled(UE::FPropertyBagRepository::EPlaceholderObjectFeature::SerializeExportReferencesOnLoad))
@@ -6430,6 +6842,10 @@ FArchive& FLinkerLoad::operator<<(FObjectPtr& ObjectPtr)
 					// This resolves to NULL on access since the reference can't be cast to a pointer bound to its original type.
 					// However, the underlying value (when serialized) will always resolve to the placeholder object (e.g. for GC).
 					return FObjectPtr(FObjectHandle{ UE::CoreUObject::Private::MakePackedObjectRef(ResolvedObject).EncodedRef });
+#elif UE_INTERNAL_UNIVERSAL_PLACEHOLDERS
+					// Type hierarchies can be incomplete in UE_INTERNAL_UNIVERSAL_PLACEHOLDERS environments, 
+					// but unless there's native code using the object we can allow the reference to exist:
+					return FObjectPtr(ResolvedObject); 
 #else
 					// If type safety features are disabled, serialize it as an unsafe reference to a placeholder-typed object.
 					// Note: Similar to hard references above, this means we won't find it for replacement at reinstancing time.
@@ -7119,6 +7535,7 @@ FLinkerLoad::ELinkerStatus FLinkerLoad::FixupExportMap()
 
 			// ActorComponents outered to a BlueprintGeneratedClass (or even older ones that are outered to Blueprint) need to be marked RF_Public, but older content was
 			// not created as such.  This updates the ExportTable such that they are correctly flagged when created and when other packages validate their imports.
+#if !UE_INTERNAL_UNIVERSAL_PLACEHOLDERS
 			if (UEVer() < VER_UE4_BLUEPRINT_GENERATED_CLASS_COMPONENT_TEMPLATES_PUBLIC)
 			{
 				if ((Export.ObjectFlags & RF_Public) == 0)
@@ -7145,6 +7562,7 @@ FLinkerLoad::ELinkerStatus FLinkerLoad::FixupExportMap()
 					}
 				}
 			}
+#endif // !UE_INTERNAL_UNIVERSAL_PLACEHOLDERS
 
 			// Look for subobject redirects and instance redirects
 			FCoreRedirectObjectName OldClassName(NameClass, NAME_None, NamePackage);
diff --git a/Engine/Source/Runtime/CoreUObject/Private/UObject/PropertyBagRepository.cpp b/Engine/Source/Runtime/CoreUObject/Private/UObject/PropertyBagRepository.cpp
--- a/Engine/Source/Runtime/CoreUObject/Private/UObject/PropertyBagRepository.cpp
+++ b/Engine/Source/Runtime/CoreUObject/Private/UObject/PropertyBagRepository.cpp
@@ -689,7 +689,7 @@ void FPropertyBagRepository::PostLoadInstanceDataObject(const UObject* Owner)
 				}
 			}
 
-#if !UE_BUILD_SHIPPING_WITH_EDITOR
+#if !UE_BUILD_SHIPPING_WITH_EDITOR && !UE_INTERNAL_UNIVERSAL_PLACEHOLDERS
 			// Run IDT operations on this instance to recover data
 			TObjectPtr<UObject> InOutInstance = const_cast<UObject*>(Owner);
 			UE::FInstanceDataTransforms::Get().PatchInstanceDataObject(BagData->InstanceDataObject, InOutInstance);
@@ -945,9 +945,15 @@ UObject* FPropertyBagRepository::CreateInstanceDataObject(UObject* Owner, FArchi
 	// Load the IDO.
 	if (StartOffset != EndOffset)
 	{
+#if UE_INTERNAL_UNIVERSAL_PLACEHOLDERS
+		// we don't need archetype matching and the FArchetypeMatchingArchive doesn't bounds
+		// check so is unsafe in universal placeholder contexts
+		FArchive& ArchiveWrapper = Archive;
+#else // !UE_INTERNAL_UNIVERSAL_PLACEHOLDERS
 		// We want to force the template we used as the archetype during SerializeScriptProperties. For this reason 
 		// we wrap the archive in a proxy that will return the correct archetype.
 		FArchetypeMatchingArchive ArchiveWrapper(Archive, InstanceDataObject, Template);
+#endif
 
 		const int64 OffsetToRestore = ArchiveWrapper.Tell();
 		FScopedObjectSerializeContext ObjectSerializeContext(InstanceDataObject, ArchiveWrapper);
@@ -1067,7 +1073,13 @@ UStruct* FPropertyBagRepository::CreatePropertyBagPlaceholderType(UObject* Outer
 	{
 		// Create and configure its CDO as if it were loaded - for non-native class types, this is required.
 		UObject* PlaceholderClassDefaults = PlaceholderTypeAsClass->GetDefaultObject();
-		PlaceholderTypeAsClass->PostLoadDefaultObject(PlaceholderClassDefaults);
+		// CDOs for UClasses themselves cannot typically be post loaded
+		// because they will attempt to bind to their native constructor
+		// and there may be none for placeholder types..
+		if(!PlaceholderClassDefaults->IsA<UClass>())
+		{
+			PlaceholderTypeAsClass->PostLoadDefaultObject(PlaceholderClassDefaults);
+		}
 
 		// This class is for internal use and should not be exposed for selection or instancing in the editor.
 		PlaceholderTypeAsClass->ClassFlags |= CLASS_Hidden | CLASS_HideDropDown;
diff --git a/Engine/Source/Runtime/CoreUObject/Private/UObject/PropertyMap.cpp b/Engine/Source/Runtime/CoreUObject/Private/UObject/PropertyMap.cpp
--- a/Engine/Source/Runtime/CoreUObject/Private/UObject/PropertyMap.cpp
+++ b/Engine/Source/Runtime/CoreUObject/Private/UObject/PropertyMap.cpp
@@ -718,7 +718,15 @@ void FMapProperty::SerializeItem(FStructuredArchive::FSlot Slot, void* Value, co
 				}
 			}
 
-			MapHelper.Rehash();
+#if UE_INTERNAL_UNIVERSAL_PLACEHOLDERS
+			// In UE_INTERNAL_UNIVERSAL_PLACEHOLDERS environs we may lose user types
+			// and contained data will thus be unhashable - we could fix this
+			// with even better placeholder support
+			if(KeyProp->HasAnyPropertyFlags(CPF_HasGetValueTypeHash)) 
+#endif
+			{
+				MapHelper.Rehash();
+			}
 		}
 		else // Slower loading path that mutates non-empty map
 		{
@@ -1459,8 +1467,31 @@ EConvertFromTypeResult FMapProperty::ConvertFromType(const FPropertyTag& Tag, FS
 	}
 	else
 	{
+		const auto IsPotentiallyUnsafeByteProperty = []( const FProperty* P )
+		{
+#if UE_INTERNAL_UNIVERSAL_PLACEHOLDERS
+			if(const FByteProperty* ByteProperty = CastField<FByteProperty>(P))
+			{
+				// byte properties with enums are unambiguously enum properties
+				// but if no enum was declared we cannot know if it's a pure
+				// byte property or an enum property that lost its enum:
+				// disaster recovery is possible, but better to just resave the asset
+				// post PROPERTY_TAG_COMPLETE_TYPE_NAME:
+				return ByteProperty->Enum == nullptr;
+			}
+#endif // UE_INTERNAL_UNIVERSAL_PLACEHOLDERS
+			return false;
+		};
 		bCanSerializeKey = (KeyTypeName == KeyProp->GetID());
+		if(bCanSerializeKey && IsPotentiallyUnsafeByteProperty(KeyProp))
+		{
+			return EConvertFromTypeResult::CannotConvert;
+		}
 		bCanSerializeValue = (ValueTypeName == ValueProp->GetID());
+		if(bCanSerializeValue && IsPotentiallyUnsafeByteProperty(ValueProp))
+		{
+			return EConvertFromTypeResult::CannotConvert;
+		}
 		if ((bCanSerializeKey || KeyTypeName.IsNone()) && (bCanSerializeValue || ValueTypeName.IsNone()))
 		{
 			return EConvertFromTypeResult::UseSerializeItem;
diff --git a/Engine/Source/Runtime/CoreUObject/Private/UObject/PropertySet.cpp b/Engine/Source/Runtime/CoreUObject/Private/UObject/PropertySet.cpp
--- a/Engine/Source/Runtime/CoreUObject/Private/UObject/PropertySet.cpp
+++ b/Engine/Source/Runtime/CoreUObject/Private/UObject/PropertySet.cpp
@@ -349,7 +349,15 @@ void FSetProperty::SerializeItem(FStructuredArchive::FSlot Slot, void* Value, co
 			}
 		}
 
-		SetHelper.Rehash();
+#if UE_INTERNAL_UNIVERSAL_PLACEHOLDERS
+		// In UE_INTERNAL_UNIVERSAL_PLACEHOLDERS environs we may lose user types
+		// and contained data will thus be unhashable - we could fix this
+		// with even better placeholder support
+		if(ElementProp->HasAnyPropertyFlags(CPF_HasGetValueTypeHash)) 
+#endif
+		{
+			SetHelper.Rehash();
+		}
 	}
 	else
 	{
diff --git a/Engine/Source/Runtime/CoreUObject/Private/UObject/PropertyTag.cpp b/Engine/Source/Runtime/CoreUObject/Private/UObject/PropertyTag.cpp
--- a/Engine/Source/Runtime/CoreUObject/Private/UObject/PropertyTag.cpp
+++ b/Engine/Source/Runtime/CoreUObject/Private/UObject/PropertyTag.cpp
@@ -597,16 +597,28 @@ void FPropertyTag::SerializeTaggedProperty(FStructuredArchive::FSlot Slot, FProp
 #endif
 		FSerializedPropertyScope SerializedProperty(UnderlyingArchive, Property);
 		FPropertyTagScope CurrentPropertyTagScope(this);
+		
+#if UE_INTERNAL_UNIVERSAL_PLACEHOLDERS
+		// Update archive bounds with tagged property known size:
+		FArchiveBoundsScope ArchiveBounds(UnderlyingArchive.SetBounds(Size)); 
+#endif
 
 		Property->SerializeItem(Slot, Value, Defaults);
 	}
 
 	// Ensure that we serialize what we expected to serialize.
 	const int64 EndOfProperty = UnderlyingArchive.Tell();
-	if (Size && (EndOfProperty - StartOfProperty != Size))
+	if (Size && (EndOfProperty - StartOfProperty != Size) 
+#if UE_INTERNAL_UNIVERSAL_PLACEHOLDERS
+	|| UnderlyingArchive.IsCriticalError()
+#endif//UE_INTERNAL_UNIVERSAL_PLACEHOLDERS
+	)
 	{
 		UE_LOG(LogClass, Error, TEXT("Failed loading tagged %s. Read %" INT64_FMT "B, expected %dB. Package: %s"), *GetFullNameSafe(Property), EndOfProperty - StartOfProperty, Size, *UnderlyingArchive.GetArchiveName());
 		UnderlyingArchive.Seek(StartOfProperty + Size);
 		Property->ClearValue(Value);
+#if UE_INTERNAL_UNIVERSAL_PLACEHOLDERS
+		UnderlyingArchive.ClearCriticalError();
+#endif//UE_INTERNAL_UNIVERSAL_PLACEHOLDERS
 	}
 }
diff --git a/Engine/Source/Runtime/CoreUObject/Public/UObject/LinkerLoad.h b/Engine/Source/Runtime/CoreUObject/Public/UObject/LinkerLoad.h
--- a/Engine/Source/Runtime/CoreUObject/Public/UObject/LinkerLoad.h
+++ b/Engine/Source/Runtime/CoreUObject/Public/UObject/LinkerLoad.h
@@ -915,6 +915,25 @@ private:
 	 */
 	UClass* TryCreatePlaceholderClassImport(int32 ImportIndex);
 
+#if UE_INTERNAL_UNIVERSAL_PLACEHOLDERS
+	// Universal placeholders force the loader to make placeholders for all objects
+	// inferring what it can about the objects within the package and creating
+	// placeholder imports to serve only the interests of the objects within the
+	// package. To understand these functions simply read CreateUniversalPlaceholders
+	void CreateUniversalPlaceholders();
+	void TryCreateImportedPackage(FObjectImport& Import);
+	void TryCreatePlaceholderClassExport(const FObjectExport& Export);
+	void TryCreateImportedClassForExport(const FObjectExport& Export);
+	void TryCreateImportedClassesForImports();
+	UClass* ForceCreatePlaceholderClassImport(int32 ImportIndex, UClass* SuperClassOverride = nullptr);
+	
+	// FArchive bounds implementation - needed for high placeholder environments but
+	// likely useful everywhere that can afford them:
+	COREUOBJECT_API virtual void SetBoundsImpl(FArchiveBounds NewBounds) override;
+	COREUOBJECT_API virtual FArchiveBounds GetBoundsImpl() const override;
+	FArchiveBounds CurrentBounds;
+#endif
+
 	/**
 	 * Utility function to create a placeholder class for the specified export.
 	 * This will be called if the export could not resolve its LoadClass import.
@@ -1076,7 +1095,7 @@ private:
 #endif // WITH_EDITOR
 
 		checkSlow(FPlatformTLS::GetCurrentThreadId() == OwnerThread);
-#if WITH_EDITOR
+#if WITH_EDITOR && !UE_INTERNAL_UNIVERSAL_PLACEHOLDERS
 		// This handles the case where we're only serializing to a property bag and we don't immediately know where the script
 		// property data starts/ends in the object's data stream. Allows for loading from packages saved with an older version.
 		if (UNLIKELY(IsLoadingToPropertyBagObject() && !bIsSerializingScriptProperties))
@@ -1084,7 +1103,19 @@ private:
 			Loader->Seek(Tell() + Length);
 		}
 		else
-#endif
+#elif UE_INTERNAL_UNIVERSAL_PLACEHOLDERS
+		// Similar to above, but UE_INTERNAL_UNIVERSAL_PLACEHOLDERS has more bounds available:
+		const bool bOutOfBounds = Tell() + Length > CurrentBounds.BoundsBasis + CurrentBounds.BoundsSize;
+		if (UNLIKELY((IsLoadingToPropertyBagObject() && !bIsSerializingScriptProperties) || bOutOfBounds))
+		{
+			Loader->Seek(Tell() + Length);
+			if(bOutOfBounds)
+			{
+				SetCriticalError(); // report misread to FPropertyTag so we can skip bad data
+			}
+		}
+		else
+#endif // WITH_EDITOR / UE_INTERNAL_UNIVERSAL_PLACEHOLDERS related extensions
 		{
 			Loader->Serialize(V, Length);
 		}
diff --git a/Engine/Source/Runtime/CoreUObject/Public/UObject/ObjectPtr.h b/Engine/Source/Runtime/CoreUObject/Public/UObject/ObjectPtr.h
--- a/Engine/Source/Runtime/CoreUObject/Public/UObject/ObjectPtr.h
+++ b/Engine/Source/Runtime/CoreUObject/Public/UObject/ObjectPtr.h
@@ -854,6 +854,22 @@ namespace ObjectPtr_Private
 		FORCEINLINE static FArchive& Serialize(FArchive& Ar, TObjectPtr<T>& InObjectPtr)
 		{
 			Ar << InObjectPtr.ObjectPtr;
+#if UE_INTERNAL_UNIVERSAL_PLACEHOLDERS
+			if( Ar.IsLoading() )
+			{
+				// Prevent incompatibly mocked objects from leaking to callers.. this can happen
+				// because of CoreUObject types (e.g. UScriptStruct) being extended in modules
+				// which are unavailable in UE_INTERNAL_UNIVERSAL_PLACEHOLDERS environments. In this case
+				// the type that CoreUObject wants to use is erased above so the linker
+				// can't create a placeholder that necessarily conforms to the native code's
+				// requirements.
+				const UObject* Resolved = InObjectPtr;
+				if(Resolved && !(Resolved->IsA<T>()))
+				{
+					InObjectPtr = nullptr;
+				}
+			}
+#endif // UE_INTERNAL_UNIVERSAL_PLACEHOLDERS
 			return Ar;
 		}
 
diff --git a/Engine/Source/Runtime/Experimental/JsonObjectGraph/Private/JsonStringifyImpl.cpp b/Engine/Source/Runtime/Experimental/JsonObjectGraph/Private/JsonStringifyImpl.cpp
--- a/Engine/Source/Runtime/Experimental/JsonObjectGraph/Private/JsonStringifyImpl.cpp
+++ b/Engine/Source/Runtime/Experimental/JsonObjectGraph/Private/JsonStringifyImpl.cpp
@@ -20,6 +20,11 @@
 #include "UObject/Package.h"
 #include "UObject/TextProperty.h"
 
+#if WITH_EDITORONLY_DATA
+#include "UObject/InstanceDataObjectUtils.h"
+#include "UObject/PropertyBagRepository.h"
+#endif// WITH_EDITORONLY_DATA
+
 namespace UE::Private
 {
 
@@ -45,6 +50,15 @@ struct FPackageReferenceFinder : public FArchiveUObject
 		{
 			((UObject*)Obj)->Serialize(*this);
 		}
+		
+#if WITH_EDITORONLY_DATA
+		UE::FPropertyBagRepository& Repository = UE::FPropertyBagRepository::Get();
+		const UObject* InstanceDataObject = Repository.FindInstanceDataObject(Obj);
+		if (InstanceDataObject)
+		{
+			InstanceDataObject->GetClass()->SerializeBin(*this, (UObject*)InstanceDataObject);
+		}
+#endif//WITH_EDITORONLY_DATA
 	}
 
 private:
@@ -146,6 +160,9 @@ static void GatherExports(
 
 		OutExports.Add(Iter);
 
+		// The object graph is incomplete when using UE_INTERNAL_UNIVERSAL_PLACEHOLDERS
+		// and so we won't be able to do full graph analysis:
+#if !UE_INTERNAL_UNIVERSAL_PLACEHOLDERS
 		ScratchRefs.Add(Iter->GetClass());
 		FPackageReferenceFinder ReferencedObjects(Iter, ScratchRefs, bFilterEditorOnly);
 		for (const UObject* Obj : ScratchRefs)
@@ -180,7 +197,24 @@ static void GatherExports(
 			PendingRefs.Add(Obj);
 		}
 		ScratchRefs.Reset();
+#endif
 	}
+	
+	// Instead in UE_INTERNAL_UNIVERSAL_PLACEHOLDERS we treat all
+	// roots as UPackage and gather all children recursively.
+	// In this configuration we cannot reliably categorize objects as
+	// editor only vs. not because that determination is made
+	// by native code and uproject/uplugin structure - neither of 
+	// which are loaded when using placeholders:
+#if UE_INTERNAL_UNIVERSAL_PLACEHOLDERS
+	checkf(!bFilterEditorOnly, TEXT("Cannot reliably filter editor only objects in this configuration"));
+	for(const UObject* Root : OutRoots)
+	{
+		TArray<UObject*> ObjectInners;
+		GetObjectsWithOuter(Root, ObjectInners, true, RF_Transient);
+		OutExports.Append(ObjectInners);
+	}
+#endif
 }
 
 static void WriteObjectPath(FUtf8StringBuilderBase& OutPath, const UObject* ForObject, const UObject* OuterLimit)
@@ -348,6 +382,9 @@ void FJsonStringifyImpl::WriteObjectToJson(const UObject* Object)
 
 		WriteIdentifierAndValueToJson(Object, ValidatedArchtype, Property);
 	}
+#if WITH_EDITORONLY_DATA
+	WriteInstanceDataObjectProperties(Object);
+#endif // WITH_EDITORONLY_DATA
 	// We may have inner objects that were not referenced directly by this object, 
 	// but will be referenced by other objects in the graph. We must write them here, or we
 	// would have to encode them at the root level, which would disrupt locality.
@@ -412,6 +449,23 @@ void FJsonStringifyImpl::WriteNativeObjectData()
 	}
 }
 
+#if WITH_EDITORONLY_DATA
+void FJsonStringifyImpl::WriteInstanceDataObjectProperties(const UObject* ForObject)
+{
+	UE::FPropertyBagRepository& Repository = UE::FPropertyBagRepository::Get();
+	const UObject* InstanceDataObject = Repository.FindInstanceDataObject(ForObject);
+	if (!InstanceDataObject)
+	{
+		return;
+	}
+
+	for (TFieldIterator<FProperty> FieldIt(InstanceDataObject->GetClass()); FieldIt; ++FieldIt)
+	{
+		WriteIdentifierAndValueToJson(InstanceDataObject, nullptr, *FieldIt);
+	}
+}
+#endif// WITH_EDITORONLY_DATA
+
 void FJsonStringifyImpl::WriteIndirectlyReferencedContainedObjects(const UObject* ForObject)
 {
 	TArray<const UObject*> UnwrittenInners;
diff --git a/Engine/Source/Runtime/Experimental/JsonObjectGraph/Private/JsonStringifyImpl.h b/Engine/Source/Runtime/Experimental/JsonObjectGraph/Private/JsonStringifyImpl.h
--- a/Engine/Source/Runtime/Experimental/JsonObjectGraph/Private/JsonStringifyImpl.h
+++ b/Engine/Source/Runtime/Experimental/JsonObjectGraph/Private/JsonStringifyImpl.h
@@ -47,6 +47,9 @@ private:
 	// data is placed into a special key for a loader that needs to order
 	// UObject construction (find archetypes, UClasses, UStructs, etc):
 	void WriteNativeObjectData();
+#if WITH_EDITORONLY_DATA
+	void WriteInstanceDataObjectProperties(const UObject* ForObject);
+#endif // WITH_EDITORONLY_DATA
 	void WriteIndirectlyReferencedContainedObjects(const UObject* ForObject);
 	FUtf8String WriteObjectReference(const UObject* Object) const;
 	FUtf8String WriteFieldReference(const FField* Value) const;
