Fix missing outer when generating a unique name for trashed objects from CL 47253414.

diff --git a/Engine/Source/Runtime/CoreUObject/Private/UObject/UObjectGlobals.cpp b/Engine/Source/Runtime/CoreUObject/Private/UObject/UObjectGlobals.cpp
--- a/Engine/Source/Runtime/CoreUObject/Private/UObject/UObjectGlobals.cpp
+++ b/Engine/Source/Runtime/CoreUObject/Private/UObject/UObjectGlobals.cpp
@@ -262,7 +262,7 @@ void TrashObject(TNotNull<UObject* const> Object)
 
 		const FName BaseName = bIsAPackage ? NAME_TrashedPackage : NAME_None;
 
-		const FName NewName = MakeUniqueObjectName(nullptr, Object->GetClass(), BaseName);
+		const FName NewName = MakeUniqueObjectName(Object->GetOuter(), Object->GetClass(), BaseName);
 		Object->Rename(*NewName.ToString(), nullptr, REN_DontCreateRedirectors | REN_NonTransactional | REN_DoNotDirty);
 
 		if (bIsAPackage)
