[CoreUObject] AUTORTFM_DISABLE SavePackage2.cpp private functions

None of this should not be run in the closed.

AutoRTFM-disabling this code:
* Reduces executable binary size, and improves compilation speed.
* Allows the called functions to be AutoRTFM-disabled without compilation error.

UPackage::Save() is not currently AutoRTFM-disabled as this would require fixing up all the places that call this function (future work). Instead we use AutoRTFM::UnreachableIfClosed() which is a promise to the compiler that this code should not be called in the closed. This allows an AutoRTFM-enabled -> AutoRTFM-disabled call edge for the function, without requiring transitive fixups of the function.

See Engine\Source\Runtime\AutoRTFM\Documentation\ModesAndAttributes.md for more information about AutoRTFM modes.

diff --git a/Engine/Source/Runtime/CoreUObject/Private/UObject/SavePackage2.cpp b/Engine/Source/Runtime/CoreUObject/Private/UObject/SavePackage2.cpp
--- a/Engine/Source/Runtime/CoreUObject/Private/UObject/SavePackage2.cpp
+++ b/Engine/Source/Runtime/CoreUObject/Private/UObject/SavePackage2.cpp
@@ -1,5 +1,6 @@
 // Copyright Epic Games, Inc. All Rights Reserved.
 
+#include "AutoRTFM.h"
 #include "UObject/SavePackage.h"
 #include "UObject/LinkerLoad.h"
 
@@ -369,7 +370,7 @@ ESavePackageResult RoutePresave(FSaveContext& SaveContext)
 	return ReturnSuccessOrCancel();
 }
 
-ESavePackageResult HarvestPackage(FSaveContext& SaveContext)
+AUTORTFM_DISABLE ESavePackageResult HarvestPackage(FSaveContext& SaveContext)
 {
 	SCOPED_SAVETIMER(UPackage_Save_HarvestPackage);
 
@@ -1795,7 +1796,7 @@ void GetCellsWithOuter(Verse::FAllocationContext Context, FSaveContext& SaveCont
 }
 #endif
 
-void SavePreloadDependencies(FStructuredArchive::FRecord& StructuredArchiveRoot, FSaveContext& SaveContext)
+AUTORTFM_DISABLE void SavePreloadDependencies(FStructuredArchive::FRecord& StructuredArchiveRoot, FSaveContext& SaveContext)
 {
 #if WITH_VERSE_VM || defined(__INTELLISENSE__)
 	Verse::FAllocationContext Context = Verse::FAllocationContextPromise{};
@@ -2355,7 +2356,7 @@ ESavePackageResult SetSummaryTotalHeaderSize(FSaveContext& SaveContext, int64 To
 	}
 }
 
-ESavePackageResult WritePackageHeader(FStructuredArchive::FRecord& StructuredArchiveRoot, FSaveContext& SaveContext)
+AUTORTFM_DISABLE ESavePackageResult WritePackageHeader(FStructuredArchive::FRecord& StructuredArchiveRoot, FSaveContext& SaveContext)
 {
 	FLinkerSave* Linker = SaveContext.GetLinker();
 #if WITH_EDITOR
@@ -3649,7 +3650,7 @@ void PostSavePackage(FSaveContext& SaveContext)
 	UPackage::PackageSavedWithContextEvent.Broadcast(SaveContext.GetFilename(), Package, FObjectPostSaveContext(ObjectSaveContext));
 }
 
-ESavePackageResult SaveHarvestedRealms(FSaveContext& SaveContext, ESaveRealm HarvestingContextToSave)
+AUTORTFM_DISABLE ESavePackageResult SaveHarvestedRealms(FSaveContext& SaveContext, ESaveRealm HarvestingContextToSave)
 {
 	// Set the current harvested context to save
 	FSaveContext::FSetSaveRealmToSaveScope Scope(SaveContext, HarvestingContextToSave);
@@ -3875,7 +3876,7 @@ ESavePackageResult SaveHarvestedRealms(FSaveContext& SaveContext, ESaveRealm Har
 /**
  * InnerSave is the portion of Save that can be safely run concurrently
  */
-ESavePackageResult InnerSave(FSaveContext& SaveContext)
+AUTORTFM_DISABLE ESavePackageResult InnerSave(FSaveContext& SaveContext)
 {
 	// Create slow task dialog if needed
 	const int32 TotalSaveSteps = 3;
@@ -3986,6 +3987,8 @@ FText GetSlowTaskStatusMessage(const FSaveContext& SaveContext)
 FSavePackageResultStruct UPackage::Save(UPackage* InPackage, UObject* InAsset, const TCHAR* InFilename,
 	const FSavePackageArgs& SaveArgs)
 {
+	AutoRTFM::UnreachableIfClosed("TODO: Mark this function as AUTORTFM_DISABLED and fixup callers");
+
 	TRACE_CPUPROFILER_EVENT_SCOPE(UPackage::Save);
 
 	COOK_STAT(FScopedDurationTimer FuncSaveTimer(FSavePackageStats::SavePackageTimeSec));
@@ -4092,6 +4095,8 @@ FSavePackageResultStruct UPackage::Save(UPackage* InPackage, UObject* InAsset, c
 
 ESavePackageResult UPackage::SaveConcurrent(TArrayView<FPackageSaveInfo> InPackages, const FSavePackageArgs& SaveArgs, TArray<FSavePackageResultStruct>& OutResults)
 {
+	AutoRTFM::UnreachableIfClosed("TODO: Mark this function as AUTORTFM_DISABLED and fixup callers");
+
 	const int32 TotalSaveSteps = 4;
 	FScopedSlowTask SlowTask(TotalSaveSteps, NSLOCTEXT("Core", "SavingFiles", "Saving files..."), SaveArgs.bSlowTask);
 	SlowTask.MakeDialogDelayed(3.0f, !!(SaveArgs.SaveFlags & SAVE_FromAutosave));
@@ -4150,7 +4155,7 @@ ESavePackageResult UPackage::SaveConcurrent(TArrayView<FPackageSaveInfo> InPacka
 		//FScopedSavingFlag IsSavingFlag(false);
 
 		// Concurrent Part
-		ParallelFor(PackageSaveContexts.Num(), [&PackageSaveContexts](int32 PackageIdx)
+		ParallelFor(PackageSaveContexts.Num(), [&PackageSaveContexts] AUTORTFM_DISABLE (int32 PackageIdx)
 			{
 				InnerSave(PackageSaveContexts[PackageIdx]);
 			});
