[ZenLoader]
- Move the serialization mutex to the leaf to make deadlock impossible
- Add HeaderData refcounting to avoid use-after-free when multithreaded is active
- Add Insight scope on contention
- Unused return value cleanup

diff --git a/Engine/Source/Runtime/CoreUObject/Private/Serialization/AsyncLoading2.cpp b/Engine/Source/Runtime/CoreUObject/Private/Serialization/AsyncLoading2.cpp
--- a/Engine/Source/Runtime/CoreUObject/Private/Serialization/AsyncLoading2.cpp
+++ b/Engine/Source/Runtime/CoreUObject/Private/Serialization/AsyncLoading2.cpp
@@ -3623,6 +3623,42 @@ struct FAsyncPackage2 : public FRefCountBase
 
 	void Release();
 
+	bool TryAddHeaderRef()
+	{
+		for (;;)
+		{
+			int32 CurrentRefCount = HeaderRefCount.load(std::memory_order_relaxed);
+			if (CurrentRefCount == 0)
+			{
+				return false;
+			}
+			
+			if (HeaderRefCount.compare_exchange_strong(CurrentRefCount, CurrentRefCount + 1))
+			{
+				return true;
+			}
+		}
+	}
+
+	void ReleaseHeaderRef()
+	{
+		int32 OldRefCount = HeaderRefCount.fetch_sub(1);
+		check(OldRefCount > 0);
+		if (OldRefCount == 1)
+		{
+			TRACE_CPUPROFILER_EVENT_SCOPE(FAsyncPackage2::ReleaseHeaderRef);
+			HeaderData.Reset();
+			SerializationState.ReleaseIoRequest();
+#if WITH_EDITOR
+			if (OptionalSegmentHeaderData.IsSet())
+			{
+				OptionalSegmentHeaderData->Reset();
+				OptionalSegmentSerializationState->ReleaseIoRequest();
+			}
+#endif
+		}
+	}
+
 	void ClearImportedPackages();
 
 	/**
@@ -3755,7 +3791,10 @@ private:
 	TAsyncAtomic<uint64>			SyncLoadContextId = 0;
 	/** Time load begun. This is NOT the time the load was requested in the case of pending requests.	*/
 	double						LoadStartTime = 0.0;
-	TAsyncAtomic<int32>			RefCount = 0;
+	/** RefCount on the package itself */
+	TAsyncAtomic<int32>			RefCount = 1;
+	/** RefCount on the HeaderData information to avoid use-after-free from other threads */
+	TAsyncAtomic<int32>			HeaderRefCount = 1;
 	bool						bHasStartedImportingPackages = false;
 	int32						ProcessedExportBundlesCount = 0;
 	/** Current bundle entry index in the current export bundle */
@@ -3771,7 +3810,7 @@ private:
 	/** Whether the constructed objects have been moved to the second loading phase. */
 	TAsyncAtomic<bool> ObjectsNowInPhase2 = false;
 	/** Protects against the same package being serialized from different threads. */
-	UE::FRecursiveMutex         SerializationMutex;
+	UE::FMutex                  SerializationMutex;
 
 	struct FAllDependenciesState
 	{
@@ -3920,10 +3959,10 @@ public:
 	bool CanRunNodeAsync(const FAsyncLoadEventSpec* Spec) const;
 
 	void EventDrivenCreateExport(const FAsyncPackageHeaderData& Header, int32 LocalExportIndex);
-	bool EventDrivenSerializeExport(const FAsyncPackageHeaderData& Header, int32 LocalExportIndex, FExportArchive* Ar);
+	void EventDrivenSerializeExport(const FAsyncPackageHeaderData& Header, int32 LocalExportIndex, FExportArchive* Ar);
 
 	void EventDrivenCreateCellExport(const FAsyncPackageHeaderData& Header, int32 LocalExportIndex, FExportArchive* Ar);
-	bool EventDrivenSerializeCellExport(const FAsyncPackageHeaderData& Header, int32 LocalExportIndex, FExportArchive* Ar);
+	void EventDrivenSerializeCellExport(const FAsyncPackageHeaderData& Header, int32 LocalExportIndex, FExportArchive* Ar);
 
 	UObject* EventDrivenIndexToObject(const FAsyncPackageHeaderData& Header, FPackageObjectIndex Index, bool bCheckSerialized);
 	template<class T>
@@ -3992,15 +4031,15 @@ private:
 	void ProcessExportDependencies(const FAsyncPackageHeaderData& Header, int32 LocalExportIndex, FExportBundleEntry::EExportCommandType CommandType);
 	int32 GetPublicExportIndex(uint64 ExportHash, FAsyncPackageHeaderData*& OutHeader);
 	FString GetNameFromPackageObjectIndex(const FAsyncPackageHeaderData& Header, FPackageObjectIndex Index);
-	UObject* ConditionalCreateExport(const FAsyncPackageHeaderData& Header, int32 LocalExportIndex);
-	UObject* ConditionalSerializeExport(const FAsyncPackageHeaderData& Header, int32 LocalExportIndex);
-	UObject* ConditionalCreateImport(const FAsyncPackageHeaderData& Header, int32 LocalImportIndex);
-	UObject* ConditionalSerializeImport(const FAsyncPackageHeaderData& Header, int32 LocalImportIndex);
+	void ConditionalCreateExport(const FAsyncPackageHeaderData& Header, int32 LocalExportIndex);
+	void ConditionalSerializeExport(const FAsyncPackageHeaderData& Header, int32 LocalExportIndex);
+	void ConditionalCreateImport(const FAsyncPackageHeaderData& Header, int32 LocalImportIndex);
+	void ConditionalSerializeImport(const FAsyncPackageHeaderData& Header, int32 LocalImportIndex);
 	int32 GetPublicCellExportIndex(uint64 CellExportHash, FAsyncPackageHeaderData*& OutHeader);
-	Verse::VCell* ConditionalCreateCellExport(const FAsyncPackageHeaderData& Header, int32 CellExportIndex);
-	Verse::VCell* ConditionalSerializeCellExport(const FAsyncPackageHeaderData& Header, int32 CellExportIndex);
-	Verse::VCell* ConditionalCreateCellImport(const FAsyncPackageHeaderData& Header, int32 CellImportIndex);
-	Verse::VCell* ConditionalSerializeCellImport(const FAsyncPackageHeaderData& Header, int32 CellImportIndex);
+	void ConditionalCreateCellExport(const FAsyncPackageHeaderData& Header, int32 CellExportIndex);
+	void ConditionalSerializeCellExport(const FAsyncPackageHeaderData& Header, int32 CellExportIndex);
+	void ConditionalCreateCellImport(const FAsyncPackageHeaderData& Header, int32 CellImportIndex);
+	void ConditionalSerializeCellImport(const FAsyncPackageHeaderData& Header, int32 CellImportIndex);
 
 	/**
 	 * Begin async loading process. Simulates parts of BeginLoad.
@@ -5473,7 +5512,6 @@ FAsyncPackage2* FAsyncLoadingThread2::FindOrInsertPackage(FAsyncLoadingThreadSta
 			{
 				Package = CreateAsyncPackage(ThreadState, Desc);
 				checkf(Package, TEXT("Failed to create async package %s"), *Desc.UPackageName.ToString());
-				Package->AddRef();
 				++LoadingPackagesCounter;
 				TRACE_COUNTER_SET(AsyncLoadingLoadingPackages, LoadingPackagesCounter);
 				bInserted = true;
@@ -7950,7 +7988,6 @@ bool FAsyncPackage2::CreateRuntimeExports(FAsyncLoadingThreadState2& ThreadState
 	UE::Core::Private::FPlayInEditorLoadingScope PlayInEditorIDScope(Desc.PIEInstanceID);
 #endif
 
-	UE::TUniqueLock ScopeLock(SerializationMutex);
 	while (ProcessedExportBundlesCount < Data.TotalExportBundleCount)
 	{
 		UE_ASYNC_PACKAGE_LOG(VeryVerbose, Desc, TEXT("CreateRuntimeExports"), TEXT("Beginning Creating Exports for runtime bundle %d"), ProcessedExportBundlesCount);
@@ -8058,7 +8095,6 @@ bool FAsyncPackage2::PreloadRuntimeExports(FAsyncLoadingThreadState2& ThreadStat
 	UE::Core::Private::FPlayInEditorLoadingScope PlayInEditorIDScope(Desc.PIEInstanceID);
 #endif
 
-	UE::TUniqueLock ScopeLock(SerializationMutex);
 	while (ProcessedExportBundlesCount < Data.TotalExportBundleCount)
 	{
 		UE_ASYNC_PACKAGE_LOG(VeryVerbose, Desc, TEXT("PreloadRuntimeExports"), TEXT("Beginning Preloads of runtime export bundle %d"), ProcessedExportBundlesCount);
@@ -8130,15 +8166,8 @@ bool FAsyncPackage2::PreloadRuntimeExports(FAsyncLoadingThreadState2& ThreadStat
 		if (++ProcessedExportBundlesCount == Data.TotalExportBundleCount)
 		{
 			ProcessedExportBundlesCount = 0;
-			HeaderData.Reset();
-			SerializationState.ReleaseIoRequest();
-#if WITH_EDITOR
-			if (OptionalSegmentHeaderData.IsSet())
-			{
-				OptionalSegmentHeaderData->Reset();
-				OptionalSegmentSerializationState->ReleaseIoRequest();
-			}
-#endif
+			// We own one ref by default since the ref is initialized to 1.
+			ReleaseHeaderRef();
 			return true;
 		}
 	}
@@ -8298,17 +8327,16 @@ int32 FAsyncPackage2::GetPublicExportIndex(uint64 ExportHash, FAsyncPackageHeade
 	return -1;
 }
 
-UObject* FAsyncPackage2::ConditionalCreateExport(const FAsyncPackageHeaderData& Header, int32 LocalExportIndex)
+void FAsyncPackage2::ConditionalCreateExport(const FAsyncPackageHeaderData& Header, int32 LocalExportIndex)
 {
 	if (!Header.ExportsView[LocalExportIndex].Object)
 	{
 		FAsyncPackageScope2 Scope(this);
 		EventDrivenCreateExport(Header, LocalExportIndex);
 	}
-	return Header.ExportsView[LocalExportIndex].Object;
 }
 
-UObject* FAsyncPackage2::ConditionalSerializeExport(const FAsyncPackageHeaderData& Header, int32 LocalExportIndex)
+void FAsyncPackage2::ConditionalSerializeExport(const FAsyncPackageHeaderData& Header, int32 LocalExportIndex)
 {
 	FExportObject& Export = Header.ExportsView[LocalExportIndex];
 
@@ -8319,7 +8347,7 @@ UObject* FAsyncPackage2::ConditionalSerializeExport(const FAsyncPackageHeaderDat
 
 	if (!Export.Object || (Export.bFiltered || Export.bExportLoadFailed))
 	{
-		return nullptr;
+		return;
 	}
 
 	if (Export.Object->HasAllFlags(RF_NeedLoad))
@@ -8327,37 +8355,39 @@ UObject* FAsyncPackage2::ConditionalSerializeExport(const FAsyncPackageHeaderDat
 		FAsyncPackageScope2 Scope(this);
 		EventDrivenSerializeExport(Header, LocalExportIndex, nullptr);
 	}
-
-	return Export.Object;
 }
 
-UObject* FAsyncPackage2::ConditionalCreateImport(const FAsyncPackageHeaderData& Header, int32 LocalImportIndex)
+void FAsyncPackage2::ConditionalCreateImport(const FAsyncPackageHeaderData& Header, int32 LocalImportIndex)
 {
 	const FPackageObjectIndex& ObjectIndex = Header.ImportMap[LocalImportIndex];
 	check(ObjectIndex.IsPackageImport());
 	if (UObject* FromImportStore = ImportStore.GetImportObject(Header, ObjectIndex))
 	{
-		return FromImportStore;
+		return;
 	}
 
 	FPackageImportReference PackageImportRef = ObjectIndex.ToPackageImportRef();
 	FAsyncPackage2* SourcePackage = Header.ImportedAsyncPackagesView[PackageImportRef.GetImportedPackageIndex()];
 	if (!SourcePackage)
 	{
-		return nullptr;
+		return;
 	}
 	uint64 ExportHash = Header.ImportedPublicExportHashes[PackageImportRef.GetImportedPublicExportHashIndex()];
 	FAsyncPackageHeaderData* SourcePackageHeader = nullptr;
-	UE::TUniqueLock ScopeLock(SourcePackage->SerializationMutex);
-	int32 ExportIndex = SourcePackage->GetPublicExportIndex(ExportHash, SourcePackageHeader);
-	if (ExportIndex < 0)
+
+	// If the header ref is at 0, it means all the serialization is done for that package, we just raced with it and missed the cue.
+	if (SourcePackage->TryAddHeaderRef())
 	{
-		return nullptr;
+		int32 ExportIndex = SourcePackage->GetPublicExportIndex(ExportHash, SourcePackageHeader);
+		if (ExportIndex >= 0)
+		{
+			SourcePackage->ConditionalCreateExport(*SourcePackageHeader, ExportIndex);
+		}
+		SourcePackage->ReleaseHeaderRef();
 	}
-	return SourcePackage->ConditionalCreateExport(*SourcePackageHeader, ExportIndex);
 }
 
-UObject* FAsyncPackage2::ConditionalSerializeImport(const FAsyncPackageHeaderData& Header, int32 LocalImportIndex)
+void FAsyncPackage2::ConditionalSerializeImport(const FAsyncPackageHeaderData& Header, int32 LocalImportIndex)
 {
 	const FPackageObjectIndex& ObjectIndex = Header.ImportMap[LocalImportIndex];
 	check(ObjectIndex.IsPackageImport());
@@ -8366,7 +8396,7 @@ UObject* FAsyncPackage2::ConditionalSerializeImport(const FAsyncPackageHeaderDat
 	{
 		if (!FromImportStore->HasAllFlags(RF_NeedLoad))
 		{
-			return FromImportStore;
+			return;
 		}
 	}
 
@@ -8374,17 +8404,21 @@ UObject* FAsyncPackage2::ConditionalSerializeImport(const FAsyncPackageHeaderDat
 	FAsyncPackage2* SourcePackage = Header.ImportedAsyncPackagesView[PackageImportRef.GetImportedPackageIndex()];
 	if (!SourcePackage)
 	{
-		return nullptr;
+		return;
 	}
 	uint64 ExportHash = Header.ImportedPublicExportHashes[PackageImportRef.GetImportedPublicExportHashIndex()];
 	FAsyncPackageHeaderData* SourcePackageHeader = nullptr;
-	UE::TUniqueLock ScopeLock(SourcePackage->SerializationMutex);
-	int32 ExportIndex = SourcePackage->GetPublicExportIndex(ExportHash, SourcePackageHeader);
-	if (ExportIndex < 0)
+
+	// If the header ref is at 0, it means all the serialization is done for that package, we just raced with it and missed the cue.
+	if (SourcePackage->TryAddHeaderRef())
 	{
-		return nullptr;
+		int32 ExportIndex = SourcePackage->GetPublicExportIndex(ExportHash, SourcePackageHeader);
+		if (ExportIndex >= 0)
+		{
+			SourcePackage->ConditionalSerializeExport(*SourcePackageHeader, ExportIndex);
+		}
+		SourcePackage->ReleaseHeaderRef();
 	}
-	return SourcePackage->ConditionalSerializeExport(*SourcePackageHeader, ExportIndex);
 }
 
 int32 FAsyncPackage2::GetPublicCellExportIndex(uint64 CellExportHash, FAsyncPackageHeaderData*& OutHeader)
@@ -8400,17 +8434,16 @@ int32 FAsyncPackage2::GetPublicCellExportIndex(uint64 CellExportHash, FAsyncPack
 	return -1;
 }
 
-Verse::VCell* FAsyncPackage2::ConditionalCreateCellExport(const FAsyncPackageHeaderData& Header, int32 CellExportIndex)
+void FAsyncPackage2::ConditionalCreateCellExport(const FAsyncPackageHeaderData& Header, int32 CellExportIndex)
 {
 	if (!Header.CellExportsView[CellExportIndex].Cell)
 	{
 		FAsyncPackageScope2 Scope(this);
 		EventDrivenCreateCellExport(Header, CellExportIndex, nullptr);
 	}
-	return Header.CellExportsView[CellExportIndex].Cell;
 }
 
-Verse::VCell* FAsyncPackage2::ConditionalSerializeCellExport(const FAsyncPackageHeaderData& Header, int32 CellExportIndex)
+void FAsyncPackage2::ConditionalSerializeCellExport(const FAsyncPackageHeaderData& Header, int32 CellExportIndex)
 {
 	ConditionalCreateCellExport(Header, CellExportIndex);
 
@@ -8419,36 +8452,39 @@ Verse::VCell* FAsyncPackage2::ConditionalSerializeCellExport(const FAsyncPackage
 		FAsyncPackageScope2 Scope(this);
 		EventDrivenSerializeCellExport(Header, CellExportIndex, nullptr);
 	}
-	return Header.CellExportsView[CellExportIndex].Cell;
 }
 
-Verse::VCell* FAsyncPackage2::ConditionalCreateCellImport(const FAsyncPackageHeaderData& Header, int32 CellImportIndex)
+void FAsyncPackage2::ConditionalCreateCellImport(const FAsyncPackageHeaderData& Header, int32 CellImportIndex)
 {
 	const FPackageObjectIndex& CellIndex = Header.CellImportMap[CellImportIndex];
 	check(CellIndex.IsPackageImport());
 	if (Verse::VCell* FromImportStore = ImportStore.FindOrGetImportCell(Header, CellIndex))
 	{
-		return FromImportStore;
+		return;
 	}
 
 	FPackageImportReference PackageImportRef = CellIndex.ToPackageImportRef();
 	FAsyncPackage2* SourcePackage = Header.ImportedAsyncPackagesView[PackageImportRef.GetImportedPackageIndex()];
 	if (!SourcePackage)
 	{
-		return nullptr;
+		return;
 	}
 	uint64 CellExportHash = Header.ImportedPublicExportHashes[PackageImportRef.GetImportedPublicExportHashIndex()];
 	FAsyncPackageHeaderData* SourcePackageHeader = nullptr;
-	UE::TUniqueLock ScopeLock(SourcePackage->SerializationMutex);
-	int32 CellExportIndex = SourcePackage->GetPublicCellExportIndex(CellExportHash, SourcePackageHeader);
-	if (CellExportIndex < 0)
+
+	// If the header ref is at 0, it means all the serialization is done for that package, we just raced with it and missed the cue.
+	if (SourcePackage->TryAddHeaderRef())
 	{
-		return nullptr;
+		int32 CellExportIndex = SourcePackage->GetPublicCellExportIndex(CellExportHash, SourcePackageHeader);
+		if (CellExportIndex >= 0)
+		{
+			SourcePackage->ConditionalCreateCellExport(*SourcePackageHeader, CellExportIndex);
+		}
+		SourcePackage->ReleaseHeaderRef();
 	}
-	return SourcePackage->ConditionalCreateCellExport(*SourcePackageHeader, CellExportIndex);
 }
 
-Verse::VCell* FAsyncPackage2::ConditionalSerializeCellImport(const FAsyncPackageHeaderData& Header, int32 CellImportIndex)
+void FAsyncPackage2::ConditionalSerializeCellImport(const FAsyncPackageHeaderData& Header, int32 CellImportIndex)
 {
 	const FPackageObjectIndex& CellIndex = Header.CellImportMap[CellImportIndex];
 	check(CellIndex.IsPackageImport());
@@ -8457,17 +8493,21 @@ Verse::VCell* FAsyncPackage2::ConditionalSerializeCellImport(const FAsyncPackage
 	FAsyncPackage2* SourcePackage = Header.ImportedAsyncPackagesView[PackageImportRef.GetImportedPackageIndex()];
 	if (!SourcePackage)
 	{
-		return nullptr;
+		return;
 	}
 	uint64 CellExportHash = Header.ImportedPublicExportHashes[PackageImportRef.GetImportedPublicExportHashIndex()];
 	FAsyncPackageHeaderData* SourcePackageHeader = nullptr;
-	UE::TUniqueLock ScopeLock(SourcePackage->SerializationMutex);
-	int32 CellExportIndex = SourcePackage->GetPublicCellExportIndex(CellExportHash, SourcePackageHeader);
-	if (CellExportIndex < 0)
+
+	// If the header ref is at 0, it means all the serialization is done for that package, we just raced with it and missed the cue.
+	if (SourcePackage->TryAddHeaderRef())
 	{
-		return nullptr;
+		int32 CellExportIndex = SourcePackage->GetPublicCellExportIndex(CellExportHash, SourcePackageHeader);
+		if (CellExportIndex >= 0)
+		{
+			SourcePackage->ConditionalSerializeCellExport(*SourcePackageHeader, CellExportIndex);
+		}
+		SourcePackage->ReleaseHeaderRef();
 	}
-	return SourcePackage->ConditionalSerializeCellExport(*SourcePackageHeader, CellExportIndex);
 }
 
 FString FAsyncPackage2::GetNameFromPackageObjectIndex(const FAsyncPackageHeaderData& Header, FPackageObjectIndex Index)
@@ -8658,6 +8698,27 @@ bool FAsyncPackage2::CanRunNodeAsync(const FAsyncLoadEventSpec* Spec) const
 	return true;
 }
 
+class FSerializationMutex
+{
+private:
+	UE::FMutex& Mutex;
+public:
+	FSerializationMutex(UE::FMutex& InMutex)
+		: Mutex(InMutex)
+	{
+		if (!Mutex.TryLock())
+		{
+			TRACE_CPUPROFILER_EVENT_SCOPE(SerializationMutexContention);
+			Mutex.Lock();
+		}
+	}
+
+	~FSerializationMutex()
+	{
+		Mutex.Unlock();
+	}
+};
+
 void FAsyncPackage2::EventDrivenCreateExport(const FAsyncPackageHeaderData& Header, int32 LocalExportIndex)
 {
 	TRACE_CPUPROFILER_EVENT_SCOPE(CreateExport);
@@ -8665,7 +8726,6 @@ void FAsyncPackage2::EventDrivenCreateExport(const FAsyncPackageHeaderData& Head
 	const FExportMapEntry& Export = Header.ExportMap[LocalExportIndex];
 	FExportObject& ExportObject = Header.ExportsView[LocalExportIndex];
 	UObject*& Object = ExportObject.Object;
-	check(!Object);
 
 	TRACE_LOADTIME_CREATE_EXPORT_SCOPE(this, &Object);
 
@@ -8691,6 +8751,14 @@ void FAsyncPackage2::EventDrivenCreateExport(const FAsyncPackageHeaderData& Head
 
 	ProcessExportDependencies(Header, LocalExportIndex, FExportBundleEntry::ExportCommandType_Create);
 
+	FSerializationMutex ScopeLock(SerializationMutex);
+
+	// Another thread might have done it for us while we were waiting...
+	if (Object)
+	{
+		return;
+	}
+
 	bool bIsCompleteyLoaded = false;
 	UClass* LoadClass = Export.ClassIndex.IsNull() ? UClass::StaticClass() : CastEventDrivenIndexToObject<UClass>(Header, Export.ClassIndex, true);
 	UObject* ThisParent = Export.OuterIndex.IsNull() ? LinkerRoot : EventDrivenIndexToObject(Header, Export.OuterIndex, false);
@@ -8882,7 +8950,7 @@ void FAsyncPackage2::EventDrivenCreateExport(const FAsyncPackageHeaderData& Head
 	Object->SetInternalFlags(FlagsToSet);
 }
 
-bool FAsyncPackage2::EventDrivenSerializeExport(const FAsyncPackageHeaderData& Header, int32 LocalExportIndex, FExportArchive* Ar)
+void FAsyncPackage2::EventDrivenSerializeExport(const FAsyncPackageHeaderData& Header, int32 LocalExportIndex, FExportArchive* Ar)
 {
 	LLM_SCOPE(ELLMTag::UObject);
 	LLM_SCOPE_BYTAG(UObject_FAsyncPackage2);
@@ -8912,7 +8980,17 @@ bool FAsyncPackage2::EventDrivenSerializeExport(const FAsyncPackageHeaderData& H
 			UE_ASYNC_PACKAGE_LOG_VERBOSE(VeryVerbose, Desc, TEXT("SerializeExport"),
 				TEXT("Skipped already serialized export %s"), *Header.NameMap.GetName(Export.ObjectName).ToString());
 		}
-		return false;
+		return;
+	}
+
+	ProcessExportDependencies(Header, LocalExportIndex, FExportBundleEntry::ExportCommandType_Serialize);
+
+	FSerializationMutex ScopeLock(SerializationMutex);
+
+	// Another thread might have done it for us while we were waiting...
+	if (ExportObject.Object && !ExportObject.Object->HasAnyFlags(RF_NeedLoad))
+	{
+		return;
 	}
 
 	TOptional<FExportArchive> LocalAr;
@@ -8932,8 +9010,6 @@ bool FAsyncPackage2::EventDrivenSerializeExport(const FAsyncPackageHeaderData& H
 		}
 	}
 
-	ProcessExportDependencies(Header, LocalExportIndex, FExportBundleEntry::ExportCommandType_Serialize);
-
 	// If this is a struct, make sure that its parent struct is completely loaded
 	if (UStruct* Struct = dynamic_cast<UStruct*>(Object))
 	{
@@ -9022,8 +9098,6 @@ bool FAsyncPackage2::EventDrivenSerializeExport(const FAsyncPackageHeaderData& H
 
 	// push stats so that we don't overflow number of tags per thread during blocking loading
 	LLM_PUSH_STATS_FOR_ASSET_TAGS();
-
-	return true;
 }
 
 void FAsyncPackage2::EventDrivenCreateCellExport(const FAsyncPackageHeaderData& Header, int32 LocalExportIndex, FExportArchive* Ar)
@@ -9031,9 +9105,18 @@ void FAsyncPackage2::EventDrivenCreateCellExport(const FAsyncPackageHeaderData&
 	const int32 CellExportIndex = LocalExportIndex - Header.ExportMap.Num();
 	const FCellExportMapEntry& Export = Header.CellExportMap[CellExportIndex];
 	FExportCell& ExportCell = Header.CellExportsView[CellExportIndex];
-	Verse::VCell* Cell = ExportCell.Cell;
-	check(!Cell);
 
+	ProcessExportDependencies(Header, LocalExportIndex, FExportBundleEntry::ExportCommandType_Create);
+
+	FSerializationMutex ScopeLock(SerializationMutex);
+
+	// Another thread might have done it for us while we were waiting...
+	if (ExportCell.Cell)
+	{
+		return;
+	}
+
+	Verse::VCell* Cell = ExportCell.Cell;
 	TOptional<FExportArchive> LocalAr;
 	if (!Ar)
 	{
@@ -9041,8 +9124,6 @@ void FAsyncPackage2::EventDrivenCreateCellExport(const FAsyncPackageHeaderData&
 		InitializeExportArchive(*Ar, false);
 	}
 
-	ProcessExportDependencies(Header, LocalExportIndex, FExportBundleEntry::ExportCommandType_Create);
-
 	Ar->ExportBufferBegin(nullptr, Export.CookedSerialOffset, Export.CookedSerialLayoutSize);
 
 	FName CppClassInfoName = Header.NameMap.GetName(Export.CppClassInfo);
@@ -9066,7 +9147,7 @@ void FAsyncPackage2::EventDrivenCreateCellExport(const FAsyncPackageHeaderData&
 	}
 }
 
-bool FAsyncPackage2::EventDrivenSerializeCellExport(const FAsyncPackageHeaderData& Header, int32 LocalExportIndex, FExportArchive* Ar)
+void FAsyncPackage2::EventDrivenSerializeCellExport(const FAsyncPackageHeaderData& Header, int32 LocalExportIndex, FExportArchive* Ar)
 {
 	const int32 CellExportIndex = LocalExportIndex - Header.ExportMap.Num();
 	const FCellExportMapEntry& Export = Header.CellExportMap[CellExportIndex];
@@ -9074,6 +9155,16 @@ bool FAsyncPackage2::EventDrivenSerializeCellExport(const FAsyncPackageHeaderDat
 	Verse::VCell* Cell = ExportCell.Cell;
 	check(Cell);
 
+	ProcessExportDependencies(Header, LocalExportIndex, FExportBundleEntry::ExportCommandType_Serialize);
+
+	FSerializationMutex ScopeLock(SerializationMutex);
+
+	// Another thread might have done it for us while we were waiting...
+	if (ExportCell.bSerialized)
+	{
+		return;
+	}
+
 	TOptional<FExportArchive> LocalAr;
 	if (!Ar)
 	{
@@ -9081,8 +9172,6 @@ bool FAsyncPackage2::EventDrivenSerializeCellExport(const FAsyncPackageHeaderDat
 		InitializeExportArchive(*Ar, false);
 	}
 
-	ProcessExportDependencies(Header, LocalExportIndex, FExportBundleEntry::ExportCommandType_Serialize);
-
 	ExportCell.bSerialized = true;
 
 	Ar->ExportBufferBegin(nullptr, Export.CookedSerialOffset, Export.CookedSerialSize);
@@ -9096,8 +9185,6 @@ bool FAsyncPackage2::EventDrivenSerializeCellExport(const FAsyncPackageHeaderDat
 #endif
 
 	Ar->ExportBufferEnd();
-
-	return true;
 }
 
 FAsyncPackage2* FAsyncPackage2::GetCurrentlyExecutingPackage(FAsyncLoadingThreadState2& ThreadState, FAsyncPackage2* PackageToFilter)
