Add TSubScriptStructOf. Similar to TSubClassOf but for UScriptStruct. Add an editor picker that works with NoClear and the metadatas
* MetaStruct: The struct path of the base struct displayed in the picker.
* ShowTreeView: The picker uses a tree view or a list view to display the available struct.
* ShowDisplayNames: Uses the DisplayName of the selectable struct.
* HideViewOptions: Hides the options menu in the picker.

diff --git a/Engine/Source/Editor/DetailCustomizations/Private/DetailCustomizations.cpp b/Engine/Source/Editor/DetailCustomizations/Private/DetailCustomizations.cpp
--- a/Engine/Source/Editor/DetailCustomizations/Private/DetailCustomizations.cpp
+++ b/Engine/Source/Editor/DetailCustomizations/Private/DetailCustomizations.cpp
@@ -150,6 +150,7 @@
 #include "SplineComponentDetails.h"
 #include "StaticMeshActorDetails.h"
 #include "StaticMeshComponentDetails.h"
+#include "SubScriptStructOfCustomization.h"
 #include "SupportedRangeTypes.h"	// StructsSupportingRangeVisibility
 #include "TemplateStringStructCustomization.h"
 #include "Templates/SharedPointer.h"
@@ -248,6 +249,7 @@ void FDetailCustomizationsModule::RegisterPropertyTypeCustomizations()
 {
 	RegisterCustomPropertyTypeLayout("SoftObjectPath", FOnGetPropertyTypeCustomizationInstance::CreateStatic(&FSoftObjectPathCustomization::MakeInstance));
 	RegisterCustomPropertyTypeLayout("SoftClassPath", FOnGetPropertyTypeCustomizationInstance::CreateStatic(&FSoftClassPathCustomization::MakeInstance));
+	RegisterCustomPropertyTypeLayout("SubScriptStructOf", FOnGetPropertyTypeCustomizationInstance::CreateStatic(&FSubScriptStructOfCustomization::MakeInstance));
 	RegisterCustomPropertyTypeLayout("DataTableRowHandle", FOnGetPropertyTypeCustomizationInstance::CreateStatic(&FDataTableCustomizationLayout::MakeInstance));
 	RegisterCustomPropertyTypeLayout("DataTableCategoryHandle", FOnGetPropertyTypeCustomizationInstance::CreateStatic(&FDataTableCategoryCustomizationLayout::MakeInstance));
 	RegisterCustomPropertyTypeLayout("CurveTableRowHandle", FOnGetPropertyTypeCustomizationInstance::CreateStatic(&FCurveTableCustomizationLayout::MakeInstance));
diff --git a/Engine/Source/Editor/DetailCustomizations/Private/SubScriptStructOfCustomization.cpp b/Engine/Source/Editor/DetailCustomizations/Private/SubScriptStructOfCustomization.cpp
new file mode 100644
--- /dev/null
+++ b/Engine/Source/Editor/DetailCustomizations/Private/SubScriptStructOfCustomization.cpp
@@ -0,0 +1,98 @@
+// Copyright Epic Games, Inc. All Rights Reserved.
+
+#include "SubScriptStructOfCustomization.h"
+
+#include "DetailWidgetRow.h"
+#include "EditorClassUtils.h"
+#include "HAL/PlatformCrt.h"
+#include "PropertyCustomizationHelpers.h"
+#include "PropertyEditorModule.h"
+#include "PropertyHandle.h"
+#include "Templates/SubScriptStructOf.h"
+#include "UObject/Object.h"
+#include "UObject/ObjectMacros.h"
+#include "UObject/UnrealType.h"
+
+class IDetailChildrenBuilder;
+
+void FSubScriptStructOfCustomization::CustomizeHeader(TSharedRef<IPropertyHandle> InPropertyHandle, FDetailWidgetRow& HeaderRow, IPropertyTypeCustomizationUtils& StructCustomizationUtils)
+{
+	PropertyHandle = InPropertyHandle;
+
+	const FString& MetaStructName = PropertyHandle->GetMetaData("MetaStruct");
+	const bool bAllowNone = !(PropertyHandle->GetMetaDataProperty()->PropertyFlags & CPF_NoClear);
+	const bool bShowTreeView = PropertyHandle->HasMetaData("ShowTreeView");
+	const bool bHideViewOptions = PropertyHandle->HasMetaData("HideViewOptions");
+	const bool bShowDisplayNames = PropertyHandle->HasMetaData("ShowDisplayNames");
+	
+
+	const UScriptStruct* MetaScriptStruct = nullptr;
+	if (!MetaStructName.IsEmpty())
+	{
+		TArray<const UScriptStruct*> Structs = PropertyCustomizationHelpers::GetStructsFromMetadataString(MetaStructName);
+		if (Structs.Num() >= 1)
+		{
+			MetaScriptStruct = Structs[0];
+		}
+	}
+
+	HeaderRow
+	.NameContent()
+	[
+		InPropertyHandle->CreatePropertyNameWidget()
+	]
+	.ValueContent()
+	.MinDesiredWidth(250.0f)
+	.MaxDesiredWidth(0.0f)
+	[
+		SNew(SStructPropertyEntryBox)
+			.MetaStruct(MetaScriptStruct)
+			.AllowNone(bAllowNone)
+			.HideViewOptions(bHideViewOptions)
+			.ShowDisplayNames(bShowDisplayNames)
+			.ShowTreeView(bShowTreeView)
+			.SelectedStruct(this, &FSubScriptStructOfCustomization::HandleGetScriptStruct)
+			.OnSetStruct(this, &FSubScriptStructOfCustomization::HandleSetScriptStruct)
+	];
+}
+
+void FSubScriptStructOfCustomization::CustomizeChildren(TSharedRef<IPropertyHandle> InStructPropertyHandle, IDetailChildrenBuilder& StructBuilder, IPropertyTypeCustomizationUtils& StructCustomizationUtils)
+{
+}
+
+const UScriptStruct* FSubScriptStructOfCustomization::HandleGetScriptStruct() const
+{
+	if (FStructProperty* StructProperty = CastField<FStructProperty>(PropertyHandle->GetProperty()))
+	{
+		check(StructProperty->Struct == TBaseStructure<FSubScriptStructOf>::Get());
+
+		TArray<void*> RawData;
+		PropertyHandle->AccessRawData(RawData);
+
+		if (RawData.Num() >= 1 && RawData[0])
+		{
+			return static_cast<FSubScriptStructOf*>(RawData[0])->Get();
+		}
+	}
+	return nullptr;
+}
+
+void FSubScriptStructOfCustomization::HandleSetScriptStruct(const UScriptStruct* NewStruct)
+{
+	if (FStructProperty* StructProperty = CastField<FStructProperty>(PropertyHandle->GetProperty()))
+	{
+		check(StructProperty->Struct == TBaseStructure<FSubScriptStructOf>::Get());
+
+		FSubScriptStructOf DefaultSubScriptStructOf;
+
+		TArray<void*> RawData;
+		PropertyHandle->AccessRawData(RawData);
+		FSubScriptStructOf* PreviousValue = RawData.Num() == 1 ? static_cast<FSubScriptStructOf*>(RawData[0]) : &DefaultSubScriptStructOf;
+
+		FSubScriptStructOf NewValue = const_cast<UScriptStruct*>(NewStruct);
+
+		FString TextValue;
+		StructProperty->Struct->ExportText(TextValue, &NewValue, PreviousValue, nullptr, EPropertyPortFlags::PPF_None, nullptr);
+		ensure(PropertyHandle->SetValueFromFormattedString(TextValue, EPropertyValueSetFlags::DefaultFlags) == FPropertyAccess::Result::Success);
+	}
+}
diff --git a/Engine/Source/Editor/DetailCustomizations/Private/SubScriptStructOfCustomization.h b/Engine/Source/Editor/DetailCustomizations/Private/SubScriptStructOfCustomization.h
new file mode 100644
--- /dev/null
+++ b/Engine/Source/Editor/DetailCustomizations/Private/SubScriptStructOfCustomization.h
@@ -0,0 +1,37 @@
+// Copyright Epic Games, Inc. All Rights Reserved.
+
+#pragma once
+
+#include "Containers/UnrealString.h"
+#include "IPropertyTypeCustomization.h"
+#include "Templates/SharedPointer.h"
+#include "UObject/WeakObjectPtr.h"
+#include "UObject/WeakObjectPtrTemplates.h"
+
+class IPropertyHandle;
+class UClass;
+
+/**
+ * Customizes a SubScriptStructOf reference to look like a UScriptStruct* property
+ */
+class FSubScriptStructOfCustomization : public IPropertyTypeCustomization
+{
+public:
+	static TSharedRef<IPropertyTypeCustomization> MakeInstance() 
+	{
+		return MakeShareable(new FSubScriptStructOfCustomization);
+	}
+
+	/** IPropertyTypeCustomization interface */
+	virtual void CustomizeHeader(TSharedRef<class IPropertyHandle> InPropertyHandle, class FDetailWidgetRow& HeaderRow, IPropertyTypeCustomizationUtils& StructCustomizationUtils) override;
+	virtual void CustomizeChildren(TSharedRef<class IPropertyHandle> InPropertyHandle, class IDetailChildrenBuilder& StructBuilder, IPropertyTypeCustomizationUtils& StructCustomizationUtils) override;
+
+private:
+	/** @return The struct currently set on this reference */
+	const UScriptStruct* HandleGetScriptStruct() const;
+	/** Set the struct used by this reference */
+	void HandleSetScriptStruct(const UScriptStruct* newStruct);
+
+	/** Handle to the property being customized */
+	TSharedPtr<IPropertyHandle> PropertyHandle;
+};
diff --git a/Engine/Source/Runtime/CoreUObject/Public/Templates/SubScriptStructOf.h b/Engine/Source/Runtime/CoreUObject/Public/Templates/SubScriptStructOf.h
new file mode 100644
--- /dev/null
+++ b/Engine/Source/Runtime/CoreUObject/Public/Templates/SubScriptStructOf.h
@@ -0,0 +1,198 @@
+// Copyright Epic Games, Inc. All Rights Reserved.
+
+#pragma once
+
+#include "Templates/UnrealTypeTraits.h"
+#include "UObject/Class.h"
+
+#include <type_traits>
+
+#include "SubScriptStructOf.generated.h"
+
+class FArchive;
+class FStructuredArchiveSlot;
+
+/**
+ * Allow UScriptStruct types to be passed around with type safety
+ */
+USTRUCT()
+struct FSubScriptStructOf
+{
+	GENERATED_BODY()
+
+	FSubScriptStructOf() = default;
+
+	/** Constructor that takes a UScriptStruct*. */
+	[[nodiscard]] UE_REWRITE FSubScriptStructOf(UScriptStruct* From)
+		: ScriptStruct(From)
+	{
+	}
+
+	/** Assign from a UScriptStruct*. */
+	UE_REWRITE FSubScriptStructOf& operator=(UScriptStruct* From)
+	{
+		ScriptStruct = From;
+		return *this;
+	}
+
+	/** Dereference back into a UScriptStruct*. */
+	[[nodiscard]] UE_REWRITE UScriptStruct* operator*() const
+	{
+		return ScriptStruct;
+	}
+
+	/** Dereference back into a UScriptStruct*. */
+	[[nodiscard]] UE_REWRITE UScriptStruct* Get() const
+	{
+		return **this;
+	}
+
+	/** Dereference back into a UScriptStruct*. */
+	[[nodiscard]] UE_REWRITE UScriptStruct* operator->() const
+	{
+		return **this;
+	}
+
+	/** Implicit conversion to UScriptStruct*. */
+	[[nodiscard]] UE_REWRITE operator UScriptStruct* () const
+	{
+		return **this;
+	}
+
+	UE_REWRITE void Serialize(FArchive& Ar)
+	{
+		Ar << ScriptStruct;
+	}
+
+	UE_REWRITE void Serialize(FStructuredArchiveSlot& Slot)
+	{
+		Slot << ScriptStruct;
+	}
+
+	[[nodiscard]] friend uint32 GetTypeHash(const FSubScriptStructOf& SubScriptStructOf)
+	{
+		return GetTypeHash(SubScriptStructOf.ScriptStruct);
+	}
+
+private:
+	UPROPERTY()
+	TObjectPtr<UScriptStruct> ScriptStruct;
+};
+
+template <typename T>
+struct TSubScriptStructOf;
+
+template <typename T>
+struct TIsTSubScriptStructOf
+{
+	static constexpr bool Value = false;
+};
+
+template <typename T> struct TIsTSubScriptStructOf<               TSubScriptStructOf<T>> { static constexpr bool Value = true; };
+template <typename T> struct TIsTSubScriptStructOf<const          TSubScriptStructOf<T>> { static constexpr bool Value = true; };
+template <typename T> struct TIsTSubScriptStructOf<      volatile TSubScriptStructOf<T>> { static constexpr bool Value = true; };
+template <typename T> struct TIsTSubScriptStructOf<const volatile TSubScriptStructOf<T>> { static constexpr bool Value = true; };
+
+template <typename T>
+inline constexpr bool TIsTSubScriptStructOf_V = TIsTSubScriptStructOf<T>::Value;
+
+
+/**
+ * Template to allow UScriptStruct types to be passed around with type safety
+ */
+template <typename T>
+struct TSubScriptStructOf : FSubScriptStructOf
+{
+private:
+	template <typename U>
+	friend struct TSubScriptStructOf;
+
+public:
+	using ElementType = T;
+
+	TSubScriptStructOf() = default;
+
+	/** Constructor that takes a UScriptStruct*. */
+	UE_REWRITE TSubScriptStructOf(UScriptStruct* From)
+		: FSubScriptStructOf(From)
+	{
+	}
+
+	/** Construct from a UScriptStruct* (or something implicitly convertible to it). */
+	template <typename U>
+		requires (UE::CConvertibleTo<U, UScriptStruct*> && !TIsTSubScriptStructOf_V<std::decay_t<U>>)
+	UE_REWRITE TSubScriptStructOf(U&& From)
+		: FSubScriptStructOf(From)
+	{
+	}
+
+	/** Construct from another TSubScriptStructOf, only if types are compatible. */
+	template <typename U>
+		requires (UE::CConvertibleTo<U*, T*>)
+	UE_REWRITE TSubScriptStructOf(const TSubScriptStructOf<U>& Other)
+		: FSubScriptStructOf(Other.Get())
+	{
+		IWYU_MARKUP_IMPLICIT_CAST(U, T);
+	}
+
+	/** Assign from another TSubScriptStructOf, only if types are compatible. */
+	template <typename U>
+		requires (UE::CConvertibleTo<U*, T*>)
+	UE_REWRITE TSubScriptStructOf& operator=(const TSubScriptStructOf<U>& From)
+	{
+		IWYU_MARKUP_IMPLICIT_CAST(U, T);
+		FSubScriptStructOf::operator=(From.Get());
+		return *this;
+	}
+
+	/** Assign from a UScriptStruct*. */
+	UE_REWRITE TSubScriptStructOf& operator=(UScriptStruct* From)
+	{
+		FSubScriptStructOf::operator=(From);
+		return *this;
+	}
+
+	/** Assign from a UScriptStruct* (or something implicitly convertible to it). */
+	template <typename U>
+		requires (UE::CConvertibleTo<U, UScriptStruct*> && !TIsTSubScriptStructOf_V<std::decay_t<U>>)
+	UE_REWRITE TSubScriptStructOf& operator=(U&& From)
+	{
+		FSubScriptStructOf::operator=(From.Get());
+		return *this;
+	}
+
+	/** Dereference back into a UScriptStruct*, does runtime type checking. */
+	UE_REWRITE UScriptStruct* operator*() const
+	{
+		UScriptStruct* Local = FSubScriptStructOf::Get();
+		if (!Local || !Local->IsChildOf(T::StaticStruct()))
+		{
+			return nullptr;
+		}
+		return Local;
+	}
+
+	/** Dereference back into a UScriptStruct*, does runtime type checking. */
+	UE_REWRITE UScriptStruct* Get() const
+	{
+		return **this;
+	}
+
+	/** Dereference back into a UScriptStruct*, does runtime type checking. */
+	UE_REWRITE UScriptStruct* operator->() const
+	{
+		return **this;
+	}
+
+	/** Implicit conversion to UScriptStruct*, does runtime type checking. */
+	UE_REWRITE operator UScriptStruct* () const
+	{
+		return **this;
+	}
+};
+
+template <typename T>
+struct TCallTraits<TSubScriptStructOf<T>> : public TCallTraitsBase<TSubScriptStructOf<T>>
+{
+	using ConstPointerType = TSubScriptStructOf<const T>;
+};
diff --git a/Engine/Source/Runtime/CoreUObject/Tests/SubScriptStructOfTest.cpp b/Engine/Source/Runtime/CoreUObject/Tests/SubScriptStructOfTest.cpp
new file mode 100644
--- /dev/null
+++ b/Engine/Source/Runtime/CoreUObject/Tests/SubScriptStructOfTest.cpp
@@ -0,0 +1,49 @@
+// Copyright Epic Games, Inc. All Rights Reserved.
+
+
+#include "SubScriptStructOfTest.h"
+
+#include UE_INLINE_GENERATED_CPP_BY_NAME(SubScriptStructOfTest)
+
+#if WITH_TESTS
+
+#include "Templates/SubScriptStructOf.h"
+#include "Tests/TestHarnessAdapter.h"
+
+namespace UE
+{
+
+TEST_CASE_NAMED(FSubScriptStructOfTest, "CoreUObject::SubScriptStructOf", "[Core][UObject][EngineFilter]")
+{
+	// operator=
+	{
+		TSubScriptStructOf<FSubScriptStructOfTest_ChildA> SubOf_ChildA;
+		SubOf_ChildA = FSubScriptStructOfTest_ChildA::StaticStruct();
+		CHECK(*SubOf_ChildA == FSubScriptStructOfTest_ChildA::StaticStruct());
+		SubOf_ChildA = FSubScriptStructOfTest_GrandChildA::StaticStruct();
+		CHECK(*SubOf_ChildA == FSubScriptStructOfTest_GrandChildA::StaticStruct());
+		SubOf_ChildA = TBaseStructure<FVector>::Get();
+		CHECK(*SubOf_ChildA == nullptr);
+	}
+
+	// constructor && operator
+	{
+		TSubScriptStructOf<FSubScriptStructOfTest_Base> SubOf_Base1 = FSubScriptStructOfTest_ChildA::StaticStruct();
+		CHECK(SubOf_Base1 == FSubScriptStructOfTest_ChildA::StaticStruct());
+		TSubScriptStructOf<FSubScriptStructOfTest_Base> SubOf_Base2 = SubOf_Base1;
+		CHECK(SubOf_Base1 == SubOf_Base2);
+		SubOf_Base1 = SubOf_Base2;
+		CHECK(SubOf_Base1 == SubOf_Base2);
+		CHECK(SubOf_Base1.Get() == *SubOf_Base2);
+	}
+
+	// conversion
+	{
+		TNotNull<UScriptStruct*> NotNullPtr = FSubScriptStructOfTest_ChildA::StaticStruct();
+		TSubScriptStructOf<FSubScriptStructOfTest_ChildA> SubOf_ChildA = NotNullPtr;
+		CHECK(SubOf_ChildA == FSubScriptStructOfTest_ChildA::StaticStruct());
+	}	
+}
+
+} // UE
+#endif
diff --git a/Engine/Source/Runtime/CoreUObject/Tests/SubScriptStructOfTest.h b/Engine/Source/Runtime/CoreUObject/Tests/SubScriptStructOfTest.h
new file mode 100644
--- /dev/null
+++ b/Engine/Source/Runtime/CoreUObject/Tests/SubScriptStructOfTest.h
@@ -0,0 +1,46 @@
+// Copyright Epic Games, Inc. All Rights Reserved.
+
+#pragma once
+
+#include "UObject/Object.h"
+
+#include "SubScriptStructOfTest.generated.h"
+
+#if WITH_TESTS
+
+USTRUCT()
+struct FSubScriptStructOfTest_Base
+{
+	GENERATED_BODY()
+
+public:
+	
+};
+
+USTRUCT()
+struct FSubScriptStructOfTest_ChildA : public FSubScriptStructOfTest_Base
+{
+	GENERATED_BODY()
+
+public:
+	
+};
+
+USTRUCT()
+struct FSubScriptStructOfTest_ChildB : public FSubScriptStructOfTest_Base
+{
+	GENERATED_BODY()
+
+public:
+	
+};
+
+USTRUCT()
+struct FSubScriptStructOfTest_GrandChildA : public FSubScriptStructOfTest_ChildA
+{
+	GENERATED_BODY()
+
+public:
+};
+
+#endif
\ No newline at end of file
