RemoteObjects:  Allow us to get an FRemoteObjectId from an interface that returns a UObject*.  The issue this is attempting to solve is that existing interfaces sometimes return a raw pointer (UObject*) when we really wanted a TObjectPtr or TWeakObjectPtr so we could check its metadata (like if it's Remote or not).  The TimerManager has been updated to include the example usage.

#rnx

diff --git a/Engine/Source/Runtime/CoreUObject/Private/UObject/RemoteObject.cpp b/Engine/Source/Runtime/CoreUObject/Private/UObject/RemoteObject.cpp
--- a/Engine/Source/Runtime/CoreUObject/Private/UObject/RemoteObject.cpp
+++ b/Engine/Source/Runtime/CoreUObject/Private/UObject/RemoteObject.cpp
@@ -76,6 +76,7 @@ namespace UE::RemoteObject::Private
 std::atomic<uint64> RemoteObjectSerialNumber(1);
 std::atomic<uint64> AssetObjectSerialNumber(1);
 int32 UnsafeToMigrateObjects = 0; // This should go into TLS
+int32 GForceReturnObjectHandles = 0; // This should go into TLS
 int32 GInitRemoteServerIdBeforeUObjectInit = 0; // Initializes remote server id before UObject system is initialized
 
 class FRemoteObjectStubMap : private TMap<FRemoteObjectId, UE::RemoteObject::Handle::FRemoteObjectStub*>
@@ -479,6 +480,22 @@ bool IsSafeToMigrateObjects()
 	return !(GIsGarbageCollecting || UnsafeToMigrateObjects);
 }
 
+bool IsForceReturnObjectHandles()
+{
+	return GForceReturnObjectHandles > 0;
+}
+
+FScopedForceReturnObjectHandles::FScopedForceReturnObjectHandles()
+{
+	Private::GForceReturnObjectHandles++;
+}
+
+FScopedForceReturnObjectHandles::~FScopedForceReturnObjectHandles()
+{
+	Private::GForceReturnObjectHandles--;
+	check(Private::GForceReturnObjectHandles >= 0);
+}
+
 } // namespace UE::RemoteObject::Private
 
 namespace UE::RemoteObject::Handle
@@ -745,6 +762,18 @@ void UpdateAllPhysicsLocalIslandId(const TMap<uint32, uint32>& PhysicsLocalIslan
 #endif
 }
 
+FRemoteObjectId GetRemoteObjectId(TFunctionRef<const UObject*()> CodeThatReturnsRawUObjectPtr)
+{
+#if UE_WITH_REMOTE_OBJECT_HANDLE
+	Private::FScopedForceReturnObjectHandles DisallowMigrationsAndForceHandlesToBeReturned;
+	FObjectHandle ObjectHandle{ CodeThatReturnsRawUObjectPtr() };
+
+	return ObjectHandle.GetRemoteId();
+#else
+	return FRemoteObjectId{};
+#endif
+}
+
 UObject* ResolveObject(const FRemoteObjectStub* Stub, ERemoteReferenceType RefType /*= ERemoteReferenceType::Strong*/)
 {
 	using namespace UE::RemoteExecutor;
@@ -775,6 +804,14 @@ UObject* ResolveObject(const FRemoteObjectStub* Stub, ERemoteReferenceType RefTy
 	bool bRemoteObject = !Object || IsRemote(Object);
 	if (bRemoteObject)
 	{
+		#if UE_WITH_REMOTE_OBJECT_HANDLE
+		if (IsForceReturnObjectHandles())
+		{
+			// Re-encode this as the FObjectHandle (which has bitfield rules)
+			return reinterpret_cast<UObject*>(FObjectHandle(Stub).PointerOrHandle);
+		}
+		#endif
+
 		checkf(!GIsGarbageCollecting, TEXT("Resolving remote objects while collecting garbage is not allowed (trying to resolve object %s (%s)"), *Stub->Id.ToString(), *Stub->Name.ToString());
 
 		MigrateObjectFromRemoteServer(Stub->Id, Stub->ResidentServerId);
diff --git a/Engine/Source/Runtime/CoreUObject/Private/UObject/RemoteObjectPrivate.h b/Engine/Source/Runtime/CoreUObject/Private/UObject/RemoteObjectPrivate.h
--- a/Engine/Source/Runtime/CoreUObject/Private/UObject/RemoteObjectPrivate.h
+++ b/Engine/Source/Runtime/CoreUObject/Private/UObject/RemoteObjectPrivate.h
@@ -90,6 +90,12 @@ namespace UE::RemoteObject::Private
 		~FUnsafeToMigrateScope();
 	};
 
+	struct FScopedForceReturnObjectHandles : FNoncopyable
+	{
+		FScopedForceReturnObjectHandles();
+		~FScopedForceReturnObjectHandles();
+	};
+
 	struct FRemoteIdLocalizationHelper
 	{
 		inline static FRemoteServerId GetLocalized(FRemoteServerId InId)
@@ -119,4 +125,4 @@ namespace UE::RemoteObject::Transfer::Private
 	* Stores unreachable objects to database
 	*/
 	void StoreObjectToDatabase(UObject* Object);
-}
\ No newline at end of file
+}
diff --git a/Engine/Source/Runtime/CoreUObject/Private/UObject/RemoteObjectSerialization.cpp b/Engine/Source/Runtime/CoreUObject/Private/UObject/RemoteObjectSerialization.cpp
--- a/Engine/Source/Runtime/CoreUObject/Private/UObject/RemoteObjectSerialization.cpp
+++ b/Engine/Source/Runtime/CoreUObject/Private/UObject/RemoteObjectSerialization.cpp
@@ -7,6 +7,7 @@
 #include "UObject/ObjectPtr.h"
 #include "UObject/Class.h"
 #include "UObject/PropertyPortFlags.h"
+#include "UObject/ObjectHandle.h"
 #include "UObject/UObjectHash.h"
 #include "UObject/UObjectMigrationContext.h"
 #include "UObject/UObjectThreadContext.h"
@@ -975,7 +976,7 @@ FRemoteObjectConstructionOverrides::FRemoteObjectConstructionOverrides(const TAr
 		FRemoteObjectConstructionParams& Params = Overrides.Emplace_GetRef();
 		Params.Name = Header.Name;
 #if UE_WITH_REMOTE_OBJECT_HANDLE
-		Params.OuterId = GetRemoteObjectId(Header.Outer.GetHandle());
+		Params.OuterId = ::GetRemoteObjectId(Header.Outer.GetHandle());
 #endif
 		Params.RemoteId = Header.RemoteId;
 		Params.SerialNumber = Header.SerialNumber;
diff --git a/Engine/Source/Runtime/CoreUObject/Public/UObject/ObjectHandle.h b/Engine/Source/Runtime/CoreUObject/Public/UObject/ObjectHandle.h
--- a/Engine/Source/Runtime/CoreUObject/Public/UObject/ObjectHandle.h
+++ b/Engine/Source/Runtime/CoreUObject/Public/UObject/ObjectHandle.h
@@ -148,11 +148,11 @@ namespace UE::CoreUObject::Private
 		};
 
 		FRemoteObjectHandlePrivate() = default;
-		explicit constexpr FRemoteObjectHandlePrivate(UObject* Object)
+		explicit constexpr FRemoteObjectHandlePrivate(const UObject* Object)
 			: Pointer(Object)
 		{
 		}
-		explicit FRemoteObjectHandlePrivate(UE::RemoteObject::Handle::FRemoteObjectStub* RemoteInfo)
+		explicit FRemoteObjectHandlePrivate(const UE::RemoteObject::Handle::FRemoteObjectStub* RemoteInfo)
 		{
 			PointerOrHandle = UPTRINT(RemoteInfo) | 1;
 		}
diff --git a/Engine/Source/Runtime/CoreUObject/Public/UObject/RemoteObject.h b/Engine/Source/Runtime/CoreUObject/Public/UObject/RemoteObject.h
--- a/Engine/Source/Runtime/CoreUObject/Public/UObject/RemoteObject.h
+++ b/Engine/Source/Runtime/CoreUObject/Public/UObject/RemoteObject.h
@@ -153,6 +153,17 @@ namespace UE::RemoteObject::Handle
 		MigrateIfNeverLocal,
 	};
 
+	/**
+	 * Advanced use only.
+	 * If you are in a situation where you are using raw pointer API's that potentially cause migrations, but you need to
+	 * grab the FRemoteObjectId of that UObject, you can use this function.
+	 * For the duration of the passed-in CodeThatReturnsRawUObjectPtr, we will disable attempts to resolve remote object pointers.
+	 * The UObject that is returned from CodeThatReturnsRawUObjectPtr will be converted to a FRemoteObjectId.
+	 * @param CodeThatReturnsRawUObjectPtr a short lambda or function that accesses nothing but the single UObject it is returning.
+	 * @return the FRemoteObjectId for the returned value. It can be invalid (in the case of nullptr).
+	 */
+	COREUOBJECT_API FRemoteObjectId GetRemoteObjectId(TFunctionRef<const UObject*()> CodeThatReturnsRawUObjectPtr);
+
 	/**
 	* Resolves a remote object given its stub, aborting the active transaction if the object is unavailable
 	* @param Stub Basic data required to migrate a remote object
diff --git a/Engine/Source/Runtime/Engine/Private/TimerManager.cpp b/Engine/Source/Runtime/Engine/Private/TimerManager.cpp
--- a/Engine/Source/Runtime/Engine/Private/TimerManager.cpp
+++ b/Engine/Source/Runtime/Engine/Private/TimerManager.cpp
@@ -247,6 +247,40 @@ namespace
 		DescribeFTimerDataSafely(Output, Data);
 		return MoveTemp(Output);
 	}
+
+	#if UE_WITH_REMOTE_OBJECT_HANDLE
+	/**
+	* Returns true if the TimerData target is locally resident. When a target is locally resident, you can
+	* execute the TimerDelegate without causing (an initial) migration. The delegate itself may cause other migrations.
+	*/
+	bool IsTimerTargetResident(const FTimerUnifiedDelegate& TimerDelegate)
+	{
+		FRemoteObjectId RemoteObjectId = UE::RemoteObject::Handle::GetRemoteObjectId([&TimerDelegate]()
+		{
+			const UObject* TargetObject;
+			switch (TimerDelegate.VariantDelegate.GetIndex())
+			{
+				case FTimerDelegateVariant::IndexOfType<FTimerDelegate>():
+					TargetObject = TimerDelegate.VariantDelegate.Get<FTimerDelegate>().GetUObject();
+					break;
+
+				case FTimerDelegateVariant::IndexOfType<FTimerDynamicDelegate>():
+					TargetObject = TimerDelegate.VariantDelegate.Get<FTimerDynamicDelegate>().GetUObject();
+					break;
+
+				case FTimerDelegateVariant::IndexOfType<FTimerFunction>():
+					// fallthrough -- TFunction does not know about an underlying UObject, so we always consider it local.
+				default:
+					TargetObject = nullptr;
+					break;
+			}
+
+			return TargetObject;
+		});
+
+		return !UE::RemoteObject::Handle::IsRemote(RemoteObjectId);
+	}
+	#endif
 }
 
 struct FTimerHeapOrder
@@ -1116,8 +1150,16 @@ void FTimerManager::Tick(float DeltaTime)
 			// test to ensure it didn't get cleared during execution
 			if (Top)
 			{
+				#if UE_WITH_REMOTE_OBJECT_HANDLE
+				// IsBound can cause a migration (and we are not inside a transaction).
+				// If the Delegate's Target is non-Resident, then we should reschedule and only do IsBound() when it's Resident.
+				const bool bForceLoopingTimerReschedule = !IsTimerTargetResident(Top->TimerDelegate);
+				#else
+				constexpr bool bForceLoopingTimerReschedule = false;
+				#endif
+
 				// if timer requires a delegate, make sure it's still validly bound (i.e. the delegate's object didn't get deleted or something)
-				if (Top->bLoop && (!Top->bRequiresDelegate || Top->TimerDelegate.IsBound()))
+				if (Top->bLoop && (bForceLoopingTimerReschedule || !Top->bRequiresDelegate || Top->TimerDelegate.IsBound()))
 				{
 					// Put this timer back on the heap
 					Top->ExpireTime += CallCount * Top->Rate;
