Fix UE_INLINE_GENERATED_CPP_BY_NAME to in correct preprocessor scope

diff --git a/Engine/Plugins/Runtime/RigVM/Source/RigVMEditor/Private/Editor/Kismet/RigVMFindInBlueprintManager.cpp b/Engine/Plugins/Runtime/RigVM/Source/RigVMEditor/Private/Editor/Kismet/RigVMFindInBlueprintManager.cpp
--- a/Engine/Plugins/Runtime/RigVM/Source/RigVMEditor/Private/Editor/Kismet/RigVMFindInBlueprintManager.cpp
+++ b/Engine/Plugins/Runtime/RigVM/Source/RigVMEditor/Private/Editor/Kismet/RigVMFindInBlueprintManager.cpp
@@ -1,8 +1,10 @@
 // Copyright Epic Games, Inc. All Rights Reserved.
 
-#if !WITH_RIGVMLEGACYEDITOR
-
 #include "Editor/Kismet/RigVMFindInBlueprintManager.h"
+
+#include UE_INLINE_GENERATED_CPP_BY_NAME(RigVMFindInBlueprintManager)
+
+#if !WITH_RIGVMLEGACYEDITOR
 #include "Misc/CoreMisc.h"
 #include "Misc/MessageDialog.h"
 #include "HAL/FileManager.h"
@@ -61,7 +63,6 @@
 #include "Application/ThrottleManager.h"
 #include "EdGraph/RigVMEdGraphNode.h"
 
-#include UE_INLINE_GENERATED_CPP_BY_NAME(RigVMFindInBlueprintManager)
 
 DEFINE_LOG_CATEGORY(LogRigVMFindInBlueprint);
 CSV_DEFINE_CATEGORY(RigVMFindInBlueprint, false);
diff --git a/Engine/Plugins/Runtime/RigVM/Source/RigVMEditor/Private/Editor/Kismet/RigVMFindInBlueprints.cpp b/Engine/Plugins/Runtime/RigVM/Source/RigVMEditor/Private/Editor/Kismet/RigVMFindInBlueprints.cpp
--- a/Engine/Plugins/Runtime/RigVM/Source/RigVMEditor/Private/Editor/Kismet/RigVMFindInBlueprints.cpp
+++ b/Engine/Plugins/Runtime/RigVM/Source/RigVMEditor/Private/Editor/Kismet/RigVMFindInBlueprints.cpp
@@ -1,9 +1,11 @@
 // Copyright Epic Games, Inc. All Rights Reserved.
 
-#if !WITH_RIGVMLEGACYEDITOR
-
 #include "Editor/Kismet/RigVMFindInBlueprints.h"
 
+#include UE_INLINE_GENERATED_CPP_BY_NAME(RigVMFindInBlueprints)
+
+#if !WITH_RIGVMLEGACYEDITOR
+
 #include "Editor/RigVMNewEditor.h"
 #include "BlueprintEditorModule.h"
 #include "BlueprintEditorSettings.h"
@@ -94,7 +96,6 @@
 #include "Widgets/Views/STableRow.h"
 #include "RigVMBlueprintLegacy.h"
 
-#include UE_INLINE_GENERATED_CPP_BY_NAME(RigVMFindInBlueprints)
 
 class ITableRow;
 class SWidget;
diff --git a/Engine/Source/Runtime/CoreUObject/Tests/ClassDefaultObjectTest.cpp b/Engine/Source/Runtime/CoreUObject/Tests/ClassDefaultObjectTest.cpp
--- a/Engine/Source/Runtime/CoreUObject/Tests/ClassDefaultObjectTest.cpp
+++ b/Engine/Source/Runtime/CoreUObject/Tests/ClassDefaultObjectTest.cpp
@@ -1,16 +1,19 @@
 // Copyright Epic Games, Inc. All Rights Reserved.
 
+#include "ClassDefaultObjectTest.h"
+
+#if WITH_TESTS
+
+#include UE_INLINE_GENERATED_CPP_BY_NAME(ClassDefaultObjectTest)
+
 #if WITH_LOW_LEVEL_TESTS
 
-#include "ClassDefaultObjectTest.h"
 #include "TestHarness.h"
 #include "TestMacros/Assertions.h"
 #include "Async/ManualResetEvent.h"
 #include "Tasks/Task.h"
 #include "UObject/Package.h"
 
-#include UE_INLINE_GENERATED_CPP_BY_NAME(ClassDefaultObjectTest)
-
 TEST_CASE("UE::CoreUObject::GetDefaultObjectRace", "[CoreUObject][GetDefaultObjectRace]")
 {
 	for (int32 Index = 0; Index < 1000; ++Index)
@@ -86,3 +89,4 @@ TEST_CASE("UE::CoreUObject::GetDefaultObjectRace", "[CoreUObject][GetDefaultObje
 }
 
 #endif // WITH_LOW_LEVEL_TESTS
+#endif // WITH_TESTS
\ No newline at end of file
