New command line argument DebugPrintGlobalCoreRedirects to print global core redirects.

#rnx

diff --git a/Engine/Source/Runtime/CoreUObject/Internal/UObject/CoreRedirects/RedirectionSummary.h b/Engine/Source/Runtime/CoreUObject/Internal/UObject/CoreRedirects/RedirectionSummary.h
--- a/Engine/Source/Runtime/CoreUObject/Internal/UObject/CoreRedirects/RedirectionSummary.h
+++ b/Engine/Source/Runtime/CoreUObject/Internal/UObject/CoreRedirects/RedirectionSummary.h
@@ -35,6 +35,8 @@ public:
 
 	FCoreRedirects::FGlobalCoreRedirectAddedEvent& GetOnGlobalCoreRedirectAdded();
 
+	void DebugPrintGlobalCoreRedirects() const;
+
 private:
 	struct FCompareRedirect
 	{
@@ -49,6 +51,7 @@ private:
 		void Empty();
 		bool TryAppendHashInReadLock(FBlake3& Hasher) const;
 		void AppendHashInWriteLock(FBlake3& Hasher);
+		void DebugPrint() const;
 	private:
 		void CalculateHash();
 		void AppendHashWithoutDirtyCheck(FBlake3& Hasher) const;
diff --git a/Engine/Source/Runtime/CoreUObject/Private/UObject/CoreRedirects.cpp b/Engine/Source/Runtime/CoreUObject/Private/UObject/CoreRedirects.cpp
--- a/Engine/Source/Runtime/CoreUObject/Private/UObject/CoreRedirects.cpp
+++ b/Engine/Source/Runtime/CoreUObject/Private/UObject/CoreRedirects.cpp
@@ -2037,7 +2037,14 @@ void FCoreRedirects::AppendHashOfGlobalRedirects(FBlake3& Hasher)
 	FCoreRedirectsContext& Context = FCoreRedirectsContext::GetThreadContext();
 	if (!!(Context.GetFlags() & FCoreRedirectsContext::EFlags::UseRedirectionSummary))
 	{
-		Context.GetRedirectionSummary().AppendHashGlobal(Hasher);
+		FRedirectionSummary& RedirectionSummary = Context.GetRedirectionSummary();
+		RedirectionSummary.AppendHashGlobal(Hasher);
+
+		static bool bDebugPrintGlobalCoreRedirects = FParse::Param(FCommandLine::Get(), TEXT("DebugPrintGlobalCoreRedirects"));
+		if (bDebugPrintGlobalCoreRedirects)
+		{
+			RedirectionSummary.DebugPrintGlobalCoreRedirects();
+		}
 	}
 }
 
diff --git a/Engine/Source/Runtime/CoreUObject/Private/UObject/RedirectionSummary.cpp b/Engine/Source/Runtime/CoreUObject/Private/UObject/RedirectionSummary.cpp
--- a/Engine/Source/Runtime/CoreUObject/Private/UObject/RedirectionSummary.cpp
+++ b/Engine/Source/Runtime/CoreUObject/Private/UObject/RedirectionSummary.cpp
@@ -212,6 +212,11 @@ FCoreRedirects::FGlobalCoreRedirectAddedEvent& FRedirectionSummary::GetOnGlobalC
 	return OnGlobalCoreRedirectAdded;
 }
 
+void FRedirectionSummary::DebugPrintGlobalCoreRedirects() const
+{
+	GlobalRedirects.DebugPrint();
+}
+
 void FRedirectionSummary::FRedirectContainer::Add(FCoreRedirect&& Redirect)
 {
 	Redirects.Add(MoveTemp(Redirect));
@@ -256,6 +261,15 @@ void FRedirectionSummary::FRedirectContainer::AppendHashInWriteLock(FBlake3& Has
 	AppendHashWithoutDirtyCheck(Hasher);
 }
 
+void FRedirectionSummary::FRedirectContainer::DebugPrint() const
+{
+	for (const TPair<FCoreRedirect, bool>& Pair : Redirects)
+	{
+		const FCoreRedirect& Redirect = Pair.Key;
+		UE_LOG(LogCoreRedirects, Display, TEXT("%s->%s (%d)"), *Redirect.OldName.ToString(), *Redirect.NewName.ToString(), Redirect.RedirectFlags);
+	}
+}
+
 void FRedirectionSummary::FRedirectContainer::CalculateHash()
 {
 	FBlake3 Hasher;
