Removing a warning when changing (updating) an existing weak object ptr SerialNumber with a new one that is assigned to the same object (which was re-created by either re-loading or migrating a remote object back)

diff --git a/Engine/Source/Runtime/CoreUObject/Private/UObject/WeakObjectPtr.cpp b/Engine/Source/Runtime/CoreUObject/Private/UObject/WeakObjectPtr.cpp
--- a/Engine/Source/Runtime/CoreUObject/Private/UObject/WeakObjectPtr.cpp
+++ b/Engine/Source/Runtime/CoreUObject/Private/UObject/WeakObjectPtr.cpp
@@ -275,8 +275,8 @@ FUObjectItem* FWeakObjectPtr::TryResolveRemoteObject() const
 		{
 			// ObjectSerialNumber doesn't play an important role when remote object support is enabled because the remote id is a globally unique
 			// identifier of an object but we keep it as an optimization (no need to resolve the remote object if the serial numbers match).
-			UE_CLOG(ObjectSerialNumber != 0, LogRemoteObject, Warning, TEXT("Attempting to change existing and valid serial number %d to %d when resolving remote object %s (%s)"),
-				ObjectSerialNumber, ObjectItem->GetSerialNumber(), *ResolvedObject->GetPathName(), *RemoteId.ToString());
+			// The reason why the SerialNumber may change is that we could have a stale weak object ptr we want to resolve after the thing it pointed at was GC'd
+			// and later created / loaded / migrated back.
 			FPlatformAtomics::AtomicStore(&const_cast<FWeakObjectPtr*>(this)->ObjectSerialNumber, ObjectItem->GetSerialNumber());
 		}
 		return ObjectItem;
