[LinkerLoad]
- Fix potential race from VerifyImportInner when multithreaded loading is active

diff --git a/Engine/Source/Runtime/CoreUObject/Private/UObject/LinkerLoad.cpp b/Engine/Source/Runtime/CoreUObject/Private/UObject/LinkerLoad.cpp
--- a/Engine/Source/Runtime/CoreUObject/Private/UObject/LinkerLoad.cpp
+++ b/Engine/Source/Runtime/CoreUObject/Private/UObject/LinkerLoad.cpp
@@ -4084,6 +4084,9 @@ bool FLinkerLoad::VerifyImportInner(const int32 ImportIndex, FString& WarningSuf
 		// if we do not have a linker, then this import is native/in memory only
 		if( Import.SourceLinker )
 		{
+			// Prevents any other serialization function from modifying this linker while we're doing operations on it.
+			UE::TUniqueLock UniqueLock(Import.SourceLinker->Mutex);
+
 			if (!Import.SourceLinker->bHasFoundExistingExports)
 			{
 				UE_ASSET_LOG(LogLinker, Log, PackagePath, TEXT("Source linker '%s' has not processed all header information, ticking it now"), *GetNameSafe(Import.SourceLinker->LinkerRoot));
