[UE] Ensure default subobject values constructed for non-nullable reference types propagate the appropriate flags from the outer instance in the case there is no existing referenced value.

Notes:
- For example, if the outer is an archetype, this ensures that the RF_ArchetypeObject flag is also propagated to the instance.
- Additionally, if the outer has RF_NeedPostLoad set, this ensures that any PostLoad logic will also be deferred on the instance as a subobject.
- Also now propagating RF_PostLoad/RF_PostLoadSubobjects if there is an existing value with those flags set that's being replaced with a default object value. This ensures that any deferred PostLoad logic can still be applied after a load.

diff --git a/Engine/Source/Runtime/CoreUObject/Private/UObject/PropertyBaseObject.cpp b/Engine/Source/Runtime/CoreUObject/Private/UObject/PropertyBaseObject.cpp
--- a/Engine/Source/Runtime/CoreUObject/Private/UObject/PropertyBaseObject.cpp
+++ b/Engine/Source/Runtime/CoreUObject/Private/UObject/PropertyBaseObject.cpp
@@ -788,8 +788,8 @@ UObject* FObjectPropertyBase::ConstructDefaultObjectValueIfNecessary(UObject* Ex
 	}
 
 	// Filter out any flags that do not belong to the RF_Load or RF_PropagateToSubObjects groups. We mostly want to exclude
-	// UObject life-time and internal GC flags.
-	const EObjectFlags FlagsToKeepMask = (RF_Load | RF_PropagateToSubObjects);
+	// UObject life-time and internal GC flags. Include PostLoad/PostLoadSubobjects to preserve any PostLoad logic during loads.
+	const EObjectFlags FlagsToKeepMask = (RF_Load | RF_NeedPostLoad | RF_NeedPostLoadSubobjects | RF_PropagateToSubObjects);
 
 	if (ExistingValue)
 	{
@@ -850,7 +850,7 @@ UObject* FObjectPropertyBase::ConstructDefaultObjectValueIfNecessary(UObject* Ex
 			bool bCanReplaceObjectWithExistingName = !ObjectWithExistingName || ObjectWithExistingName->GetClass()->IsChildOf(PropertyClass);
 
 			FName NewName = (bCanReplaceObjectWithExistingName) ? ExistingName : NAME_None;
-			EObjectFlags NewFlags = ExistingValue ? (ExistingValue->GetFlags() & FlagsToKeepMask) : RF_NoFlags;
+			EObjectFlags NewFlags = ExistingValue ? (ExistingValue->GetFlags() & FlagsToKeepMask) : Outer->GetMaskedFlags(RF_NeedPostLoad | RF_NeedPostLoadSubobjects | RF_PropagateToSubObjects);
 			UObject* Template = Defaults ? GetObjectPropertyValue(Defaults) : nullptr;
 			if (Template && !Template->IsA(PropertyClass))
 			{
