Add CSV region support for timing async loading. Off by default. Use s.AsyncLoadingCsvRegions to enable
#tests console and android

diff --git a/Engine/Source/Runtime/CoreUObject/Private/Serialization/AsyncLoading2.cpp b/Engine/Source/Runtime/CoreUObject/Private/Serialization/AsyncLoading2.cpp
--- a/Engine/Source/Runtime/CoreUObject/Private/Serialization/AsyncLoading2.cpp
+++ b/Engine/Source/Runtime/CoreUObject/Private/Serialization/AsyncLoading2.cpp
@@ -111,6 +111,15 @@ TAutoConsoleVariable<bool> CVar_CheckBlueprintCDOWillBePostLoadedAfterPreload(
 		"When disabled, we will still `ensure` regardless, but will not treat these as fatal errors.")
 );
 
+
+#if CSV_PROFILER
+TAutoConsoleVariable<bool> CVarAsyncLoadingCsvRegions(
+	TEXT("s.AsyncLoadingCsvRegions"),
+	false,
+	TEXT("If true, emit CSV regions for async loading. These are quite spammy so we want them off by default")
+);
+#endif
+
 template<typename PayloadType> using TAsyncAtomic = std::atomic<PayloadType>;
 
 CSV_DECLARE_CATEGORY_MODULE_EXTERN(CORE_API, FileIO);
@@ -118,6 +127,33 @@ CSV_DEFINE_STAT(FileIO, FrameCompletedExportBundleLoadsKB);
 
 LLM_DEFINE_TAG(UObject_FAsyncPackage2);
 
+// Scope helper class for region instrumentation
+struct FAsyncLoadingCsvRegionScope
+{
+#if CSV_PROFILER
+	FAsyncLoadingCsvRegionScope()
+	{
+		bEnable = CVarAsyncLoadingCsvRegions.GetValueOnAnyThread();
+		if (bEnable)
+		{
+			CSV_REGION_BEGIN(FileIO, TEXT("AsyncLoading"));
+		}
+	}
+	~FAsyncLoadingCsvRegionScope()
+	{
+		if (bEnable)
+		{
+			CSV_REGION_END(FileIO, TEXT("AsyncLoading"));
+		}
+	}
+	bool bEnable = false;
+#else
+	FAsyncLoadingCsvRegionScope() {}
+	~FAsyncLoadingCsvRegionScope() {}
+#endif
+};
+
+
 FArchive& operator<<(FArchive& Ar, FZenPackageVersioningInfo& VersioningInfo)
 {
 	Ar << VersioningInfo.ZenVersion;
@@ -10569,7 +10605,7 @@ uint32 FAsyncLoadingThread2::Run()
 			}
 			UE::Sanitizer::RaceDetector::FRaceDetectorScope RaceDetectorScope(GDetectRaceDuringLoading, GDetectRaceGracePeriod);
 #endif
-
+			FAsyncLoadingCsvRegionScope CsvRegionScope;
 			bool bShouldSuspend = false;
 			bool bShouldWaitForExternalReads = false;
 			while (!bStopRequested.load(std::memory_order_relaxed))
