Asset Registry - Always initialize skip classes, even when running from a (non-cook) commandlet
* Prior to CL 43810446, those skip classes were initialized through a different path with a call to SetSkipClasses(). Remove these conditions restore the skip classes initialization and retain our previous behavior.

diff --git a/Engine/Source/Runtime/CoreUObject/Private/Misc/AssetRegistryInterface.cpp b/Engine/Source/Runtime/CoreUObject/Private/Misc/AssetRegistryInterface.cpp
--- a/Engine/Source/Runtime/CoreUObject/Private/Misc/AssetRegistryInterface.cpp
+++ b/Engine/Source/Runtime/CoreUObject/Private/Misc/AssetRegistryInterface.cpp
@@ -48,7 +48,6 @@ namespace Private
 			// Since we only collect these the first on-demand time, it is possible we will miss subclasses
 			// from plugins that load later. This flaw is a rare edge case, though, and this solution will
 			// be replaced eventually, so leaving it for now.
-			if (GIsEditor && (!IsRunningCommandlet() || IsRunningCookCommandlet()))
 			{
 				UE::TWriteScopeLock WriteLock(SkipClassesLock);
 				Utils::PopulateSkipClasses(SkipUncookedClasses, SkipCookedClasses);
