[UObject] Apply AUTORTFM_INFER to reference collector classes / functions

It's likely that this should be AUTORTFM_DISABLE'd, but this is the smallest change I can find to break a chicken-and-egg between AutoRTFM-disabling the Automation Tests and UE GC.

diff --git a/Engine/Source/Runtime/CoreUObject/Public/UObject/FastReferenceCollector.h b/Engine/Source/Runtime/CoreUObject/Public/UObject/FastReferenceCollector.h
--- a/Engine/Source/Runtime/CoreUObject/Public/UObject/FastReferenceCollector.h
+++ b/Engine/Source/Runtime/CoreUObject/Public/UObject/FastReferenceCollector.h
@@ -2,6 +2,7 @@
 
 #pragma once
 
+#include "AutoRTFM/Defines.h"
 #include "CoreMinimal.h"
 #include "Stats/Stats.h"
 #include "UObject/GarbageCollection.h"
@@ -773,7 +774,7 @@ struct HasHandleTokenStreamVerseCellReference <T, std::void_t<HandleTokenStreamV
 
 /** Forwards references directly to ProcessorType::HandleTokenStreamObjectReference(), unlike TBatchDispatcher */
 template<class ProcessorType>
-struct TDirectDispatcher
+struct AUTORTFM_INFER TDirectDispatcher
 {
 	static constexpr bool bBatching = false;
 	static constexpr bool bParallel = IsParallel(ProcessorType::Options);
@@ -1079,7 +1080,7 @@ private:
 
 /** Default reference collector for CollectReferences() */
 template <typename ProcessorType>
-class TDefaultCollector : public FReferenceCollector
+class AUTORTFM_INFER TDefaultCollector : public FReferenceCollector
 {
 protected:
 	ProcessorType& Processor;
@@ -1165,7 +1166,7 @@ public:
 
 
 template<class CollectorType, class ProcessorType>
-FORCEINLINE static void CollectReferences(ProcessorType& Processor, UE::GC::FWorkerContext& Context)
+AUTORTFM_INFER FORCEINLINE static void CollectReferences(ProcessorType& Processor, UE::GC::FWorkerContext& Context)
 {
 	using namespace UE::GC;
 	using FastReferenceCollector = TFastReferenceCollector<ProcessorType, CollectorType>;
@@ -1181,7 +1182,7 @@ FORCEINLINE static void CollectReferences(ProcessorType& Processor, UE::GC::FWor
 }
 
 template<class ProcessorType>
-FORCEINLINE static void CollectReferences(ProcessorType& Processor, UE::GC::FWorkerContext& Context)
+AUTORTFM_INFER FORCEINLINE static void CollectReferences(ProcessorType& Processor, UE::GC::FWorkerContext& Context)
 {
 	CollectReferences<UE::GC::TDefaultCollector<ProcessorType>, ProcessorType>(Processor, Context);
 }
