[BugFix] In the Data Recovery tool, remapping or deleting all the sub-properties of unknown struct properties can cause that struct to disapear from the tool but still be in the unknown property tree. This makes the object unfixable
[BugFix] some struct sub-property redirects don't do anything in the recovery tool
[BugFix] sometimes structs in the right panel are displayed without any sub-properties even though they do have sub properties

diff --git a/Engine/Plugins/Experimental/InstanceDataObjectFixupTool/Source/Private/InstanceDataObjectFixupDetailCustomization.cpp b/Engine/Plugins/Experimental/InstanceDataObjectFixupTool/Source/Private/InstanceDataObjectFixupDetailCustomization.cpp
--- a/Engine/Plugins/Experimental/InstanceDataObjectFixupTool/Source/Private/InstanceDataObjectFixupDetailCustomization.cpp
+++ b/Engine/Plugins/Experimental/InstanceDataObjectFixupTool/Source/Private/InstanceDataObjectFixupDetailCustomization.cpp
@@ -152,12 +152,9 @@ int32 FInstanceDataObjectNameWidgetOverride::GetNameWidgetIndex(FPropertyPath Pa
 	{
 		if (Panel->HasViewFlag(FInstanceDataObjectFixupPanel::EViewFlags::AllowRemapLooseProperties))
 		{
-			if (const FProperty* Property = Path.GetLeafMostProperty().Property.Get())
+			if (Panel->PropertyNeedsRedirect(Path))
 			{
-				if (Property->GetBoolMetaData(TEXT("isLoose")))
-				{
-					return DisplayRedirectMenu;
-				}
+				return DisplayRedirectMenu;
 			}
 		}
 	}
diff --git a/Engine/Plugins/Experimental/InstanceDataObjectFixupTool/Source/Private/InstanceDataObjectFixupPanel.cpp b/Engine/Plugins/Experimental/InstanceDataObjectFixupTool/Source/Private/InstanceDataObjectFixupPanel.cpp
--- a/Engine/Plugins/Experimental/InstanceDataObjectFixupTool/Source/Private/InstanceDataObjectFixupPanel.cpp
+++ b/Engine/Plugins/Experimental/InstanceDataObjectFixupTool/Source/Private/InstanceDataObjectFixupPanel.cpp
@@ -377,8 +377,6 @@ void FInstanceDataObjectFixupPanel::RedirectProperty(const FPropertyPath& From,
 
 		TObjectPtr<UObject> OwnerObjectPtr = const_cast<UObject*>(Owner);
 		UE::FInstanceDataTransforms::Get().ApplyTransformSet(*Transforms, InstanceDataObject, OwnerObjectPtr);
-
-		RefreshDetailsView();
 	}
 }
 
@@ -429,8 +427,6 @@ void FInstanceDataObjectFixupPanel::MarkForDelete(const FPropertyPath& CurrentPa
 		TObjectPtr<UObject> OwnerObjectPtr = const_cast<UObject*>(Owner);
 		UE::FInstanceDataTransforms::Get().ApplyTransformSet(*Transforms, InstanceDataObject, OwnerObjectPtr);
 	}
-
-	RefreshDetailsView();
 }
 
 void FInstanceDataObjectFixupPanel::OnMarkForDelete(FPropertyPath Path)
diff --git a/Engine/Source/Runtime/CoreUObject/Private/UObject/InstanceDataObjectUtils.cpp b/Engine/Source/Runtime/CoreUObject/Private/UObject/InstanceDataObjectUtils.cpp
--- a/Engine/Source/Runtime/CoreUObject/Private/UObject/InstanceDataObjectUtils.cpp
+++ b/Engine/Source/Runtime/CoreUObject/Private/UObject/InstanceDataObjectUtils.cpp
@@ -1028,6 +1028,12 @@ namespace UE
 			}
 		}
 
+		if (!PropertyTree)
+		{
+			// all sub-properties are known so we can use the regular type
+			//return;
+		}
+
 		if (FStructProperty* AsStructProperty = CastField<FStructProperty>(Property))
 		{
 			// Structs that use native or binary serialization cannot safely generate an IDO.
@@ -1160,11 +1166,11 @@ namespace UE
 				}
 			}
 
-			MarkPropertyAsLoose(AsArrayProperty->Inner, PropertyFlags);
+			AsArrayProperty->Inner->SetPropertyFlags(CPF_Edit | CPF_EditConst | PropertyFlags);
 		}
 		else if (const FSetProperty* AsSetProperty = CastField<FSetProperty>(Property))
-{
-			MarkPropertyAsLoose(AsSetProperty->ElementProp);
+		{
+			AsSetProperty->ElementProp->SetPropertyFlags(CPF_Edit | CPF_EditConst);
 		}
 		else if (const FMapProperty* AsMapProperty = CastField<FMapProperty>(Property))
 		{
@@ -1178,19 +1184,12 @@ namespace UE
 				}
 			}
 
-			MarkPropertyAsLoose(AsMapProperty->KeyProp, PropertyFlags);
-			MarkPropertyAsLoose(AsMapProperty->ValueProp);	// override serialization doesn't require any flags on the value property
+			AsMapProperty->KeyProp->SetPropertyFlags(CPF_Edit | CPF_EditConst | PropertyFlags);
+			AsMapProperty->ValueProp->SetPropertyFlags(CPF_Edit | CPF_EditConst); // override serialization doesn't require any flags on the value property
 		}
 		else if (const FOptionalProperty* AsOptionalProperty = CastField<FOptionalProperty>(Property))
 		{
-			MarkPropertyAsLoose(AsOptionalProperty->GetValueProperty());
-		}
-		else if (const FStructProperty* AsStructProperty = CastField<FStructProperty>(Property))
-		{
-			for (FProperty* InnerProperty : TFieldRange<FProperty>(AsStructProperty->Struct))
-			{
-				MarkPropertyAsLoose(InnerProperty);
-			}
+			AsOptionalProperty->GetValueProperty()->SetPropertyFlags(CPF_Edit | CPF_EditConst);
 		}
 		else if (FObjectProperty* AsObjectProperty = CastField<FObjectProperty>(Property))
 		{
@@ -1200,6 +1199,8 @@ namespace UE
 			// also assign the property class to UObject because loose properties can't infer their PropertyClass from TPS data so we'll assume it's as lenient as possible
 			AsObjectProperty->PropertyClass = UObject::StaticClass();
 		}
+
+		// struct sub-properties don't need to be handled. If the sub-properties are loose, this method will get called on them regardless when they're created.
 	}
 
 	bool IsPropertyLoose(const FProperty* Property)
diff --git a/Engine/Source/Runtime/CoreUObject/Private/UObject/InstanceDataTransforms.cpp b/Engine/Source/Runtime/CoreUObject/Private/UObject/InstanceDataTransforms.cpp
--- a/Engine/Source/Runtime/CoreUObject/Private/UObject/InstanceDataTransforms.cpp
+++ b/Engine/Source/Runtime/CoreUObject/Private/UObject/InstanceDataTransforms.cpp
@@ -15,6 +15,7 @@
 #include "UObject/TextProperty.h"
 #include "UObject/PropertyBagRepository.h"
 #include "UObject/OverrideEventHelper.h"
+#include "UObject/OverridableManager.h"
 #include "UObject/PropertyHelper.h"
 #include "UObject/Package.h"
 #include "Dom/JsonValue.h"
@@ -320,27 +321,27 @@ namespace UE::Private::IDT
 		}
 		InferMissingPathInfo(SrcPath, Owner->GetClass());
 
-		if (UnknownPropertyTree->Remove(SrcPath))
+		if (UnknownPropertyTree->Remove(SrcPath) || SubPathIsLeaf(UnknownPropertyTree, SrcPath))
 		{
+			// if a subpath was an overridden struct or container containing this property, we can't easily remove the property from the IDO
+			// class so instead we'll just remove the override on that value so it doesn't get saved anymore
+			FPropertyPathNameResolver SrcValue(SrcPath, IDO);
+			while (SrcValue)
+			{
+				TSharedPtr<FPropertyChangedChainEvent> ClearOverrideChangeEvent;
+				SrcValue.BuildChangeEvent(ClearOverrideChangeEvent, EPropertyChangeType::ResetToDefault);
+				UE::SendClearOverriddenPropertyEvent(IDO, *ClearOverrideChangeEvent, ClearOverrideChangeEvent->PropertyChain);
+				SrcValue.Next();
+			}
+
 			// create a new IDO with the modified property tree
 			TObjectPtr<UObject> OldIDO = IDO;
 			IDO = UE::FPropertyBagRepository::Get().ReinstanceWithCachedUnknownPropertyTree(Owner);
+
 #if WITH_EDITOR
 			FCoreUObjectDelegates::OnObjectsReinstanced.Broadcast({ {OldIDO, IDO} });
 #endif
 		}
-		else if (SubPathIsLeaf(UnknownPropertyTree, SrcPath))
-		{
-			// if a subpath was an overridden struct or container containing this property, we can't easily remove the property from the IDO
-			// class so instead we'll just remove the override on that value so it doesn't get saved anymore
-			FPropertyPathNameResolver ValueToClearOverrides(SrcPath, IDO);
-			if (ValueToClearOverrides)
-			{
-				TSharedPtr<FPropertyChangedChainEvent> OverrideChangeEvent;
-				ValueToClearOverrides.BuildChangeEvent(OverrideChangeEvent, EPropertyChangeType::ResetToDefault);
-				UE::SendClearOverriddenPropertyEvent(IDO, *OverrideChangeEvent, OverrideChangeEvent->PropertyChain);
-			}
-		}
 
 
 		return ETransformResult::Continue;
@@ -387,7 +388,22 @@ namespace UE::Private::IDT
 
 			// find value in old location (in the old IDO)
 			FPropertyPathNameResolver SrcValue(SrcPath, IDO);
-			check(SrcValue); // it should be impossible for SrcPath to be in the old unknown property tree and not be in the old IDO
+			if (!SrcValue)
+			{
+				return ETransformResult::Continue;
+			}
+
+			// clear override(s) on souce value(s)
+			while (SrcValue)
+			{
+				TSharedPtr<FPropertyChangedChainEvent> ClearOverrideChangeEvent;
+				SrcValue.BuildChangeEvent(ClearOverrideChangeEvent, EPropertyChangeType::ResetToDefault);
+				UE::SendClearOverriddenPropertyEvent(IDO, *ClearOverrideChangeEvent, ClearOverrideChangeEvent->PropertyChain);
+				SrcValue.Next();
+			}
+			// reset source value to the beginning of the iteration list (needed in the case of wildcards)
+			SrcValue.Reset();
+			check(SrcValue);
 
 			// create a new IDO with the modified property tree
 			TObjectPtr<UObject> OldIDO = IDO;
@@ -407,16 +423,6 @@ namespace UE::Private::IDT
 				TryCopyOrConvert(SrcValue.Value, DstIDOValue.Value);
 				IDO->PostEditChangeChainProperty(*DstIdoChangeEvent);
 
-				// if a subpath was an overridden struct or container containing this property, then we can't easily remove the property from the IDO
-				// class so instead we'll just remove the override on that value so it doesn't get saved anymore
-				FPropertyPathNameResolver ValueToClearOverrides(SrcPath, IDO);
-				if (ValueToClearOverrides)
-				{
-					TSharedPtr<FPropertyChangedChainEvent> OverrideChangeEvent;
-					ValueToClearOverrides.BuildChangeEvent(OverrideChangeEvent, EPropertyChangeType::ResetToDefault);
-					UE::SendClearOverriddenPropertyEvent(IDO, *OverrideChangeEvent, OverrideChangeEvent->PropertyChain);
-				}
-
 				/*
 					because containers may contain more than one instance of a struct being redirected, we redirect them all by using wildcards in the path. 
 					FPropertyPathName wildcards are resolved by treating INDEX_NONE as "all indices of this container". FPropertyPathNameResolver will resove 
diff --git a/Engine/Source/Runtime/CoreUObject/Private/UObject/PropertyBagRepository.cpp b/Engine/Source/Runtime/CoreUObject/Private/UObject/PropertyBagRepository.cpp
--- a/Engine/Source/Runtime/CoreUObject/Private/UObject/PropertyBagRepository.cpp
+++ b/Engine/Source/Runtime/CoreUObject/Private/UObject/PropertyBagRepository.cpp
@@ -23,6 +23,7 @@
 #include "UObject/UObjectHash.h"
 #include "UObject/UObjectSerializeContext.h"
 #include "UObject/UObjectThreadContext.h"
+#include "UObject/OverridableManager.h"
 #include "UObject/PropertyPathNameTree.h"
 #include "Templates/UnrealTemplate.h"
 
@@ -825,6 +826,61 @@ const UObject* FPropertyBagRepository::FindInstanceForDataObject(const UObject*
 	return Owner ? *Owner : nullptr;
 }
 
+// Compare unknown property tree with the override set. If there's a branch that's not overridden but still in the tree, it should be pruned because
+// it's no longer needed given that the non-overridden properties won't get saved.
+static void PruneUnknownPropertyTree(FPropertyPathNameTree* UnknownPropertyTree, const UStruct* Struct, const FOverriddenPropertyNode* Overrides)
+{
+	static const FName NAME_IsLooseMetadata(ANSITEXTVIEW("IsLoose"));
+	TArray<FPropertyPathName> PathsToRemove;
+	auto PushPathToRemove = [&PathsToRemove](const FPropertyPathNameSegment& Segment) {
+		FPropertyPathName SubPath;
+		SubPath.Push(Segment);
+		PathsToRemove.Add(SubPath);
+	};
+
+	for (UE::FPropertyPathNameTree::FConstIterator Itr = UnknownPropertyTree->CreateConstIterator(); Itr; ++Itr)
+	{
+		const FProperty* Property = UE::FindPropertyByNameAndTypeName(Struct, Itr.GetName(), Itr.GetType());
+		if (!Property)
+		{
+			PushPathToRemove({ Itr.GetName(), Itr.GetType() });
+			continue;
+		}
+
+		// const cast needed because there's no non-const iterator. I'm being very careful to not iterate and modify at the same time so this is fine.
+		FPropertyPathNameTree* SubTree = const_cast<FPropertyPathNameTree*>(Itr.GetNode().GetSubTree());
+		const FStructProperty* AsStructProperty = CastField<FStructProperty>(Property);
+		const FOverriddenPropertyNode* FoundOverride = Overrides->GetSubPropertyNodes().FindByKey(FOverriddenPropertyNodeID(Property));
+		if (!FoundOverride)
+		{
+			PushPathToRemove({ Itr.GetName(), Itr.GetType() });
+			continue;
+		}
+
+		// recurse into struct properties with children
+		if (SubTree && !SubTree->IsEmpty() && AsStructProperty)
+		{
+			PruneUnknownPropertyTree(SubTree, AsStructProperty->Struct, FoundOverride);
+		}
+		
+		// TODO: should we handle other containers here?
+
+		// prune overridden leaves if they aren't marked as loose
+		if (!SubTree || SubTree->IsEmpty())
+		{
+			if (!Property->GetBoolMetaData(NAME_IsLooseMetadata))
+			{
+				PushPathToRemove({ Itr.GetName(), Itr.GetType() });
+			}
+			continue;
+		}
+	}
+	for (FPropertyPathName& Path : PathsToRemove)
+	{
+		UnknownPropertyTree->Remove(Path);
+	}
+}
+
 UObject* FPropertyBagRepository::ReinstanceWithCachedUnknownPropertyTree(UObject* Owner)
 {
 	if (FPropertyBagRepositoryLock LockRepo(this); FPropertyBagAssociationData* BagData = AssociatedData.Find(Owner))
@@ -863,6 +919,15 @@ UObject* FPropertyBagRepository::ReinstanceWithCachedUnknownPropertyTree(UObject
 
 		CopyTaggedProperties(IDO, NewIDO);
 
+		// Cull nodes in the UnknownPropertyTree if they're not overridden
+		if (const FOverriddenPropertySet* OverriddenProperties = FOverridableManager::Get().GetOverriddenProperties(NewIDO))
+		{
+			if (const FOverriddenPropertyNode* OverrideTree = OverriddenProperties->GetOverriddenPropertyNode(nullptr))
+			{
+				PruneUnknownPropertyTree(UnknownPropertyTree.Get(), NewIdoClass, OverrideTree);
+			}
+		}
+
 		// replace the old IDO with the new one in the repository
 		InstanceDataObjectToOwner.Remove(BagData->InstanceDataObject);
 		BagData->InstanceDataObject = NewIDO;
diff --git a/Engine/Source/Runtime/CoreUObject/Private/UObject/PropertyPathFunctions.cpp b/Engine/Source/Runtime/CoreUObject/Private/UObject/PropertyPathFunctions.cpp
--- a/Engine/Source/Runtime/CoreUObject/Private/UObject/PropertyPathFunctions.cpp
+++ b/Engine/Source/Runtime/CoreUObject/Private/UObject/PropertyPathFunctions.cpp
@@ -124,6 +124,12 @@ bool FPropertyPathNameResolver::BuildChangeEvent(TSharedPtr<FPropertyChangedChai
 	return false;
 }
 
+void FPropertyPathNameResolver::Reset()
+{
+	ActiveContainerIterators.Empty();
+	Resolve();
+}
+
 FPropertyPathNameResolver& FPropertyPathNameResolver::Resolve()
 {
 	const UStruct* NextStruct;
@@ -134,6 +140,7 @@ FPropertyPathNameResolver& FPropertyPathNameResolver::Resolve()
 		NextStruct = Object ? Object->GetClass() : nullptr;
 		NextContainer = Object;
 		PathIndex = 0;
+		EventChain.Empty();
 	}
 	else
 	{
@@ -141,6 +148,10 @@ FPropertyPathNameResolver& FPropertyPathNameResolver::Resolve()
 		NextStruct = ActiveContainerIterators.Last().NextStruct;
 		NextContainer = ActiveContainerIterators.Last().NextContainer;
 		PathIndex = ActiveContainerIterators.Last().PathIndex;
+		while (EventChain.Num() > ActiveContainerIterators.Last().ChainLength)
+		{
+			EventChain.RemoveNode(EventChain.GetTail());
+		}
 	}
 
 	for (const int32 SegmentCount = Path.GetSegmentCount(); PathIndex < SegmentCount; ++PathIndex)
@@ -174,7 +185,7 @@ FPropertyPathNameResolver& FPropertyPathNameResolver::Resolve()
 				else if (Property->ArrayDim > 1 || Property->IsA<FArrayProperty>() || Property->IsA<FSetProperty>() || Property->IsA<FMapProperty>())
 				{
 					Segment.Index = 0;
-					ActiveContainerIterators.Emplace(PathIndex, Property, NextContainer, NextStruct, Segment.Index);
+					ActiveContainerIterators.Emplace(PathIndex, Property, NextContainer, NextStruct, EventChain.Num(), Segment.Index);
 				}
 			}
 		}
diff --git a/Engine/Source/Runtime/CoreUObject/Public/UObject/PropertyPathFunctions.h b/Engine/Source/Runtime/CoreUObject/Public/UObject/PropertyPathFunctions.h
--- a/Engine/Source/Runtime/CoreUObject/Public/UObject/PropertyPathFunctions.h
+++ b/Engine/Source/Runtime/CoreUObject/Public/UObject/PropertyPathFunctions.h
@@ -97,6 +97,7 @@ public:
 	// in the case of wildcarded paths, Next() will iterate to the next available index in the conatainer(s).
 	UE_API FPropertyPathNameResolver& Next();
 	UE_API bool BuildChangeEvent(TSharedPtr<FPropertyChangedChainEvent>& OutEvent, int32 ChangeType);
+	UE_API void Reset();
 private:
 	FPropertyPathNameResolver& Resolve();
 	const FPropertyPathName& Path;
@@ -110,6 +111,7 @@ private:
 		const FProperty* Property;
 		void* NextContainer;
 		const UStruct* NextStruct;
+		int32 ChainLength;
 		int32 CurIndex;
 	};
 
