Firefight: fix compilation error about unknown macro WITH_METADATA

diff --git a/Engine/Source/Runtime/CoreUObject/Tests/PackageMetaDataTest.cpp b/Engine/Source/Runtime/CoreUObject/Tests/PackageMetaDataTest.cpp
--- a/Engine/Source/Runtime/CoreUObject/Tests/PackageMetaDataTest.cpp
+++ b/Engine/Source/Runtime/CoreUObject/Tests/PackageMetaDataTest.cpp
@@ -1,5 +1,7 @@
 // Copyright Epic Games, Inc. All Rights Reserved.
 
+#include "Misc/CoreMiscDefines.h" //for WITH_METADATA
+
 #if WITH_LOW_LEVEL_TESTS && WITH_METADATA
 
 #include "TestHarness.h"
