[ZenLoader]
- Prevent use-after-free by avoiding the TMap of the packages stored in the global import to be resized while we're holding a ref to one of its element
- Simplify logic and fix some missing locks by using TStripedMap for other TMaps
- Add Insight scope over any contention encountered on the global import store

diff --git a/Engine/Source/Runtime/CoreUObject/Private/Serialization/AsyncLoading2.cpp b/Engine/Source/Runtime/CoreUObject/Private/Serialization/AsyncLoading2.cpp
--- a/Engine/Source/Runtime/CoreUObject/Private/Serialization/AsyncLoading2.cpp
+++ b/Engine/Source/Runtime/CoreUObject/Private/Serialization/AsyncLoading2.cpp
@@ -1606,7 +1606,7 @@ class FLoadedPackageCellsRef
 	friend class FAsyncLoadingVerseRoot;
 
 private:
-	UE::FMutex Mutex;
+	mutable UE::FMutex Mutex;
 	TAsyncAtomic<bool> bPinned = true;
 	TMap<uint64, Verse::VCell*> PublicExportMap;
 
@@ -1630,7 +1630,7 @@ public:
 		PublicExportMap.Add(ExportHash, Cell);
 	}
 
-	Verse::VCell* GetPublicCellExport(uint64 ExportHash)
+	Verse::VCell* GetPublicCellExport(uint64 ExportHash) const
 	{
 		check(bPinned);
 		UE::TUniqueLock Lock(Mutex);
@@ -1663,6 +1663,57 @@ public:
 	}
 };
 
+// Provides us with a way to profile contention.
+class FAsyncLoadingMutex
+{
+	UE::FSharedMutex Mutex;
+public:
+	void Lock()
+	{
+		if (!Mutex.TryLock())
+		{
+			TRACE_CPUPROFILER_EVENT_SCOPE(FAsyncLoadingMutexWriteLockContention);
+			Mutex.Lock();
+		}
+	}
+
+	void LockShared()
+	{
+		if (!Mutex.TryLockShared())
+		{
+			TRACE_CPUPROFILER_EVENT_SCOPE(FAsyncLoadingMutexReadLockContention);
+			Mutex.LockShared();
+		}
+	}
+
+	void UnlockShared()
+	{
+		Mutex.UnlockShared();
+	}
+
+	void Unlock()
+	{
+		Mutex.Unlock();
+	}
+};
+
+struct FAsyncLoadingSharedMutexStripedMapLockingPolicy
+{
+	typedef FAsyncLoadingMutex                  MutexType;
+	typedef UE::TUniqueLock<FAsyncLoadingMutex> ExclusiveLockType;
+	typedef UE::TSharedLock<FAsyncLoadingMutex> SharedLockType;
+};
+
+template <
+	int32 BucketCount,
+	typename KeyType,
+	typename ValueType,
+	typename SetAllocator = FDefaultSetAllocator,
+	typename KeyFuncs = TDefaultMapHashableKeyFuncs<KeyType, ValueType, false>,
+	typename LockingPolicy = FAsyncLoadingSharedMutexStripedMapLockingPolicy
+>
+using TAsyncLoadingStripedMap = TStripedMap<BucketCount, KeyType, ValueType, SetAllocator, KeyFuncs, LockingPolicy>;
+
 #if ALT2_ENABLE_SCRIPT_IMPORTS_DEBUGGING
 class FCookedScriptObjectsDebug
 {
@@ -1675,12 +1726,91 @@ public:
 	FCookedScriptObjectsDebug();
 	~FCookedScriptObjectsDebug();
 	bool IsEnabled() const { return Objects.Num() > 0; }
-	void TrimDebugData(const TMap<FPackageObjectIndex, UObject*>& RuntimeScriptObjects);
+	void TrimDebugData(const TAsyncLoadingStripedMap<1, FPackageObjectIndex, UObject*>& RuntimeScriptObjects);
 	void OnCompiledInUObjectsRegistered(FName, ECompiledInUObjectsRegisteredStatus Status);
 	FString GetPathName(FPackageObjectIndex Index) const;
 };
 #endif
 
+template <int32 Type>
+class FLockedPackageRef
+{
+private:
+	FLoadedPackageRef*  PackageRef = nullptr;
+	FAsyncLoadingMutex* SharedMutex = nullptr;
+public:
+	FLockedPackageRef(const FLockedPackageRef&) = delete;
+	FLockedPackageRef& operator=(const FLockedPackageRef&) = delete;
+
+	FLockedPackageRef(FLockedPackageRef&& Other)
+	{
+		SharedMutex = Other.SharedMutex;
+		PackageRef  = Other.PackageRef;
+		Other.SharedMutex = nullptr;
+		Other.PackageRef = nullptr;
+	}
+
+	FLockedPackageRef& operator=(FLockedPackageRef&& Other)
+	{
+		if (this != &Other)
+		{
+			Reset();
+			SharedMutex = Other.SharedMutex;
+			PackageRef = Other.PackageRef;
+			Other.SharedMutex = nullptr;
+			Other.PackageRef = nullptr;
+		}
+
+		return *this;
+	}
+
+	FLockedPackageRef() = default;
+
+	FLoadedPackageRef& operator*() { return *PackageRef; }
+	FLoadedPackageRef* operator ->()
+	{
+		return PackageRef;
+	}
+
+	operator bool() const
+	{
+		return PackageRef != nullptr;
+	}
+
+	void Reset()
+	{
+		PackageRef = nullptr;
+		if (SharedMutex)
+		{
+			static_assert(Type >= 0 && Type <= 1);
+			if constexpr (Type == 0)
+			{
+				SharedMutex->UnlockShared();
+			}
+			else
+			{
+				SharedMutex->Unlock();
+			}
+
+			SharedMutex = nullptr;
+		}
+	}
+
+	FLockedPackageRef(FLoadedPackageRef* InPackageRef, FAsyncLoadingMutex& InSharedMutex)
+		: PackageRef(InPackageRef)
+		, SharedMutex(&InSharedMutex)
+	{
+	}
+
+	~FLockedPackageRef()
+	{
+		Reset();
+	}
+};
+
+using FReadLockedPackageRef = FLockedPackageRef<0>;
+using FWriteLockedPackageRef = FLockedPackageRef<1>;
+
 class FGlobalImportStore
 {
 	friend class FAsyncLoadingVerseRoot;
@@ -1690,20 +1820,17 @@ private:
 	FAsyncLoadingThread2& AsyncLoadingThread;
 	// Packages in active loading or completely loaded packages, with Desc.UPackageId as key.
 	// Does not track temp packages with custom UPackage names, since they are never imorted by other packages.
-	mutable FRWLock PackagesLock;
+	mutable FAsyncLoadingMutex PackagesLock;
 	TMap<FPackageId, FLoadedPackageRef> Packages;
+ 	// Contention is currently negligible so only use TStripedMap as a thread-safe TMap with just 1 stripe.
 	// Packages in active loading that export VCells. Stored separately from FLoadedPackageRefs to save space.
-	TMap<FPackageId, FLoadedPackageCellsRef> PackageCells;
-	// PackageCells is read and written from ALT and read from Verse GC thread.
-	UE::FMutex PackageCellsMutex;
+	TAsyncLoadingStripedMap<1, FPackageId, FLoadedPackageCellsRef> PackageCells;
 	// All native script objects (structs, enums, classes, CDOs and their subobjects) from the initial load phase
-	TMap<FPackageObjectIndex, UObject*> ScriptObjects;
+	TAsyncLoadingStripedMap<1, FPackageObjectIndex, UObject*> ScriptObjects;
 	// Built-in VCells always defined by the Verse VM.
-	TMap<FPackageObjectIndex, Verse::VCell*> ScriptCells;
-	// We need to protect ScriptObjects as it is accessed from both GT and ALT.
-	mutable FRWLock ScriptObjectsLock;
+	TAsyncLoadingStripedMap<1, FPackageObjectIndex, Verse::VCell*> ScriptCells;
 	// All currently loaded public export objects from any loaded package
-	TMap<int32, FPublicExportKey> ObjectIndexToPublicExport;
+	TAsyncLoadingStripedMap<1, int32, FPublicExportKey> ObjectIndexToPublicExport;
 #if ALT2_ENABLE_SCRIPT_IMPORTS_DEBUGGING
 public:
 	FCookedScriptObjectsDebug CookedScriptObjectsDebug;
@@ -1723,19 +1850,17 @@ public:
 
 	int32 GetStoredPackagesCount() const
 	{
-		FReadScopeLock ScopeLock(PackagesLock);
+		UE::TSharedLock ScopeLock(PackagesLock);
 		return Packages.Num();
 	}
 
 	int32 GetStoredScriptObjectsCount() const
 	{
-		FReadScopeLock ScopeLock(ScriptObjectsLock);
 		return ScriptObjects.Num();
 	}
 
 	SIZE_T GetStoredScriptObjectsAllocatedSize() const
 	{
-		FReadScopeLock ScopeLock(ScriptObjectsLock);
 		return ScriptObjects.GetAllocatedSize();
 	}
 
@@ -1749,27 +1874,51 @@ public:
 		return Packages.Find(PackageId);
 
 	}
-	inline FLoadedPackageRef* FindPackageRef(FPackageId PackageId)
+	inline FReadLockedPackageRef FindPackageRef(FPackageId PackageId)
+	{
+		PackagesLock.LockShared();
+		FLoadedPackageRef* PackageRef = FindPackageRefNoLock(PackageId);
+		if (PackageRef)
+		{
+			return FReadLockedPackageRef(PackageRef, PackagesLock);
+		}
+		PackagesLock.UnlockShared();
+		return FReadLockedPackageRef();
+	}
+
+	inline FReadLockedPackageRef FindPackageRefChecked(FPackageId PackageId, FName Name = FName())
 	{
-		FReadScopeLock ScopeLock(PackagesLock);
-		return FindPackageRefNoLock(PackageId);
+		PackagesLock.LockShared();
+		FLoadedPackageRef* PackageRef = FindPackageRefNoLock(PackageId);
+		if (PackageRef)
+		{
+			return FReadLockedPackageRef(PackageRef, PackagesLock);
+		}
+		PackagesLock.UnlockShared();
+		UE_LOG(LogStreaming, Fatal, TEXT("FindPackageRefChecked: Package %s (0x%s) has been deleted"), *Name.ToString(), *LexToString(PackageId));
+		return FReadLockedPackageRef();
 	}
 
-	inline FLoadedPackageRef& FindPackageRefChecked(FPackageId PackageId, FName Name = FName())
+	inline FWriteLockedPackageRef FindPackageRefCheckedForWrite(FPackageId PackageId, FName Name = FName())
 	{
-		FLoadedPackageRef* PackageRef = FindPackageRef(PackageId);
-		UE_CLOG(!PackageRef, LogStreaming, Fatal, TEXT("FindPackageRefChecked: Package %s (0x%s) has been deleted"),
-			*Name.ToString(), *LexToString(PackageId));
-		return *PackageRef;
+		PackagesLock.Lock();
+		FLoadedPackageRef* PackageRef = FindPackageRefNoLock(PackageId);
+		if (PackageRef)
+		{
+			return FWriteLockedPackageRef(PackageRef, PackagesLock);
+		}
+		PackagesLock.Unlock();
+		UE_LOG(LogStreaming, Fatal, TEXT("FindPackageRefChecked: Package %s (0x%s) has been deleted"), *Name.ToString(), *LexToString(PackageId));
+		return FWriteLockedPackageRef();
 	}
 
-	inline FLoadedPackageRef& AddPackageRef(FPackageId PackageId, FName PackageNameIfKnown, EPackageLoader PackageLoaderIfKnown, EPackageExtension PackageHeaderExtensionIfKnown)
+	inline FWriteLockedPackageRef AddPackageRef(FPackageId PackageId, FName PackageNameIfKnown, EPackageLoader PackageLoaderIfKnown, EPackageExtension PackageHeaderExtensionIfKnown)
 	{
 		LLM_SCOPE_BYNAME(TEXT("AsyncLoadPackageStore"));
 
 		FlushDeferredCleanupPackagesQueue();
 
-		FWriteScopeLock ScopeLock(PackagesLock);
+		PackagesLock.Lock();
 		FLoadedPackageRef& PackageRef = Packages.FindOrAdd(PackageId);
 
 		// Detect package that have been renamed but are still referenced and attempt to get rid of the references
@@ -1893,10 +2042,12 @@ public:
 						ObjectIndexToPublicExport.Remove(ObjectIndex);
 					}
 
-					if (FLoadedPackageCellsRef* PackageCellsRef = PackageCells.Find(PackageId))
-					{
-						PackageCellsRef->PinPublicCellExportsForGC();
-					}
+					PackageCells.FindAndApply(PackageId,
+						[](FLoadedPackageCellsRef& PackageCellsRef)
+						{
+							PackageCellsRef.PinPublicCellExportsForGC();
+						}
+					);
 				}
 			}
 
@@ -1921,55 +2072,57 @@ public:
 			}
 		}
 		++PackageRef.RefCount;
-		return PackageRef;
+		return FWriteLockedPackageRef(&PackageRef, PackagesLock);
 	}
 
 	inline void ReleasePackageRef(FPackageId PackageId, FPackageId FromPackageId = FPackageId())
 	{
-		FLoadedPackageRef& PackageRef = FindPackageRefChecked(PackageId);
+		FReadLockedPackageRef PackageRef = FindPackageRefChecked(PackageId);
 
-		check(PackageRef.RefCount > 0);
-		--PackageRef.RefCount;
+		check(PackageRef->RefCount > 0);
+		--PackageRef->RefCount;
 
 #if DO_CHECK
-		ensureMsgf(!PackageRef.bHasBeenLoadedDebug || PackageRef.bAreAllPublicExportsLoaded || PackageRef.bIsMissing || PackageRef.bIsNotInstalled || PackageRef.bHasFailed,
+		ensureMsgf(!PackageRef->bHasBeenLoadedDebug || PackageRef->bAreAllPublicExportsLoaded || PackageRef->bIsMissing || PackageRef->bIsNotInstalled || PackageRef->bHasFailed,
 			TEXT("LoadedPackageRef from None (0x%s) to %s (0x%s) should not have been released when the package is not complete.")
 			TEXT("RefCount=%d, AreAllExportsLoaded=%d, IsMissing=%d, IsNotInstalled=%d, HasFailed=%d, HasBeenLoaded=%d"),
 			*LexToString(FromPackageId),
-			*PackageRef.GetOriginalPackageName().ToString(),
+			*PackageRef->GetOriginalPackageName().ToString(),
 			*LexToString(PackageId),
-			PackageRef.RefCount,
-			PackageRef.bAreAllPublicExportsLoaded,
-			PackageRef.bIsMissing,
-			PackageRef.bIsNotInstalled,
-			PackageRef.bHasFailed,
-			PackageRef.bHasBeenLoadedDebug);
+			PackageRef->RefCount,
+			PackageRef->bAreAllPublicExportsLoaded,
+			PackageRef->bIsMissing,
+			PackageRef->bIsNotInstalled,
+			PackageRef->bHasFailed,
+			PackageRef->bHasBeenLoadedDebug);
 
-		if (PackageRef.bAreAllPublicExportsLoaded)
+		if (PackageRef->bAreAllPublicExportsLoaded)
 		{
-			check(!PackageRef.bIsMissing && !PackageRef.bIsNotInstalled);
+			check(!PackageRef->bIsMissing && !PackageRef->bIsNotInstalled);
 		}
-		if (PackageRef.bIsMissing || PackageRef.bIsNotInstalled)
+		if (PackageRef->bIsMissing || PackageRef->bIsNotInstalled)
 		{
-			check(!PackageRef.bAreAllPublicExportsLoaded);
+			check(!PackageRef->bAreAllPublicExportsLoaded);
 		}
 #endif
 		// is this the last reference to a loaded package?
-		if (PackageRef.RefCount == 0 && PackageRef.HasPackage())
+		if (PackageRef->RefCount == 0 && PackageRef->HasPackage())
 		{
-			PackageRef.UnpinPublicExportsForGC();
+			PackageRef->UnpinPublicExportsForGC();
 
-			if (FLoadedPackageCellsRef* PackageCellsRef = PackageCells.Find(PackageId))
-			{
-				PackageCellsRef->UnpinPublicCellExportsForGC();
-			}
+			PackageCells.FindAndApply(PackageId,
+				[](FLoadedPackageCellsRef& PackageCellsRef)
+				{
+					PackageCellsRef.UnpinPublicCellExportsForGC();
+				}
+			);
 		}
 	}
 
 	void VerifyLoadedPackages()
 	{
 		TRACE_CPUPROFILER_EVENT_SCOPE(VerifyLoadedPackages);
-		FReadScopeLock ScopeLock(PackagesLock);
+		UE::TSharedLock ScopeLock(PackagesLock);
 		for (TPair<FPackageId, FLoadedPackageRef>& Pair : Packages)
 		{
 			FPackageId& PackageId = Pair.Key;
@@ -1998,10 +2151,12 @@ public:
 			ObjectIndexToPublicExport.Remove(ObjectIndex);
 		}
 
-		if (FLoadedPackageCellsRef* PackageCellsRef = PackageCells.Find(OldPackage->GetPackageId()))
-		{
-			PackageCellsRef->RemoveUnreferencedObsoletePackage();
-		}
+		PackageCells.FindAndApply(OldPackage->GetPackageId(),
+			[](FLoadedPackageCellsRef& PackageCellsRef)
+			{
+				PackageCellsRef.RemoveUnreferencedObsoletePackage();
+			}
+		);
 
 		PackageRef.RemoveUnreferencedObsoletePackage();
 		// Reset PackageId to prevent a double remove from GC NotifyUnreachableObjects
@@ -2010,11 +2165,11 @@ public:
 
 	void RemoveUnreferencedObsoletePackage(FLoadedPackageRef& PackageRef)
 	{
-		FWriteScopeLock ScopeLock(PackagesLock);
+		UE::TUniqueLock ScopeLock(PackagesLock);
 		RemoveUnreferencedObsoletePackageNoLock(PackageRef);
 	}
 
-	void ReplaceReferencedRenamedPackage(FLoadedPackageRef& PackageRef, UPackage* NewPackage)
+	void ReplaceReferencedRenamedPackageNoLock(FLoadedPackageRef& PackageRef, UPackage* NewPackage)
 	{
 		UPackage* OldPackage = PackageRef.GetPackage();
 		UE_ASYNC_UPACKAGE_DEBUG(OldPackage);
@@ -2029,7 +2184,7 @@ public:
 	void RemovePackages(const FUnreachableObjects& ObjectsToRemove)
 	{
 		TRACE_CPUPROFILER_EVENT_SCOPE(RemovePackages);
-		FWriteScopeLock ScopeLock(PackagesLock);
+		UE::TUniqueLock ScopeLock(PackagesLock);
 		for (const FUnreachableObject& Item : ObjectsToRemove)
 		{
 			const FPackageId& PackageId = Item.PackageId;
@@ -2047,7 +2202,6 @@ public:
 		bool bRemoved = Packages.RemoveAndCopyValue(PackageId, PackageRef);
 
 		{
-			UE::TUniqueLock Lock(PackageCellsMutex);
 			bool bRemovedCells = PackageCells.Remove(PackageId) > 0;
 #if DO_CHECK
 			ensureMsgf(!bRemovedCells || bRemoved, TEXT("Removed %s from cell package map when there should have been nothing to remove"), *LexToString(PackageId));
@@ -2066,7 +2220,7 @@ public:
 
 	void RemovePackage(FPackageId PackageId)
 	{
-		FWriteScopeLock ScopeLock(PackagesLock);
+		UE::TUniqueLock ScopeLock(PackagesLock);
 		RemovePackageNoLock(PackageId);
 	}
 
@@ -2075,7 +2229,7 @@ public:
 		TRACE_CPUPROFILER_EVENT_SCOPE(RemovePublicExports);
 
 		FPackageId LastPackageId;
-		FLoadedPackageRef* PackageRef = nullptr;
+		FReadLockedPackageRef PackageRef;
 		for (const FUnreachableObject& Item : ObjectsToRemove)
 		{
 			int32 ObjectIndex = Item.ObjectIndex;
@@ -2102,50 +2256,51 @@ public:
 	void VerifyObjectForRemovalNoLock(UObject* GCObject)
 	{
 		int32 ObjectIndex = GUObjectArray.ObjectToIndex(GCObject);
-		FPublicExportKey* PublicExportKey = ObjectIndexToPublicExport.Find(ObjectIndex);
-		if (PublicExportKey)
-		{
-			FPackageId PackageId = PublicExportKey->GetPackageId();
-			FLoadedPackageRef* PackageRef = FindPackageRefNoLock(PackageId);
-			if (PackageRef)
-			{
-				UObject* ExistingObject = FindPublicExportObjectUncheckedNoLock(*PublicExportKey);
-				UE_CLOG(!ExistingObject, LogStreaming, Fatal,
-					TEXT("FGlobalImportStore::VerifyObjectForRemoval: The loaded public export object '%s' with flags (ObjectFlags=%x, InternalObjectFlags=%x) and id %s:0x%llX is missing in GlobalImportStore. ")
-					TEXT("Reason unknown. Double delete? Bug or hash collision?"),
-					*GCObject->GetFullName(),
-					GCObject->GetFlags(),
-					int(GCObject->GetInternalFlags()),
-					*FormatPackageId(PackageId),
-					PublicExportKey->GetExportHash());
-
-				UE_CLOG(ExistingObject != GCObject, LogStreaming, Fatal,
-					TEXT("FGlobalImportStore::VerifyObjectForRemoval: The loaded public export object '%s' with flags (ObjectFlags=%x, InternalObjectFlags=%x) and id %s:0x%llX is not matching the object '%s' in GlobalImportStore. ")
-					TEXT("Reason unknown. Overwritten after it was added? Bug or hash collision?"),
-					*GCObject->GetFullName(),
-					GCObject->GetFlags(),
-					int(GCObject->GetInternalFlags()),
-					*FormatPackageId(PackageId),
-					PublicExportKey->GetExportHash(),
-					*ExistingObject->GetFullName());
-			}
-			else
+		ObjectIndexToPublicExport.FindAndApply(ObjectIndex,
+			[this, GCObject](const FPublicExportKey& PublicExportKey)
 			{
-				UE_LOG(LogStreaming, Warning,
-					TEXT("FGlobalImportStore::VerifyObjectForRemoval: The package for the serialized GC object '%s' with flags (ObjectFlags=%x, InternalObjectFlags=%x) and id %s:0x%llX is missing in GlobalImportStore. ")
-					TEXT("Most likely this object has been moved into this package after it was loaded, while the original package is still around."),
-					*GCObject->GetFullName(),
-					GCObject->GetFlags(),
-					int(GCObject->GetInternalFlags()),
-					*FormatPackageId(PackageId),
-					PublicExportKey->GetExportHash());
+				FPackageId PackageId = PublicExportKey.GetPackageId();
+				FLoadedPackageRef* PackageRef = FindPackageRefNoLock(PackageId);
+				if (PackageRef)
+				{
+					UObject* ExistingObject = FindPublicExportObjectUncheckedNoLock(PublicExportKey);
+					UE_CLOG(!ExistingObject, LogStreaming, Fatal,
+						TEXT("FGlobalImportStore::VerifyObjectForRemoval: The loaded public export object '%s' with flags (ObjectFlags=%x, InternalObjectFlags=%x) and id %s:0x%llX is missing in GlobalImportStore. ")
+						TEXT("Reason unknown. Double delete? Bug or hash collision?"),
+						*GCObject->GetFullName(),
+						GCObject->GetFlags(),
+						int(GCObject->GetInternalFlags()),
+						*FormatPackageId(PackageId),
+						PublicExportKey.GetExportHash());
+
+					UE_CLOG(ExistingObject != GCObject, LogStreaming, Fatal,
+						TEXT("FGlobalImportStore::VerifyObjectForRemoval: The loaded public export object '%s' with flags (ObjectFlags=%x, InternalObjectFlags=%x) and id %s:0x%llX is not matching the object '%s' in GlobalImportStore. ")
+						TEXT("Reason unknown. Overwritten after it was added? Bug or hash collision?"),
+						*GCObject->GetFullName(),
+						GCObject->GetFlags(),
+						int(GCObject->GetInternalFlags()),
+						*FormatPackageId(PackageId),
+						PublicExportKey.GetExportHash(),
+						*ExistingObject->GetFullName());
+				}
+				else
+				{
+					UE_LOG(LogStreaming, Warning,
+						TEXT("FGlobalImportStore::VerifyObjectForRemoval: The package for the serialized GC object '%s' with flags (ObjectFlags=%x, InternalObjectFlags=%x) and id %s:0x%llX is missing in GlobalImportStore. ")
+						TEXT("Most likely this object has been moved into this package after it was loaded, while the original package is still around."),
+						*GCObject->GetFullName(),
+						GCObject->GetFlags(),
+						int(GCObject->GetInternalFlags()),
+						*FormatPackageId(PackageId),
+						PublicExportKey.GetExportHash());
+				}
 			}
-		}
+		);
 	}
 
 	void VerifyObjectForRemoval(UObject* GCObject)
 	{
-		FReadScopeLock ScopeLock(PackagesLock);
+		UE::TSharedLock ScopeLock(PackagesLock);
 		return VerifyObjectForRemovalNoLock(GCObject);
 	}
 
@@ -2172,22 +2327,27 @@ public:
 					Object->GetInternalFlags(),
 					*FormatPackageId(PackageId));
 
-			FPublicExportKey* PublicExportKey = ObjectIndexToPublicExport.Find(ObjectIndex);
-			UE_CLOG(!PublicExportKey, LogStreaming, Fatal,
-				TEXT("FGlobalImportStore::VerifyPackageForRemoval: %s (%s) - ")
-				TEXT("The loaded public export object '%s' is missing in GlobalImportStore."),
-				*Package->GetName(),
-				*FormatPackageId(PackageId),
-				*Object->GetFullName());
-
-			FPackageId ObjectPackageId = PublicExportKey->GetPackageId();
-			UE_CLOG(ObjectPackageId != PackageId, LogStreaming, Fatal,
-				TEXT("FGlobalImportStore::VerifyPackageForRemoval: %s (%s) - ")
-				TEXT("The loaded public export object '%s' has a mismatching package id %s in GlobalImportStore."),
-				*Package->GetName(),
-				*FormatPackageId(PackageId),
-				*Object->GetFullName(),
-				*FormatPackageId(ObjectPackageId));
+			bool bExists = 
+				ObjectIndexToPublicExport.FindAndApply(ObjectIndex,
+					[Object, Package, PackageId](const FPublicExportKey& PublicExportKey)
+					{
+						FPackageId ObjectPackageId = PublicExportKey.GetPackageId();
+						UE_CLOG(ObjectPackageId != PackageId, LogStreaming, Fatal,
+							TEXT("FGlobalImportStore::VerifyPackageForRemoval: %s (%s) - ")
+							TEXT("The loaded public export object '%s' has a mismatching package id %s in GlobalImportStore."),
+							*Package->GetName(),
+							*FormatPackageId(PackageId),
+							*Object->GetFullName(),
+							*FormatPackageId(ObjectPackageId));
+					}
+				);
+			
+			UE_CLOG(!bExists, LogStreaming, Fatal,
+						TEXT("FGlobalImportStore::VerifyPackageForRemoval: %s (%s) - ")
+						TEXT("The loaded public export object '%s' is missing in GlobalImportStore."),
+						*Package->GetName(),
+						*FormatPackageId(PackageId),
+						*Object->GetFullName());
 
 			VerifyObjectForRemovalNoLock(Object);
 		}
@@ -2195,7 +2355,7 @@ public:
 
 	void VerifyPackageForRemoval(FLoadedPackageRef& PackageRef)
 	{
-		FReadScopeLock ScopeLock(PackagesLock);
+		UE::TSharedLock ScopeLock(PackagesLock);
 		VerifyPackageForRemovalNoLock(PackageRef);
 	}
 
@@ -2211,7 +2371,7 @@ public:
 
 	inline UObject* FindPublicExportObjectUnchecked(const FPublicExportKey& Key)
 	{
-		FReadScopeLock ScopeLock(PackagesLock);
+		UE::TSharedLock ScopeLock(PackagesLock);
 		return FindPublicExportObjectUncheckedNoLock(Key);
 	}
 
@@ -2222,7 +2382,7 @@ public:
 		return Object;
 	}
 
-	inline UObject* FindScriptImportObjectNoLock(FPackageObjectIndex GlobalIndex)
+	inline UObject* FindScriptImportObject(FPackageObjectIndex GlobalIndex)
 	{
 		check(GlobalIndex.IsScriptImport());
 		UObject* Object = nullptr;
@@ -2230,26 +2390,20 @@ public:
 		return Object;
 	}
 
-	inline UObject* FindScriptImportObject(FPackageObjectIndex GlobalIndex)
-	{
-		FReadScopeLock ScopeLock(ScriptObjectsLock);
-		return FindScriptImportObjectNoLock(GlobalIndex);
-	}
-
 	Verse::VCell* FindPublicExportCell(const FPublicExportKey& Key)
 	{
-		FLoadedPackageCellsRef* PackageCellsRef = PackageCells.Find(Key.GetPackageId());
-		if (!PackageCellsRef)
-		{
-			return nullptr;
-		}
-		return PackageCellsRef->GetPublicCellExport(Key.GetExportHash());
+		Verse::VCell* Cell = nullptr;
+		PackageCells.FindAndApply(Key.GetPackageId(),
+			[&Key, &Cell](const FLoadedPackageCellsRef& PackageCellsRef)
+			{
+				Cell = PackageCellsRef.GetPublicCellExport(Key.GetExportHash());
+			}
+		);
+		return Cell;
 	}
 
 	Verse::VCell* FindScriptImportCell(FPackageObjectIndex GlobalIndex)
 	{
-		FReadScopeLock ScopeLock(ScriptObjectsLock);
-
 		check(GlobalIndex.IsScriptImport());
 		return ScriptCells.FindRef(GlobalIndex);
 	}
@@ -2282,42 +2436,52 @@ public:
 			ObjectIndexToPublicExport.Remove(ExistingObjectIndex);
 		}
 
-		FPublicExportKey* ExistingKey = ObjectIndexToPublicExport.Find(ObjectIndex);
-		if (ExistingKey && *ExistingKey != Key)
-		{
-			UE_LOG(LogStreaming, Verbose,
-				TEXT("FGlobalImportStore::StoreGlobalObject: The constructed public export object '%s' with index %d and id %s:0x%llX already exists in GlobalImportStore but with a different key %s:0x%llX.")
-				TEXT("The existing object will be replaced since it or its package was most likely renamed after it was loaded the first time."),
-				Object ? *Object->GetFullName() : TEXT("null"),
-				ObjectIndex,
-				*FormatPackageId(Key.GetPackageId()), Key.GetExportHash(),
-				*FormatPackageId(ExistingKey->GetPackageId()), ExistingKey->GetExportHash());
-
-			// Break the link with the old package now because otherwise, we wouldn't be able to remove the export
-			// during GC since the ObjectIndex can only be linked to a single package ref.
-			if (FLoadedPackageRef* ExistingPackageRef = FindPackageRef(ExistingKey->GetPackageId()))
+		ObjectIndexToPublicExport.RemoveIf(ObjectIndex,
+			[this, ObjectIndex, Object, &Key](const FPublicExportKey& ExistingKey)
 			{
-				ExistingPackageRef->RemovePublicExport(ExistingKey->GetExportHash());
+				if (ExistingKey != Key)
+				{
+					UE_LOG(LogStreaming, Verbose,
+						TEXT("FGlobalImportStore::StoreGlobalObject: The constructed public export object '%s' with index %d and id %s:0x%llX already exists in GlobalImportStore but with a different key %s:0x%llX.")
+						TEXT("The existing object will be replaced since it or its package was most likely renamed after it was loaded the first time."),
+						Object ? *Object->GetFullName() : TEXT("null"),
+						ObjectIndex,
+						*FormatPackageId(Key.GetPackageId()), Key.GetExportHash(),
+						*FormatPackageId(ExistingKey.GetPackageId()), ExistingKey.GetExportHash());
+
+					// Break the link with the old package now because otherwise, we wouldn't be able to remove the export
+					// during GC since the ObjectIndex can only be linked to a single package ref.
+					if (FReadLockedPackageRef ExistingPackageRef = FindPackageRef(ExistingKey.GetPackageId()))
+					{
+						ExistingPackageRef->RemovePublicExport(ExistingKey.GetExportHash());
+					}
+
+					return true; // Remove
+				}
+
+				return false;
 			}
+		);
 
-			ObjectIndexToPublicExport.Remove(ObjectIndex);
+		{
+			FReadLockedPackageRef PackageRef = FindPackageRefChecked(Key.GetPackageId());
+			PackageRef->StorePublicExport(ExportHash, Object);
 		}
-
-		FLoadedPackageRef& PackageRef = FindPackageRefChecked(Key.GetPackageId());
-		PackageRef.StorePublicExport(ExportHash, Object);
 		ObjectIndexToPublicExport.Add(ObjectIndex, Key);
 	}
 
 	void StoreGlobalCell(FPackageId PackageId, uint64 ExportHash, Verse::VCell* Cell)
 	{
-		uint32 KeyHash = GetTypeHash(PackageId);
-		FLoadedPackageCellsRef* PackageCellsRef = PackageCells.FindByHash(KeyHash, PackageId);
-		if (PackageCellsRef == nullptr)
-		{
-			UE::TUniqueLock Lock(PackageCellsMutex);
-			PackageCellsRef = &PackageCells.EmplaceByHash(KeyHash, PackageId, InPlace);
-		}
-		PackageCellsRef->StorePublicCellExport(ExportHash, Cell);
+		PackageCells.FindOrProduceAndApplyForWrite(PackageId,
+			[]()
+			{
+				return InPlace;
+			},
+			[ExportHash, Cell](FLoadedPackageCellsRef& PackageCellsRef)
+			{
+				PackageCellsRef.StorePublicCellExport(ExportHash, Cell);
+			}
+		);
 	}
 
 	void FindAllScriptObjects(bool bVerifyOnly);
@@ -2325,7 +2489,7 @@ public:
 
 	void AddScriptObject(FStringView PackageName, FStringView Name, UObject* Object)
 	{
-		FWriteScopeLock ScopeLock(ScriptObjectsLock);
+		TRACE_CPUPROFILER_EVENT_SCOPE(FGlobalImportStore::AddScriptObject);
 
 		TStringBuilder<FName::StringBufferSize> FullName;
 		FPathViews::Append(FullName, PackageName);
@@ -2354,8 +2518,6 @@ public:
 #if WITH_VERSE_VM || defined(__INTELLISENSE__)
 	void AddScriptCellPackage(Verse::VPackage* Package)
 	{
-		FWriteScopeLock ScopeLock(ScriptObjectsLock);
-
 		for (int32 I = 0; I < Package->NumDefinitions(); ++I)
 		{
 			FUtf8StringView VersePath = Package->GetDefinitionName(I).AsStringView();
@@ -2521,7 +2683,7 @@ public:
 		return !!GlobalImportStore.FindPackageRef(ImportedPackageId);
 	}
 
-	FLoadedPackageRef& AddImportedPackageReference(FPackageId ImportedPackageId, FName PackageNameIfKnown, EPackageLoader PackageLoaderIfKnown, EPackageExtension PackageHeaderExtensionIfKnown)
+	FWriteLockedPackageRef AddImportedPackageReference(FPackageId ImportedPackageId, FName PackageNameIfKnown, EPackageLoader PackageLoaderIfKnown, EPackageExtension PackageHeaderExtensionIfKnown)
 	{
 		return GlobalImportStore.AddPackageRef(ImportedPackageId, PackageNameIfKnown, PackageLoaderIfKnown, PackageHeaderExtensionIfKnown);
 	}
@@ -2530,8 +2692,8 @@ public:
 	{
 		if (Desc.bCanBeImported)
 		{
-			FLoadedPackageRef& PackageRef = GlobalImportStore.AddPackageRef(Desc.UPackageId, Desc.UPackageName, Desc.Loader, Desc.PackagePathToLoad.GetHeaderExtension());
-			PackageRef.ClearErrorFlags();
+			FWriteLockedPackageRef PackageRef = GlobalImportStore.AddPackageRef(Desc.UPackageId, Desc.UPackageName, Desc.Loader, Desc.PackagePathToLoad.GetHeaderExtension());
+			PackageRef->ClearErrorFlags();
 		}
 	}
 
@@ -3057,7 +3219,7 @@ enum class EEventLoadNodeExecutionResult : uint8
 class FEventLoadNode2
 {
 public:
-	FEventLoadNode2(const FAsyncLoadEventSpec* InSpec, FAsyncPackage2* InPackage, int32 InImportOrExportIndex, int32 InBarrierCount);
+	FEventLoadNode2(const FAsyncLoadEventSpec* InSpec, FAsyncPackage2* InPackage, int32 InBarrierCount);
 	void DependsOn(FEventLoadNode2* Other);
 	void AddBarrier();
 	void AddBarrier(int32 Count);
@@ -3098,7 +3260,6 @@ private:
 	uint32 DependenciesCount = 0;
 	uint32 DependenciesCapacity = 0;
 	int32 Priority = 0;
-	int32 ImportOrExportIndex = -1;
 	TAsyncAtomic<int32> BarrierCount{ 0 };
 	enum EQueueStatus
 	{
@@ -3218,7 +3379,7 @@ enum class EExecutionType
 
 struct FAsyncLoadEventSpec
 {
-	typedef EEventLoadNodeExecutionResult(*FAsyncLoadEventFunc)(FAsyncLoadingThreadState2&, FAsyncPackage2*, int32);
+	typedef EEventLoadNodeExecutionResult(*FAsyncLoadEventFunc)(FAsyncLoadingThreadState2&, FAsyncPackage2*);
 	FAsyncLoadEventFunc Func = nullptr;
 	FAsyncLoadEventQueue2* EventQueue = nullptr;
 	EExecutionType ExecutionType = EExecutionType::Normal;
@@ -3982,15 +4143,15 @@ public:
 
 	static FAsyncPackage2* GetCurrentlyExecutingPackage(FAsyncLoadingThreadState2& ThreadState, FAsyncPackage2* PackageToFilter = nullptr);
 
-	static EEventLoadNodeExecutionResult Event_CreateExports(FAsyncLoadingThreadState2& ThreadState, FAsyncPackage2* Package, int32);
-	static EEventLoadNodeExecutionResult Event_ResolveImports(FAsyncLoadingThreadState2& ThreadState, FAsyncPackage2* Package, int32);
-	static EEventLoadNodeExecutionResult Event_PreloadExports(FAsyncLoadingThreadState2& ThreadState, FAsyncPackage2* Package, int32);
-	static EEventLoadNodeExecutionResult Event_ProcessPackageSummary(FAsyncLoadingThreadState2& ThreadState, FAsyncPackage2* Package, int32);
-	static EEventLoadNodeExecutionResult Event_ProcessDependencies(FAsyncLoadingThreadState2& ThreadState, FAsyncPackage2* Package, int32);
-	static EEventLoadNodeExecutionResult Event_DependenciesReady(FAsyncLoadingThreadState2& ThreadState, FAsyncPackage2* Package, int32);
-	static EEventLoadNodeExecutionResult Event_ExportsDone(FAsyncLoadingThreadState2& ThreadState, FAsyncPackage2* Package, int32);
-	static EEventLoadNodeExecutionResult Event_PostLoad(FAsyncLoadingThreadState2& ThreadState, FAsyncPackage2* Package, int32 ExportBundleIndex);
-	static EEventLoadNodeExecutionResult Event_DeferredPostLoad(FAsyncLoadingThreadState2& ThreadState, FAsyncPackage2* Package, int32 ExportBundleIndex);
+	static EEventLoadNodeExecutionResult Event_CreateExports(FAsyncLoadingThreadState2& ThreadState, FAsyncPackage2* Package);
+	static EEventLoadNodeExecutionResult Event_ResolveImports(FAsyncLoadingThreadState2& ThreadState, FAsyncPackage2* Package);
+	static EEventLoadNodeExecutionResult Event_PreloadExports(FAsyncLoadingThreadState2& ThreadState, FAsyncPackage2* Package);
+	static EEventLoadNodeExecutionResult Event_ProcessPackageSummary(FAsyncLoadingThreadState2& ThreadState, FAsyncPackage2* Package);
+	static EEventLoadNodeExecutionResult Event_ProcessDependencies(FAsyncLoadingThreadState2& ThreadState, FAsyncPackage2* Package);
+	static EEventLoadNodeExecutionResult Event_DependenciesReady(FAsyncLoadingThreadState2& ThreadState, FAsyncPackage2* Package);
+	static EEventLoadNodeExecutionResult Event_ExportsDone(FAsyncLoadingThreadState2& ThreadState, FAsyncPackage2* Package);
+	static EEventLoadNodeExecutionResult Event_PostLoad(FAsyncLoadingThreadState2& ThreadState, FAsyncPackage2* Package);
+	static EEventLoadNodeExecutionResult Event_DeferredPostLoad(FAsyncLoadingThreadState2& ThreadState, FAsyncPackage2* Package);
 
 	bool CanRunNodeAsync(const FAsyncLoadEventSpec* Spec) const;
 
@@ -4733,7 +4894,7 @@ private:
 	bool TryGetExistingLoadedPackagePath(FPackageId InPackageIdToLoad, FName InPackageNameToLoad, EPackageExtension& OutPackageExtension, EPackageLoader& OutPackageLoader, FPackageStoreEntry& OutZenStoreEntry)
 	{
 		// Lookup if we have the package in the importstore already to avoid expensive calls in TryGetPackagePathFromFileSystem
-		if (FLoadedPackageRef* PackageRef = GlobalImportStore.FindPackageRef(InPackageIdToLoad))
+		if (FReadLockedPackageRef PackageRef = GlobalImportStore.FindPackageRef(InPackageIdToLoad))
 		{
 			// If a package has failed or wasn't loaded from the loader, fail the lookup: we cannot trust these references
 			// and it's possible their actual state is now different than what was cached in the GlobalImportStore.
@@ -4776,7 +4937,7 @@ private:
 				while (!ImportsToCheck.IsEmpty() && !UnloadableImportPackageId.IsValid())
 				{
 					FPackageId ImportId = ImportsToCheck.Pop();
-					if (FLoadedPackageRef* PackageRef = GlobalImportStore.FindPackageRef(ImportId))
+					if (FReadLockedPackageRef PackageRef = GlobalImportStore.FindPackageRef(ImportId))
 					{
 						if (!PackageRef->HasErrors())
 						{
@@ -5293,18 +5454,19 @@ public:
 		// This plays the role of PinPublicExportsForGC/UnpinPublicExportsForGC and EInternalObjectFlags::LoaderImport
 		{
 			FGlobalImportStore& GlobalImportStore = AsyncLoadingThread->GlobalImportStore;
-			UE::TUniqueLock PackageCellsLock(GlobalImportStore.PackageCellsMutex);
-			for (TPair<FPackageId, FLoadedPackageCellsRef>& PackageCellsRef : GlobalImportStore.PackageCells)
-			{
-				if (PackageCellsRef.Value.bPinned.load(std::memory_order_acquire))
+			GlobalImportStore.PackageCells.ForEach(
+				[&Visitor](const TPair<FPackageId, FLoadedPackageCellsRef>& PackageCellsRef)
 				{
-					UE::TUniqueLock PackageCellsRefLock(PackageCellsRef.Value.Mutex);
-					for (TPair<uint64, Verse::VCell*>& CellExport : PackageCellsRef.Value.PublicExportMap)
+					if (PackageCellsRef.Value.bPinned.load(std::memory_order_acquire))
 					{
-						Visitor.Visit(CellExport.Value, TEXT(""));
+						UE::TUniqueLock PackageCellsRefLock(PackageCellsRef.Value.Mutex);
+						for (const TPair<uint64, Verse::VCell*>& CellExport : PackageCellsRef.Value.PublicExportMap)
+						{
+							Visitor.Visit(CellExport.Value, TEXT(""));
+						}
 					}
 				}
-			}
+			);
 		}
 	}
 
@@ -5322,21 +5484,22 @@ public:
 		//   support reloading on top of existing VCells either.
 		{
 			FGlobalImportStore& GlobalImportStore = AsyncLoadingThread->GlobalImportStore;
-			UE::TUniqueLock PackageCellsLock(GlobalImportStore.PackageCellsMutex);
-			for (TPair<FPackageId, FLoadedPackageCellsRef>& PackageCellsRef : GlobalImportStore.PackageCells)
-			{
-				if (!PackageCellsRef.Value.bPinned.load(std::memory_order_acquire))
+			GlobalImportStore.PackageCells.ForEach(
+				[](TPair<FPackageId, FLoadedPackageCellsRef>& PackageCellsRef)
 				{
-					UE::TUniqueLock PackageCellsRefLock(PackageCellsRef.Value.Mutex);
-					for (auto It = PackageCellsRef.Value.PublicExportMap.CreateIterator(); It; ++It)
+					if (!PackageCellsRef.Value.bPinned.load(std::memory_order_acquire))
 					{
-						if (!Verse::FHeap::IsMarked(It->Value))
+						UE::TUniqueLock PackageCellsRefLock(PackageCellsRef.Value.Mutex);
+						for (auto It = PackageCellsRef.Value.PublicExportMap.CreateIterator(); It; ++It)
 						{
-							It.RemoveCurrent();
+							if (!Verse::FHeap::IsMarked(It->Value))
+							{
+								It.RemoveCurrent();
+							}
 						}
 					}
 				}
-			}
+			);
 		}
 	}
 };
@@ -5882,10 +6045,9 @@ bool FAsyncLoadingThread2::CreateAsyncPackagesFromQueue(FAsyncLoadingThreadState
 	return bPackagesCreated;
 }
 
-FEventLoadNode2::FEventLoadNode2(const FAsyncLoadEventSpec* InSpec, FAsyncPackage2* InPackage, int32 InImportOrExportIndex, int32 InBarrierCount)
+FEventLoadNode2::FEventLoadNode2(const FAsyncLoadEventSpec* InSpec, FAsyncPackage2* InPackage, int32 InBarrierCount)
 	: Spec(InSpec)
 	, Package(InPackage)
-	, ImportOrExportIndex(InImportOrExportIndex)
 	, BarrierCount(InBarrierCount)
 {
 	check(Spec);
@@ -6079,7 +6241,7 @@ EEventLoadNodeExecutionResult FEventLoadNode2::Execute(FAsyncLoadingThreadState2
 	EEventLoadNodeExecutionResult Result;
 	{
 		ThreadState.CurrentlyExecutingEventNodeStack.Push(this);
-		Result = Spec->Func(ThreadState, Package, ImportOrExportIndex);
+		Result = Spec->Func(ThreadState, Package);
 		ThreadState.CurrentlyExecutingEventNodeStack.Pop();
 	}
 	if (Result == EEventLoadNodeExecutionResult::Complete)
@@ -6358,10 +6520,9 @@ void FAsyncPackage2::SetupScriptDependencies()
 	}
 }
 
-
 void FGlobalImportStore::FindAllScriptObjects(bool bVerifyOnly)
 {
-	FWriteScopeLock ScopeLock(ScriptObjectsLock);
+	TRACE_CPUPROFILER_EVENT_SCOPE(FGlobalImportStore::FindAllScriptObjects);
 
 	TStringBuilder<FName::StringBufferSize> Name;
 	TArray<UObject*> Objects;
@@ -6376,18 +6537,19 @@ void FGlobalImportStore::FindAllScriptObjects(bool bVerifyOnly)
 				Name.Reset();
 				Package->GetPathName(nullptr, Name);
 				FPackageObjectIndex PackageGlobalImportIndex = FPackageObjectIndex::FromScriptPath(Name);
-				if (!ScriptObjects.Contains(PackageGlobalImportIndex))
+
+				if (bVerifyOnly)
 				{
-					if (bVerifyOnly)
+					if (!ScriptObjects.Contains(PackageGlobalImportIndex))
 					{
 						UE_LOG(LogStreaming, Display, TEXT("Script package %s (0x%016llX) is missing a NotifyRegistrationEvent from the initial load phase."),
 							*Package->GetFullName(),
 							PackageGlobalImportIndex.Value());
 					}
-					else
-					{
-						ScriptObjects.Add(PackageGlobalImportIndex, Package);
-					}
+				}
+				else
+				{
+					ScriptObjects.Add(PackageGlobalImportIndex, Package);
 				}
 #endif
 				Objects.Reset();
@@ -6406,18 +6568,18 @@ void FGlobalImportStore::FindAllScriptObjects(bool bVerifyOnly)
 						Name.Reset();
 						Object->GetPathName(nullptr, Name);
 						FPackageObjectIndex GlobalImportIndex = FPackageObjectIndex::FromScriptPath(Name);
-						if (!ScriptObjects.Contains(GlobalImportIndex))
+						if (bVerifyOnly)
 						{
-							if (bVerifyOnly)
+							if (!ScriptObjects.Contains(GlobalImportIndex))
 							{
 								UE_LOG(LogStreaming, Warning, TEXT("Script object %s (0x%016llX) is missing a NotifyRegistrationEvent from the initial load phase."),
 									*Object->GetFullName(),
 									GlobalImportIndex.Value());
 							}
-							else
-							{
-								ScriptObjects.Add(GlobalImportIndex, Object);
-							}
+						}
+						else
+						{
+							ScriptObjects.Add(GlobalImportIndex, Object);
 						}
 					}
 				}
@@ -6495,7 +6657,7 @@ void FCookedScriptObjectsDebug::LoadDebugData()
 	}
 }
 
-void FCookedScriptObjectsDebug::TrimDebugData(const TMap<FPackageObjectIndex, UObject*>& RuntimeScriptObjects)
+void FCookedScriptObjectsDebug::TrimDebugData(const TAsyncLoadingStripedMap<1, FPackageObjectIndex, UObject*>& RuntimeScriptObjects)
 {
 	if (!IsEnabled())
 	{
@@ -6528,15 +6690,18 @@ void FCookedScriptObjectsDebug::TrimDebugData(const TMap<FPackageObjectIndex, UO
 			}
 		};
 	// 1. Mark all runtime objects to be removed
-	for (const TPair<FPackageObjectIndex, UObject*>& Runtime : RuntimeScriptObjects)
-	{
-		FSetElementId Id = Objects.FindId(Runtime.Key);
-		if (Id.IsValidId())
+	RuntimeScriptObjects.ForEach(
+		[this, &Remove, &NumRemaining](const TPair<FPackageObjectIndex, UObject*>& Runtime)
 		{
-			Remove[Id.AsInteger()] = true;
-			--NumRemaining;
+			FSetElementId Id = Objects.FindId(Runtime.Key);
+			if (Id.IsValidId())
+			{
+				Remove[Id.AsInteger()] = true;
+				--NumRemaining;
+			}
 		}
-	}
+	);
+	
 	// 2. Keep the outers for all objects that are kept
 	for (int I = 0; I < Num; ++I)
 	{
@@ -6622,7 +6787,7 @@ void FGlobalImportStore::RegistrationComplete()
 		FindAllScriptObjects(/*bVerifyOnly*/true);
 	}
 #endif
-	FWriteScopeLock ScopeLock(ScriptObjectsLock);
+
 	ScriptObjects.Shrink();
 #if ALT2_ENABLE_SCRIPT_IMPORTS_DEBUGGING
 	CookedScriptObjectsDebug.TrimDebugData(ScriptObjects);
@@ -6742,89 +6907,94 @@ void FAsyncPackage2::ImportPackagesRecursiveInner(FAsyncLoadingThreadState2& Thr
 		}
 		ImportedPackagePath.SetHeaderExtension(ImportedPackageExtension);
 
-		// Add the imported package reference and confirm that its loader and extension match what was found from the package store or import store directly, 
-		// and if not, ensure the package was marked as missing or not installed
-		FLoadedPackageRef& ImportedPackageRef = ImportStore.AddImportedPackageReference(ImportedPackageId, ImportedPackageUPackageName, ImportedPackageLoader, ImportedPackagePath.GetHeaderExtension());
-		checkf((ImportedPackageStatus == EPackageStoreEntryStatus::Missing || ImportedPackageStatus == EPackageStoreEntryStatus::NotInstalled)
-			|| ImportedPackageRef.GetPackageLoader() == ImportedPackageLoader,
-			TEXT("ImportPackageStatus=%s, ImportedPackageRef.GetPackageLoader()=%s, ImportedPackageLoader=%s"),
-			LexToString(ImportedPackageStatus), LexToString(ImportedPackageRef.GetPackageLoader()), LexToString(ImportedPackageLoader));
-		checkf((ImportedPackageStatus == EPackageStoreEntryStatus::Missing || ImportedPackageStatus == EPackageStoreEntryStatus::NotInstalled)
-			|| ImportedPackageRef.GetPackageHeaderExtension() == ImportedPackagePath.GetHeaderExtension(),
-			TEXT("ImportPackageStatus=%s, ImportedPackageRef.GetPackageHeaderExtension()=%s, ImportedPackagePath.GetHeaderExtension()=%s"),
-			LexToString(ImportedPackageStatus), LexToString(ImportedPackageRef.GetPackageHeaderExtension()), LexToString(ImportedPackagePath.GetHeaderExtension()));
-
 		FAsyncPackage2* ImportedPackage = nullptr;
 		bool bInserted = false;
-		bool bIsFullyLoaded = ImportedPackageRef.AreAllPublicExportsLoaded();
+		bool bIsFullyLoaded = false;
+		{
+			// Add the imported package reference and confirm that its loader and extension match what was found from the package store or import store directly, 
+			// and if not, ensure the package was marked as missing or not installed
+			FWriteLockedPackageRef ImportedPackageRef = ImportStore.AddImportedPackageReference(ImportedPackageId, ImportedPackageUPackageName, ImportedPackageLoader, ImportedPackagePath.GetHeaderExtension());
+			checkf((ImportedPackageStatus == EPackageStoreEntryStatus::Missing || ImportedPackageStatus == EPackageStoreEntryStatus::NotInstalled)
+				|| ImportedPackageRef->GetPackageLoader() == ImportedPackageLoader,
+				TEXT("ImportPackageStatus=%s, ImportedPackageRef.GetPackageLoader()=%s, ImportedPackageLoader=%s"),
+				LexToString(ImportedPackageStatus), LexToString(ImportedPackageRef->GetPackageLoader()), LexToString(ImportedPackageLoader));
+			checkf((ImportedPackageStatus == EPackageStoreEntryStatus::Missing || ImportedPackageStatus == EPackageStoreEntryStatus::NotInstalled)
+				|| ImportedPackageRef->GetPackageHeaderExtension() == ImportedPackagePath.GetHeaderExtension(),
+				TEXT("ImportPackageStatus=%s, ImportedPackageRef.GetPackageHeaderExtension()=%s, ImportedPackagePath.GetHeaderExtension()=%s"),
+				LexToString(ImportedPackageStatus), LexToString(ImportedPackageRef->GetPackageHeaderExtension()), LexToString(ImportedPackagePath.GetHeaderExtension()));
+			
+			bIsFullyLoaded = ImportedPackageRef->AreAllPublicExportsLoaded();
 
-		if (ImportedPackageLoader == EPackageLoader::LinkerLoad && (!ImportedPackageRef.HasPackage() || (!ImportedPackageRef.GetPackage()->GetLinker() && !ImportedPackageRef.GetPackage()->HasAnyPackageFlags(PKG_InMemoryOnly))))
-		{
-			// If we're importing a linker load package and it doesn't have its linker we need to reload it, otherwise we can't reliably link to its imports
-			// Note: Legacy loader path appears to do this only for uncooked packages in the editor?
-			bIsFullyLoaded = false;
-		}
+			if (ImportedPackageLoader == EPackageLoader::LinkerLoad && (!ImportedPackageRef->HasPackage() || (!ImportedPackageRef->GetPackage()->GetLinker() && !ImportedPackageRef->GetPackage()->HasAnyPackageFlags(PKG_InMemoryOnly))))
+			{
+				// If we're importing a linker load package and it doesn't have its linker we need to reload it, otherwise we can't reliably link to its imports
+				// Note: Legacy loader path appears to do this only for uncooked packages in the editor?
+				bIsFullyLoaded = false;
+			}
 
-		const bool bIsNotInstalled = ImportedPackageStatus == EPackageStoreEntryStatus::NotInstalled;
-		const bool bIsMissing = ImportedPackageStatus == EPackageStoreEntryStatus::Missing;
-		if (bIsFullyLoaded)
-		{
-			ImportedPackage = AsyncLoadingThread.FindAsyncPackage(ImportedPackageId);
-			if (!ImportedPackage)
+			const bool bIsNotInstalled = ImportedPackageStatus == EPackageStoreEntryStatus::NotInstalled;
+			const bool bIsMissing = ImportedPackageStatus == EPackageStoreEntryStatus::Missing;
+			if (bIsFullyLoaded)
 			{
-				continue;
+				ImportedPackage = AsyncLoadingThread.FindAsyncPackage(ImportedPackageId);
+				if (!ImportedPackage)
+				{
+					continue;
+				}
+				bInserted = false;
 			}
-			bInserted = false;
-		}
-		else if (bIsNotInstalled)
-		{
-			UE_ASYNC_PACKAGE_LOG(Log, Desc, TEXT("ImportPackages: SkipPackage"),
-				TEXT("Skipping imported package %s (0x%s) because it or one of its dependencies is not installed"), *ImportedPackageNameToLoad.ToString(), *LexToString(ImportedPackageId));
-			ImportedPackageRef.SetIsNotInstalledPackage();
-			continue;
-		}
-		else if (bIsMissing)
-		{
-			if (!ImportedPackageRef.HasPackage()) // If we found a package it's not actually missing but we can't load it anyway
+			else if (bIsNotInstalled)
 			{
 				UE_ASYNC_PACKAGE_LOG(Log, Desc, TEXT("ImportPackages: SkipPackage"),
-					TEXT("Skipping non mounted imported package %s (0x%s)"), *ImportedPackageNameToLoad.ToString(), *LexToString(ImportedPackageId));
-				if (!FLinkerLoad::IsKnownMissingPackage(ImportedPackageNameToLoad))
+					TEXT("Skipping imported package %s (0x%s) because it or one of its dependencies is not installed"), *ImportedPackageNameToLoad.ToString(), *LexToString(ImportedPackageId));
+				ImportedPackageRef->SetIsNotInstalledPackage();
+				continue;
+			}
+			else if (bIsMissing)
+			{
+				if (!ImportedPackageRef->HasPackage()) // If we found a package it's not actually missing but we can't load it anyway
 				{
-					FName PackagePathToLoadName = Desc.PackagePathToLoad.GetPackageFName();
-					ExecuteOnGameThread(TEXT("GetExplanationForUnavailablePackage"), [ImportedPackageNameToLoad, ImportedPackageId, PackagePathToLoadName]()
+					UE_ASYNC_PACKAGE_LOG(Log, Desc, TEXT("ImportPackages: SkipPackage"),
+						TEXT("Skipping non mounted imported package %s (0x%s)"), *ImportedPackageNameToLoad.ToString(), *LexToString(ImportedPackageId));
+					if (!FLinkerLoad::IsKnownMissingPackage(ImportedPackageNameToLoad))
 					{
-						FString ImportedPackageNameToLoadString = ImportedPackageNameToLoad.ToString();
-						FString PackagePathToLoadString = PackagePathToLoadName.ToString();
-						TStringBuilder<2048> Explanation;
-						FPackageName::GetExplanationForUnavailablePackage(ImportedPackageNameToLoad, Explanation);
-						if (Explanation.Len())
+						FName PackagePathToLoadName = Desc.PackagePathToLoad.GetPackageFName();
+						ExecuteOnGameThread(TEXT("GetExplanationForUnavailablePackage"), [ImportedPackageNameToLoad, ImportedPackageId, PackagePathToLoadName]()
 						{
-							FMessageLog("LoadErrors").Info(FText::FormatNamed(NSLOCTEXT("Core", "AsyncLoading_SkippedPackage_Explanation", "While trying to load package {MissingPackage}, a dependent package {DependentPackage} ({DependentPackageId}) was not available. Additional explanatory information follows:\n{Explanation}"),
-								TEXT("MissingPackage"), FText::FromString(PackagePathToLoadString),
-								TEXT("DependentPackage"), FText::FromString(ImportedPackageNameToLoadString),
-								TEXT("DependentPackageId"), FText::FromString(LexToString(ImportedPackageId)),
-								TEXT("Explanation"), FText::FromString(Explanation.ToString())));
-						}
-						else
-						{
-							FMessageLog("LoadErrors").Info(FText::FormatNamed(NSLOCTEXT("Core", "AsyncLoading_SkippedPackage_NoExplanation", "While trying to load package {MissingPackage}, a dependent package {DependentPackage} ({DependentPackageId}) was not available. No additional explanation was available."),
-								TEXT("MissingPackage"), FText::FromString(PackagePathToLoadString),
-								TEXT("DependentPackage"), FText::FromString(ImportedPackageNameToLoadString),
-								TEXT("DependentPackageId"), FText::FromString(LexToString(ImportedPackageId))));
-						}
+							FString ImportedPackageNameToLoadString = ImportedPackageNameToLoad.ToString();
+							FString PackagePathToLoadString = PackagePathToLoadName.ToString();
+							TStringBuilder<2048> Explanation;
+							FPackageName::GetExplanationForUnavailablePackage(ImportedPackageNameToLoad, Explanation);
+							if (Explanation.Len())
+							{
+								FMessageLog("LoadErrors").Info(FText::FormatNamed(NSLOCTEXT("Core", "AsyncLoading_SkippedPackage_Explanation", "While trying to load package {MissingPackage}, a dependent package {DependentPackage} ({DependentPackageId}) was not available. Additional explanatory information follows:\n{Explanation}"),
+									TEXT("MissingPackage"), FText::FromString(PackagePathToLoadString),
+									TEXT("DependentPackage"), FText::FromString(ImportedPackageNameToLoadString),
+									TEXT("DependentPackageId"), FText::FromString(LexToString(ImportedPackageId)),
+									TEXT("Explanation"), FText::FromString(Explanation.ToString())));
+							}
+							else
+							{
+								FMessageLog("LoadErrors").Info(FText::FormatNamed(NSLOCTEXT("Core", "AsyncLoading_SkippedPackage_NoExplanation", "While trying to load package {MissingPackage}, a dependent package {DependentPackage} ({DependentPackageId}) was not available. No additional explanation was available."),
+									TEXT("MissingPackage"), FText::FromString(PackagePathToLoadString),
+									TEXT("DependentPackage"), FText::FromString(ImportedPackageNameToLoadString),
+									TEXT("DependentPackageId"), FText::FromString(LexToString(ImportedPackageId))));
+							}
 
-						return false;
-					});
+							return false;
+						});
+					}
+					ImportedPackageRef->SetIsMissingPackage();
 				}
-				ImportedPackageRef.SetIsMissingPackage();
+				continue;
+			}
+			else
+			{
+				// Release this to avoid deadlocks
+				ImportedPackageRef.Reset();
+				FAsyncPackageDesc2 PackageDesc = FAsyncPackageDesc2::FromPackageImport(Desc, ImportedPackageUPackageName, ImportedPackageId, ImportedPackageIdToLoad, MoveTemp(ImportedPackagePath), ImportedPackageLoader);
+				ImportedPackage = AsyncLoadingThread.FindOrInsertPackage(ThreadState, PackageDesc, bInserted, this);
 			}
-			continue;
-		}
-		else
-		{
-			FAsyncPackageDesc2 PackageDesc = FAsyncPackageDesc2::FromPackageImport(Desc, ImportedPackageUPackageName, ImportedPackageId, ImportedPackageIdToLoad, MoveTemp(ImportedPackagePath), ImportedPackageLoader);
-			ImportedPackage = AsyncLoadingThread.FindOrInsertPackage(ThreadState, PackageDesc, bInserted, this);
 		}
 
 		checkf(ImportedPackage, TEXT("Failed to find or insert imported package with id '%s'"), *FormatPackageId(ImportedPackageId));
@@ -7287,10 +7457,10 @@ EEventLoadNodeExecutionResult FAsyncPackage2::ProcessLinkerLoadPackageSummary(FA
 	// Package that can't be imported are never registered in the global import store, don't try to search for them
 	if (Desc.bCanBeImported)
 	{
-		FLoadedPackageRef& PackageRef = AsyncLoadingThread.GlobalImportStore.FindPackageRefChecked(Desc.UPackageId, Desc.UPackageName);
+		FWriteLockedPackageRef PackageRef = AsyncLoadingThread.GlobalImportStore.FindPackageRefCheckedForWrite(Desc.UPackageId, Desc.UPackageName);
 		if (!HasLoadFailed())
 		{
-			PackageRef.PreInsertPublicExports({ LinkerLoadState->LinkerLoadHeaderData.ExportMap });
+			PackageRef->PreInsertPublicExports({ LinkerLoadState->LinkerLoadHeaderData.ExportMap });
 
 #if WITH_METADATA
 			UE::TUniqueLock LinkerLock(LinkerLoadState->Linker->Mutex);
@@ -7310,7 +7480,7 @@ EEventLoadNodeExecutionResult FAsyncPackage2::ProcessLinkerLoadPackageSummary(FA
 		}
 		else
 		{
-			PackageRef.SetHasFailed();
+			PackageRef->SetHasFailed();
 		}
 	}
 
@@ -7659,7 +7829,7 @@ EEventLoadNodeExecutionResult FAsyncPackage2::ExecuteDeferredPostLoadLinkerLoadP
 	return EEventLoadNodeExecutionResult::Complete;
 }
 
-EEventLoadNodeExecutionResult FAsyncPackage2::Event_CreateExports(FAsyncLoadingThreadState2& ThreadState, FAsyncPackage2* Package, int32)
+EEventLoadNodeExecutionResult FAsyncPackage2::Event_CreateExports(FAsyncLoadingThreadState2& ThreadState, FAsyncPackage2* Package)
 {
 	TRACE_CPUPROFILER_EVENT_SCOPE(Event_CreateExports);
 	UE_ASYNC_PACKAGE_DEBUG(Package->Desc);
@@ -7723,7 +7893,7 @@ void FAsyncPackage2::ConditionalBeginPostLoadPhase(FAsyncLoadingThreadState2& Th
 	);
 }
 
-EEventLoadNodeExecutionResult FAsyncPackage2::Event_ResolveImports(FAsyncLoadingThreadState2& ThreadState, FAsyncPackage2* Package, int32)
+EEventLoadNodeExecutionResult FAsyncPackage2::Event_ResolveImports(FAsyncLoadingThreadState2& ThreadState, FAsyncPackage2* Package)
 {
 	TRACE_CPUPROFILER_EVENT_SCOPE(Event_ResolveImports);
 	UE_ASYNC_PACKAGE_DEBUG(Package->Desc);
@@ -7749,7 +7919,7 @@ EEventLoadNodeExecutionResult FAsyncPackage2::Event_ResolveImports(FAsyncLoading
 	return EEventLoadNodeExecutionResult::Complete;
 }
 
-EEventLoadNodeExecutionResult FAsyncPackage2::Event_PreloadExports(FAsyncLoadingThreadState2& ThreadState, FAsyncPackage2* Package, int32)
+EEventLoadNodeExecutionResult FAsyncPackage2::Event_PreloadExports(FAsyncLoadingThreadState2& ThreadState, FAsyncPackage2* Package)
 {
 	TRACE_CPUPROFILER_EVENT_SCOPE(Event_PreloadExports);
 	UE_ASYNC_PACKAGE_DEBUG(Package->Desc);
@@ -7798,7 +7968,7 @@ EEventLoadNodeExecutionResult FAsyncPackage2::Event_PreloadExports(FAsyncLoading
 	return EEventLoadNodeExecutionResult::Complete;
 }
 
-EEventLoadNodeExecutionResult FAsyncPackage2::Event_ProcessPackageSummary(FAsyncLoadingThreadState2& ThreadState, FAsyncPackage2* Package, int32)
+EEventLoadNodeExecutionResult FAsyncPackage2::Event_ProcessPackageSummary(FAsyncLoadingThreadState2& ThreadState, FAsyncPackage2* Package)
 {
 	TRACE_CPUPROFILER_EVENT_SCOPE(Event_ProcessPackageSummary);
 	UE_ASYNC_PACKAGE_DEBUG(Package->Desc);
@@ -7823,8 +7993,8 @@ EEventLoadNodeExecutionResult FAsyncPackage2::Event_ProcessPackageSummary(FAsync
 	{
 		if (Package->Desc.bCanBeImported)
 		{
-			FLoadedPackageRef& PackageRef = Package->AsyncLoadingThread.GlobalImportStore.FindPackageRefChecked(Package->Desc.UPackageId, Package->Desc.UPackageName);
-			PackageRef.SetHasFailed();
+			FWriteLockedPackageRef PackageRef = Package->AsyncLoadingThread.GlobalImportStore.FindPackageRefCheckedForWrite(Package->Desc.UPackageId, Package->Desc.UPackageName);
+			PackageRef->SetHasFailed();
 		}
 	}
 	else
@@ -7842,16 +8012,16 @@ EEventLoadNodeExecutionResult FAsyncPackage2::Event_ProcessPackageSummary(FAsync
 #endif
 		if (Package->Desc.bCanBeImported)
 		{
-			FLoadedPackageRef& PackageRef = Package->AsyncLoadingThread.GlobalImportStore.FindPackageRefChecked(Package->Desc.UPackageId, Package->Desc.UPackageName);
+			FWriteLockedPackageRef PackageRef = Package->AsyncLoadingThread.GlobalImportStore.FindPackageRefCheckedForWrite(Package->Desc.UPackageId, Package->Desc.UPackageName);
 #if WITH_EDITOR
 			if (Package->OptionalSegmentHeaderData.IsSet())
 			{
-				PackageRef.PreInsertPublicExports({ Package->HeaderData.ExportMap, Package->OptionalSegmentHeaderData->ExportMap });
+				PackageRef->PreInsertPublicExports({ Package->HeaderData.ExportMap, Package->OptionalSegmentHeaderData->ExportMap });
 			}
 			else
 #endif
 			{
-				PackageRef.PreInsertPublicExports({ Package->HeaderData.ExportMap });
+				PackageRef->PreInsertPublicExports({ Package->HeaderData.ExportMap });
 			}
 		}
 
@@ -7901,7 +8071,7 @@ EEventLoadNodeExecutionResult FAsyncPackage2::Event_ProcessPackageSummary(FAsync
 	return EEventLoadNodeExecutionResult::Complete;
 }
 
-EEventLoadNodeExecutionResult FAsyncPackage2::Event_ProcessDependencies(FAsyncLoadingThreadState2& ThreadState, FAsyncPackage2* Package, int32)
+EEventLoadNodeExecutionResult FAsyncPackage2::Event_ProcessDependencies(FAsyncLoadingThreadState2& ThreadState, FAsyncPackage2* Package)
 {
 	Package->AsyncPackageLoadingState = EAsyncPackageLoadingState2::WaitingForDependencies;
 
@@ -7915,7 +8085,7 @@ EEventLoadNodeExecutionResult FAsyncPackage2::Event_ProcessDependencies(FAsyncLo
 	return EEventLoadNodeExecutionResult::Complete;
 }
 
-EEventLoadNodeExecutionResult FAsyncPackage2::Event_DependenciesReady(FAsyncLoadingThreadState2& ThreadState, FAsyncPackage2* Package, int32)
+EEventLoadNodeExecutionResult FAsyncPackage2::Event_DependenciesReady(FAsyncLoadingThreadState2& ThreadState, FAsyncPackage2* Package)
 {
 	TRACE_CPUPROFILER_EVENT_SCOPE(Event_DependenciesReady);
 	UE_ASYNC_PACKAGE_DEBUG(Package->Desc);
@@ -9302,7 +9472,7 @@ void FAsyncPackage2::ConditionalReleasePartialRequests(FAsyncLoadingThreadState2
 	}
 }
 
-EEventLoadNodeExecutionResult FAsyncPackage2::Event_ExportsDone(FAsyncLoadingThreadState2& ThreadState, FAsyncPackage2* Package, int32)
+EEventLoadNodeExecutionResult FAsyncPackage2::Event_ExportsDone(FAsyncLoadingThreadState2& ThreadState, FAsyncPackage2* Package)
 {
 	TRACE_CPUPROFILER_EVENT_SCOPE(Event_ExportsDone);
 	UE_ASYNC_PACKAGE_DEBUG(Package->Desc);
@@ -9316,9 +9486,9 @@ EEventLoadNodeExecutionResult FAsyncPackage2::Event_ExportsDone(FAsyncLoadingThr
 
 	if (!Package->HasLoadFailed() && Package->Desc.bCanBeImported)
 	{
-		FLoadedPackageRef& PackageRef = Package->AsyncLoadingThread.GlobalImportStore.FindPackageRefChecked(Package->Desc.UPackageId, Package->Desc.UPackageName);
+		FWriteLockedPackageRef PackageRef = Package->AsyncLoadingThread.GlobalImportStore.FindPackageRefCheckedForWrite(Package->Desc.UPackageId, Package->Desc.UPackageName);
 		const bool bSnapshotExportCount = true;
-		PackageRef.SetAllPublicExportsLoaded(bSnapshotExportCount);
+		PackageRef->SetAllPublicExportsLoaded(bSnapshotExportCount);
 	}
 
 	if (!Package->Data.ShaderMapHashes.IsEmpty())
@@ -9597,7 +9767,7 @@ void FAsyncPackage2::ConditionalFinishLoading(FAsyncLoadingThreadState2& ThreadS
 		});
 }
 
-EEventLoadNodeExecutionResult FAsyncPackage2::Event_PostLoad(FAsyncLoadingThreadState2& ThreadState, FAsyncPackage2* Package, int32)
+EEventLoadNodeExecutionResult FAsyncPackage2::Event_PostLoad(FAsyncLoadingThreadState2& ThreadState, FAsyncPackage2* Package)
 {
 	TRACE_CPUPROFILER_EVENT_SCOPE(Event_PostLoad);
 	UE_ASYNC_PACKAGE_DEBUG(Package->Desc);
@@ -9732,7 +9902,7 @@ EEventLoadNodeExecutionResult FAsyncPackage2::Event_PostLoad(FAsyncLoadingThread
 	return EEventLoadNodeExecutionResult::Complete;
 }
 
-EEventLoadNodeExecutionResult FAsyncPackage2::Event_DeferredPostLoad(FAsyncLoadingThreadState2& ThreadState, FAsyncPackage2* Package, int32)
+EEventLoadNodeExecutionResult FAsyncPackage2::Event_DeferredPostLoad(FAsyncLoadingThreadState2& ThreadState, FAsyncPackage2* Package)
 {
 	SCOPE_CYCLE_COUNTER(STAT_FAsyncPackage_PostLoadObjectsGameThread);
 	TRACE_CPUPROFILER_EVENT_SCOPE(Event_DeferredPostLoad);
@@ -10955,10 +11125,10 @@ void FAsyncLoadingThread2::CollectUnreachableObjects(
 				if (Package->bCanBeImported)
 				{
 					FPackageId PackageId = Package->GetPackageId();
-					FLoadedPackageRef* PackageRef = GlobalImportStore.FindPackageRef(PackageId);
+					FReadLockedPackageRef PackageRef = GlobalImportStore.FindPackageRef(PackageId);
 					if (PackageRef)
 					{
-						GlobalImportStore.VerifyPackageForRemoval(*PackageRef);
+						GlobalImportStore.VerifyPackageForRemovalNoLock(*PackageRef);
 					}
 				}
 			}
@@ -11155,55 +11325,57 @@ void FAsyncLoadingThread2::NotifyCompiledVersePackage(Verse::VPackage* VersePack
 	UPackage* Package = VersePackage->GetOrCreateUPackage(Context);
 	FName PackageName = Package->GetFName();
 	FPackageId PackageId = FPackageId::FromName(PackageName);
-	FLoadedPackageRef& PackageRef = GlobalImportStore.AddPackageRef(PackageId, PackageName, EPackageLoader::Unknown, EPackageExtension::Unspecified);
+	{
+		FWriteLockedPackageRef PackageRef = GlobalImportStore.AddPackageRef(PackageId, PackageName, EPackageLoader::Unknown, EPackageExtension::Unspecified);
 
 #if !WITH_EDITOR
-	ensureMsgf(!PackageRef.HasPackage(), TEXT("Compiled Verse package %s has already been added to the loader."), *PackageName.ToString());
-	PackageRef.SetPackage(Package);
-	Package->SetCanBeImportedFlag(true);
-	Package->SetPackageId(PackageId);
+		ensureMsgf(!PackageRef->HasPackage(), TEXT("Compiled Verse package %s has already been added to the loader."), *PackageName.ToString());
+		PackageRef->SetPackage(Package);
+		Package->SetCanBeImportedFlag(true);
+		Package->SetPackageId(PackageId);
 #endif
-	PackageRef.SetAllPublicExportsLoaded(true);
-	Package->SetInternalFlags(EInternalObjectFlags::LoaderImport);
+		PackageRef->SetAllPublicExportsLoaded(true);
+		Package->SetInternalFlags(EInternalObjectFlags::LoaderImport);
 
-	ForEachObjectWithOuter(Package, [this, PackageId](UObject* Object)
-	{
-		if (Object->HasAllFlags(RF_Public))
+		ForEachObjectWithOuter(Package, [this, PackageId](UObject* Object)
 		{
-			Object->SetInternalFlags(EInternalObjectFlags::LoaderImport);
-			Object->SetFlags(RF_WasLoaded | RF_LoadCompleted);
-
-			TArray<FName, TInlineAllocator<64>> FullPath;
-			FullPath.Add(Object->GetFName());
-			UObject* Outer = Object->GetOuter();
-			while (Outer)
+			if (Object->HasAllFlags(RF_Public))
 			{
-				FullPath.Add(Outer->GetFName());
-				Outer = Outer->GetOuter();
-			}
-			TStringBuilder<256> PackageRelativeExportPath;
-			for (int32 PathIndex = FullPath.Num() - 2; PathIndex >= 0; --PathIndex)
-			{
-				TCHAR NameStr[FName::StringBufferSize];
-				uint32 NameLen = FullPath[PathIndex].ToString(NameStr);
-				for (uint32 I = 0; I < NameLen; ++I)
+				Object->SetInternalFlags(EInternalObjectFlags::LoaderImport);
+				Object->SetFlags(RF_WasLoaded | RF_LoadCompleted);
+
+				TArray<FName, TInlineAllocator<64>> FullPath;
+				FullPath.Add(Object->GetFName());
+				UObject* Outer = Object->GetOuter();
+				while (Outer)
 				{
-					NameStr[I] = TChar<TCHAR>::ToLower(NameStr[I]);
+					FullPath.Add(Outer->GetFName());
+					Outer = Outer->GetOuter();
 				}
-				PackageRelativeExportPath.AppendChar('/');
-				PackageRelativeExportPath.Append(FStringView(NameStr, NameLen));
+				TStringBuilder<256> PackageRelativeExportPath;
+				for (int32 PathIndex = FullPath.Num() - 2; PathIndex >= 0; --PathIndex)
+				{
+					TCHAR NameStr[FName::StringBufferSize];
+					uint32 NameLen = FullPath[PathIndex].ToString(NameStr);
+					for (uint32 I = 0; I < NameLen; ++I)
+					{
+						NameStr[I] = TChar<TCHAR>::ToLower(NameStr[I]);
+					}
+					PackageRelativeExportPath.AppendChar('/');
+					PackageRelativeExportPath.Append(FStringView(NameStr, NameLen));
+				}
+				uint64 ExportHash = CityHash64(reinterpret_cast<const char*>(PackageRelativeExportPath.GetData() + 1), (PackageRelativeExportPath.Len() - 1) * sizeof(TCHAR));
+				GlobalImportStore.StoreGlobalObject(PackageId, ExportHash, Object);
 			}
-			uint64 ExportHash = CityHash64(reinterpret_cast<const char*>(PackageRelativeExportPath.GetData() + 1), (PackageRelativeExportPath.Len() - 1) * sizeof(TCHAR));
-			GlobalImportStore.StoreGlobalObject(PackageId, ExportHash, Object);
-		}
-	}, /* bIncludeNestedObjects*/ true);
-	for (int32 Index = 0; Index < VersePackage->NumDefinitions(); Index++)
-	{
-		if (Verse::VCell* Cell = VersePackage->GetDefinition(Index).DynamicCast<Verse::VCell>())
+		}, /* bIncludeNestedObjects*/ true);
+		for (int32 Index = 0; Index < VersePackage->NumDefinitions(); Index++)
 		{
-			FUtf8StringView VersePath = VersePackage->GetDefinitionName(Index).AsStringView();
-			uint64 ExportHash = CityHash64(reinterpret_cast<const char*>(VersePath.GetData()), VersePath.Len());
-			GlobalImportStore.StoreGlobalCell(PackageId, ExportHash, Cell);
+			if (Verse::VCell* Cell = VersePackage->GetDefinition(Index).DynamicCast<Verse::VCell>())
+			{
+				FUtf8StringView VersePath = VersePackage->GetDefinitionName(Index).AsStringView();
+				uint64 ExportHash = CityHash64(reinterpret_cast<const char*>(VersePath.GetData()), VersePath.Len());
+				GlobalImportStore.StoreGlobalCell(PackageId, ExportHash, Cell);
+			}
 		}
 	}
 
@@ -11254,7 +11426,7 @@ void FAsyncPackage2::CreatePackageNodes(const FAsyncLoadEventSpec* EventSpecs)
 	FEventLoadNode2* Node = reinterpret_cast<FEventLoadNode2*>(PackageNodesMemory);
 	for (int32 Phase = 0; Phase < EEventLoadNode2::Package_NumPhases; ++Phase)
 	{
-		new (Node + Phase) FEventLoadNode2(EventSpecs + Phase, this, -1, BarrierCount);
+		new (Node + Phase) FEventLoadNode2(EventSpecs + Phase, this, BarrierCount);
 	}
 	PackageNodes = MakeArrayView(Node, EEventLoadNode2::Package_NumPhases);
 }
@@ -11525,20 +11697,20 @@ void FAsyncPackage2::CreateUPackage()
 	// temp packages can't be imported and are never stored or found in global import store
 	if (Desc.bCanBeImported)
 	{
-		FLoadedPackageRef& PackageRef = AsyncLoadingThread.GlobalImportStore.FindPackageRefChecked(Desc.UPackageId, Desc.UPackageName);
-		UPackage* ExistingPackage = PackageRef.GetPackage();
+		FWriteLockedPackageRef PackageRef = AsyncLoadingThread.GlobalImportStore.FindPackageRefCheckedForWrite(Desc.UPackageId, Desc.UPackageName);
+		UPackage* ExistingPackage = PackageRef->GetPackage();
 		if (!ExistingPackage)
 		{
-			PackageRef.SetPackage(LinkerRoot);
+			PackageRef->SetPackage(LinkerRoot);
 		}
 		else if (ExistingPackage != LinkerRoot)
 		{
 			UE_ASYNC_PACKAGE_LOG(Warning, Desc, TEXT("CreateUPackage: ReplacePackage"),
 				TEXT("Replacing renamed package %s (0x%s) while being referenced by the loader, RefCount=%d"),
 				*ExistingPackage->GetName(), *LexToString(ExistingPackage->GetPackageId()),
-				PackageRef.GetRefCount());
+				PackageRef->GetRefCount());
 
-			AsyncLoadingThread.GlobalImportStore.ReplaceReferencedRenamedPackage(PackageRef, LinkerRoot);
+			AsyncLoadingThread.GlobalImportStore.ReplaceReferencedRenamedPackageNoLock(*PackageRef, LinkerRoot);
 		}
 	}
 
