[ZenLoader]
- Fix deadlock in verse VM

diff --git a/Engine/Source/Runtime/CoreUObject/Private/Serialization/AsyncLoading2.cpp b/Engine/Source/Runtime/CoreUObject/Private/Serialization/AsyncLoading2.cpp
--- a/Engine/Source/Runtime/CoreUObject/Private/Serialization/AsyncLoading2.cpp
+++ b/Engine/Source/Runtime/CoreUObject/Private/Serialization/AsyncLoading2.cpp
@@ -11326,16 +11326,18 @@ void FAsyncLoadingThread2::NotifyCompiledVersePackage(Verse::VPackage* VersePack
 	FName PackageName = Package->GetFName();
 	FPackageId PackageId = FPackageId::FromName(PackageName);
 	{
-		FWriteLockedPackageRef PackageRef = GlobalImportStore.AddPackageRef(PackageId, PackageName, EPackageLoader::Unknown, EPackageExtension::Unspecified);
+		{
+			FWriteLockedPackageRef PackageRef = GlobalImportStore.AddPackageRef(PackageId, PackageName, EPackageLoader::Unknown, EPackageExtension::Unspecified);
 
 #if !WITH_EDITOR
-		ensureMsgf(!PackageRef->HasPackage(), TEXT("Compiled Verse package %s has already been added to the loader."), *PackageName.ToString());
-		PackageRef->SetPackage(Package);
-		Package->SetCanBeImportedFlag(true);
-		Package->SetPackageId(PackageId);
+			ensureMsgf(!PackageRef->HasPackage(), TEXT("Compiled Verse package %s has already been added to the loader."), *PackageName.ToString());
+			PackageRef->SetPackage(Package);
+			Package->SetCanBeImportedFlag(true);
+			Package->SetPackageId(PackageId);
 #endif
-		PackageRef->SetAllPublicExportsLoaded(true);
-		Package->SetInternalFlags(EInternalObjectFlags::LoaderImport);
+			PackageRef->SetAllPublicExportsLoaded(true);
+			Package->SetInternalFlags(EInternalObjectFlags::LoaderImport);
+		}
 
 		ForEachObjectWithOuter(Package, [this, PackageId](UObject* Object)
 		{
