Replaced TSet with TCompactSet for FSetHashBucket

diff --git a/Engine/Source/Runtime/CoreUObject/Private/UObject/UObjectHash.cpp b/Engine/Source/Runtime/CoreUObject/Private/UObject/UObjectHash.cpp
--- a/Engine/Source/Runtime/CoreUObject/Private/UObject/UObjectHash.cpp
+++ b/Engine/Source/Runtime/CoreUObject/Private/UObject/UObjectHash.cpp
@@ -5,6 +5,7 @@
 =============================================================================*/
 
 #include "UObject/UObjectHash.h"
+#include "Containers/CompactSet.h"
 #include "HAL/LowLevelMemTracker.h"
 #include "Misc/AssertionMacros.h"
 #include "Misc/TransactionallySafeCriticalSection.h"
@@ -86,33 +87,33 @@ struct THashBucketIterator;
  */
 struct FSetHashBucket
 {
-	typedef TSet<UObjectBase*>::TIterator TIterator;
+	typedef TCompactSet<UObjectBase*>::TIterator TIterator;
 
-	/** This always empty set is used to get an iterator if the bucket doesn't use a TSet (has only 1 element) */
-	static TSet<UObjectBase*> EmptyBucket;
+	/** This always empty set is used to get an iterator if the bucket doesn't use a TCompactSet (has only 1 element) */
+	static TCompactSet<UObjectBase*> EmptyBucket;
 
 	/*
-	* Acts like an alias for either 1-2 elements or a pointer to a TSet
+	* Acts like an alias for either 1-2 elements or a pointer to a TCompactSet
 	* If these are both null, this bucket is empty
-	* If the first one is null, but the second one is non-null, then the second one is a TSet pointer
+	* If the first one is null, but the second one is non-null, then the second one is a TCompactSet pointer
 	* If the first one is not null, then it is a uobject ptr, and the second ptr is either null or a second element
 	*/
 	void* Elements[2];
 
-	FORCEINLINE TSet<UObjectBase*>* GetItems()
+	FORCEINLINE TCompactSet<UObjectBase*>* GetItems()
 	{
 		if (Elements[1] && !Elements[0])
 		{
-			return (TSet<UObjectBase*>*)Elements[1];
+			return (TCompactSet<UObjectBase*>*)Elements[1];
 		}
 		return nullptr;
 	}
 
-	FORCEINLINE const TSet<UObjectBase*>* GetItems() const
+	FORCEINLINE const TCompactSet<UObjectBase*>* GetItems() const
 	{
 		if (Elements[1] && !Elements[0])
 		{
-			return (TSet<UObjectBase*>*)Elements[1];
+			return (TCompactSet<UObjectBase*>*)Elements[1];
 		}
 		return nullptr;
 	}
@@ -132,14 +133,14 @@ struct FSetHashBucket
 	/** Adds an Object to the bucket */
 	FORCEINLINE void Add(UObjectBase* Object)
 	{
-		TSet<UObjectBase*>* Items = GetItems();
+		TCompactSet<UObjectBase*>* Items = GetItems();
 		if (Items)
 		{
 			Items->Add(Object);
 		}
 		else if (Elements[0] && Elements[1])
 		{
-			Items = new TSet<UObjectBase*>();
+			Items = new TCompactSet<UObjectBase*>();
 			Items->Add((UObjectBase*)Elements[0]);
 			Items->Add((UObjectBase*)Elements[1]);
 			Items->Add(Object);
@@ -161,13 +162,13 @@ struct FSetHashBucket
 	FORCEINLINE int32 Remove(UObjectBase* Object)
 	{
 		int32 Result = 0;
-		TSet<UObjectBase*>* Items = GetItems();
+		TCompactSet<UObjectBase*>* Items = GetItems();
 		if (Items)
 		{
 			Result = Items->Remove(Object);
 			if (Items->Num() <= 2)
 			{
-				auto It = TSet<UObjectBase*>::TIterator(*Items);
+				auto It = TCompactSet<UObjectBase*>::TIterator(*Items);
 				Elements[0] = *It;
 				checkSlow((bool)It);
 				++It;
@@ -192,7 +193,7 @@ struct FSetHashBucket
 	/** Checks if an Object exists in this bucket */
 	FORCEINLINE bool Contains(UObjectBase* Object) const
 	{
-		const TSet<UObjectBase*>* Items = GetItems();
+		const TCompactSet<UObjectBase*>* Items = GetItems();
 		if (Items)
 		{
 			return Items->Contains(Object);
@@ -203,7 +204,7 @@ struct FSetHashBucket
 	/** Returns the number of Objects in this bucket */
 	FORCEINLINE int32 Num() const
 	{
-		const TSet<UObjectBase*>* Items = GetItems();
+		const TCompactSet<UObjectBase*>* Items = GetItems();
 		if (Items)
 		{
 			return Items->Num();
@@ -211,10 +212,10 @@ struct FSetHashBucket
 		return !!Elements[0] + !!Elements[1];
 	}
 
-	/** Returns the amount of memory allocated for and by Items TSet */
+	/** Returns the amount of memory allocated for and by Items TCompactSet */
 	FORCEINLINE SIZE_T GetAllocatedSize() const
 	{
-		const TSet<UObjectBase*>* Items = GetItems();
+		const TCompactSet<UObjectBase*>* Items = GetItems();
 		if (Items)
 		{
 			return sizeof(*Items) + Items->GetAllocatedSize();
@@ -224,7 +225,7 @@ struct FSetHashBucket
 
 	void Shrink()
 	{
-		TSet<UObjectBase*>* Items = GetItems();
+		TCompactSet<UObjectBase*>* Items = GetItems();
 		if (Items)
 		{
 			Items->Compact();
@@ -233,15 +234,15 @@ struct FSetHashBucket
 
 	FORCEINLINE struct THashBucketIterator<FSetHashBucket> CreateIterator();
 
-	/** Gets an iterator for the TSet in this bucket or for the EmptyBucker if Items is null */
-	FORCEINLINE TSet<UObjectBase*>::TIterator GetIterator()
+	/** Gets an iterator for the TCompactSet in this bucket or for the EmptyBucker if Items is null */
+	FORCEINLINE TCompactSet<UObjectBase*>::TIterator GetIterator()
 	{
-		TSet<UObjectBase*>* Items = GetItems();
+		TCompactSet<UObjectBase*>* Items = GetItems();
 		return Items ? Items->CreateIterator() : EmptyBucket.CreateIterator();
 	}
 };
 
-TSet<UObjectBase*> FSetHashBucket::EmptyBucket;
+TCompactSet<UObjectBase*> FSetHashBucket::EmptyBucket;
 
 
 /*
