Remove script objects from the async loading package store before unloading their code

#rnx

diff --git a/Engine/Source/Runtime/CoreUObject/Private/Serialization/AsyncLoading2.cpp b/Engine/Source/Runtime/CoreUObject/Private/Serialization/AsyncLoading2.cpp
--- a/Engine/Source/Runtime/CoreUObject/Private/Serialization/AsyncLoading2.cpp
+++ b/Engine/Source/Runtime/CoreUObject/Private/Serialization/AsyncLoading2.cpp
@@ -2515,6 +2515,32 @@ public:
 			}, /* bIncludeNestedObjects*/ true);
 	}
 
+	void RemoveUnreachableScriptObjects()
+	{
+		TRACE_CPUPROFILER_EVENT_SCOPE(RemoveUnreachableScriptObjects);
+
+		int32 RemovedCount = 0;
+
+		{
+			FWriteScopeLock ScopeLock(ScriptObjectsLock);
+
+			for (auto It = ScriptObjects.CreateIterator(); It; ++It)
+			{
+				const UObject* ScriptObject = It->Value;
+				ensure(ScriptObject->IsValidLowLevel());
+
+				if (ScriptObject->IsUnreachable())
+				{
+					UE_LOG(LogStreaming, Verbose, TEXT("RemoveUnreachableScriptObjects: %s"), *ScriptObject->GetName());
+					It.RemoveCurrent();
+					RemovedCount++;
+				}
+			}
+		}
+
+		UE_CLOG(RemovedCount, LogStreaming, Display, TEXT("RemoveUnreachableScriptObjects: removed %d entries"), RemovedCount);
+	}
+
 #if WITH_VERSE_VM || defined(__INTELLISENSE__)
 	void AddScriptCellPackage(Verse::VPackage* Package)
 	{
@@ -11196,9 +11222,18 @@ void FAsyncLoadingThread2::NotifyUnreachableObjects(const TArrayView<FUObjectIte
 	// unreachable objects from last GC should typically have been processed already,
 	// if not handle them here before adding new ones
 	RemoveUnreachableObjects(UnreachableObjects);
+	
+	UE::Tasks::FTask RemoveUnreachableScriptObjectsTask = UE::Tasks::Launch(
+		TEXT("RemoveUnreachableScriptObject"),
+		[this]()
+		{
+			GlobalImportStore.RemoveUnreachableScriptObjects();
+		});
 
 	CollectUnreachableObjects(UnreachableObjectItems, UnreachableObjects);
 
+	RemoveUnreachableScriptObjectsTask.Wait();
+
 	if (GVerifyObjectLoadFlags && !IsAsyncLoading())
 	{
 		GlobalImportStore.VerifyLoadedPackages();
