Reintroduce async loading fix for merged modules, now with the proper API for TStripedMap

#rnx

diff --git a/Engine/Source/Runtime/CoreUObject/Private/Serialization/AsyncLoading2.cpp b/Engine/Source/Runtime/CoreUObject/Private/Serialization/AsyncLoading2.cpp
--- a/Engine/Source/Runtime/CoreUObject/Private/Serialization/AsyncLoading2.cpp
+++ b/Engine/Source/Runtime/CoreUObject/Private/Serialization/AsyncLoading2.cpp
@@ -2515,6 +2515,29 @@ public:
 			}, /* bIncludeNestedObjects*/ true);
 	}
 
+	void RemoveUnreachableScriptObjects()
+	{
+		TRACE_CPUPROFILER_EVENT_SCOPE(RemoveUnreachableScriptObjects);
+
+		const int32 RemovedCount = ScriptObjects.RemoveIf([](const TPair<FPackageObjectIndex, UObject*>& Entry)
+		{
+			const UObject* ScriptObject = Entry.Value;
+			ensure(ScriptObject->IsValidLowLevel());
+
+			if (ScriptObject->IsUnreachable())
+			{
+				UE_LOG(LogStreaming, Verbose, TEXT("RemoveUnreachableScriptObjects: %s"), *ScriptObject->GetName());
+				return true;
+			}
+			else
+			{
+				return false;
+			}
+		});
+
+		UE_CLOG(RemovedCount, LogStreaming, Display, TEXT("RemoveUnreachableScriptObjects: removed %d entries"), RemovedCount);
+	}
+
 #if WITH_VERSE_VM || defined(__INTELLISENSE__)
 	void AddScriptCellPackage(Verse::VPackage* Package)
 	{
