[ZenLoader]
- Restore missing parallelism for RemoveUnreachableScriptObject that got backed out
- Fix potential deadlock in between RemovePublicExports and writers to the global import table

diff --git a/Engine/Source/Runtime/CoreUObject/Private/Serialization/AsyncLoading2.cpp b/Engine/Source/Runtime/CoreUObject/Private/Serialization/AsyncLoading2.cpp
--- a/Engine/Source/Runtime/CoreUObject/Private/Serialization/AsyncLoading2.cpp
+++ b/Engine/Source/Runtime/CoreUObject/Private/Serialization/AsyncLoading2.cpp
@@ -2241,6 +2241,13 @@ public:
 				FPackageId PackageId = PublicExportKey.GetPackageId();
 				if (PackageId != LastPackageId)
 				{
+					// Reset the current PackageRef since our intention is to get another one
+					// and we need to release the sharedlock otherwise we could deadlock if
+					// a writelock inserts itself on the mutex, we won't be able to acquire
+					// the new readlock until the writelock has been satisfied, but if we don't release
+					// our sharedlock, it would prevent the writelock from making progress
+					// so we would effectively deadlock.
+					PackageRef.Reset();
 					UE_ASYNC_PACKAGEID_DEBUG(PackageId);
 					LastPackageId = PackageId;
 					PackageRef = FindPackageRef(PackageId);
@@ -11219,9 +11226,18 @@ void FAsyncLoadingThread2::NotifyUnreachableObjects(const TArrayView<FUObjectIte
 	// unreachable objects from last GC should typically have been processed already,
 	// if not handle them here before adding new ones
 	RemoveUnreachableObjects(UnreachableObjects);
+	
+	UE::Tasks::FTask RemoveUnreachableScriptObjectsTask = UE::Tasks::Launch(
+		TEXT("RemoveUnreachableScriptObject"),
+		[this]()
+		{
+			GlobalImportStore.RemoveUnreachableScriptObjects();
+		});
 
 	CollectUnreachableObjects(UnreachableObjectItems, UnreachableObjects);
 
+	RemoveUnreachableScriptObjectsTask.Wait();
+
 	if (GVerifyObjectLoadFlags && !IsAsyncLoading())
 	{
 		GlobalImportStore.VerifyLoadedPackages();
