Remove obsolete friend class FVerseObjectClassReplacer on UObjectBase

#rnx

diff --git a/Engine/Source/Runtime/CoreUObject/Public/UObject/UObjectBase.h b/Engine/Source/Runtime/CoreUObject/Public/UObject/UObjectBase.h
--- a/Engine/Source/Runtime/CoreUObject/Public/UObject/UObjectBase.h
+++ b/Engine/Source/Runtime/CoreUObject/Public/UObject/UObjectBase.h
@@ -445,7 +445,6 @@ private:
 	ObjectPtr_Private::TNonAccessTrackedObjectPtr<UObject>						OuterPrivate;
 	
 	friend class FBlueprintCompileReinstancer;
-	friend class FVerseObjectClassReplacer;
 	friend class FContextObjectManager;
 	friend class FDeferredRegistry;
 	friend class FCompiledInObjectRegistry;
