[ZenLoader]
- Parallelize ScriptObjects cleanup by using ParallelRemoveIf and 4 stripes

diff --git a/Engine/Source/Runtime/CoreUObject/Private/Serialization/AsyncLoading2.cpp b/Engine/Source/Runtime/CoreUObject/Private/Serialization/AsyncLoading2.cpp
--- a/Engine/Source/Runtime/CoreUObject/Private/Serialization/AsyncLoading2.cpp
+++ b/Engine/Source/Runtime/CoreUObject/Private/Serialization/AsyncLoading2.cpp
@@ -1727,7 +1727,7 @@ public:
 	FCookedScriptObjectsDebug();
 	~FCookedScriptObjectsDebug();
 	bool IsEnabled() const { return Objects.Num() > 0; }
-	void TrimDebugData(const TAsyncLoadingStripedMap<1, FPackageObjectIndex, UObject*>& RuntimeScriptObjects);
+	void TrimDebugData(const TAsyncLoadingStripedMap<4, FPackageObjectIndex, UObject*>& RuntimeScriptObjects);
 	void OnCompiledInUObjectsRegistered(FName, ECompiledInUObjectsRegisteredStatus Status);
 	FString GetPathName(FPackageObjectIndex Index) const;
 };
@@ -1827,7 +1827,7 @@ private:
 	// Packages in active loading that export VCells. Stored separately from FLoadedPackageRefs to save space.
 	TAsyncLoadingStripedMap<1, FPackageId, FLoadedPackageCellsRef> PackageCells;
 	// All native script objects (structs, enums, classes, CDOs and their subobjects) from the initial load phase
-	TAsyncLoadingStripedMap<1, FPackageObjectIndex, UObject*> ScriptObjects;
+	TAsyncLoadingStripedMap<4, FPackageObjectIndex, UObject*> ScriptObjects;
 	// Built-in VCells always defined by the Verse VM.
 	TAsyncLoadingStripedMap<1, FPackageObjectIndex, Verse::VCell*> ScriptCells;
 	// All currently loaded public export objects from any loaded package
@@ -2527,7 +2527,7 @@ public:
 	{
 		TRACE_CPUPROFILER_EVENT_SCOPE(RemoveUnreachableScriptObjects);
 
-		const int32 RemovedCount = ScriptObjects.RemoveIf([](const TPair<FPackageObjectIndex, UObject*>& Entry)
+		const int32 RemovedCount = ScriptObjects.ParallelRemoveIf([](const TPair<FPackageObjectIndex, UObject*>& Entry)
 		{
 			const UObject* ScriptObject = Entry.Value;
 			ensure(ScriptObject->IsValidLowLevel());
@@ -6703,7 +6703,7 @@ void FCookedScriptObjectsDebug::LoadDebugData()
 	}
 }
 
-void FCookedScriptObjectsDebug::TrimDebugData(const TAsyncLoadingStripedMap<1, FPackageObjectIndex, UObject*>& RuntimeScriptObjects)
+void FCookedScriptObjectsDebug::TrimDebugData(const TAsyncLoadingStripedMap<4, FPackageObjectIndex, UObject*>& RuntimeScriptObjects)
 {
 	if (!IsEnabled())
 	{
