[UObject] Resolve compilation error with duplicate UE::NAME_IsLooseMetadata by making inline const and sharing, which will still be instantiated only once.

diff --git a/Engine/Source/Runtime/CoreUObject/Private/UObject/InstanceDataObjectUtils.cpp b/Engine/Source/Runtime/CoreUObject/Private/UObject/InstanceDataObjectUtils.cpp
--- a/Engine/Source/Runtime/CoreUObject/Private/UObject/InstanceDataObjectUtils.cpp
+++ b/Engine/Source/Runtime/CoreUObject/Private/UObject/InstanceDataObjectUtils.cpp
@@ -17,6 +17,7 @@
 #include "Serialization/ArchiveCountMem.h"
 #include "Serialization/ImportTypeHierarchy.h"
 #include "UObject/Class.h"
+#include "UObject/CommonFNames.h"
 #include "UObject/EnumProperty.h"
 #include "UObject/Field.h"
 #include "UObject/IntrinsicClassFwd.h"
@@ -266,7 +267,6 @@ namespace UE
 {
 	static const FName NAME_DisplayName(ANSITEXTVIEW("DisplayName"));
 	static const FName NAME_PresentAsTypeMetadata(ANSITEXTVIEW("PresentAsType"));
-	static const FName NAME_IsLooseMetadata(ANSITEXTVIEW("IsLoose"));
 	static const FName NAME_IsInstanceDataObjectStruct(ANSITEXTVIEW("IsInstanceDataObjectClass"));
 	static const FName NAME_ContainsLoosePropertiesMetadata(ANSITEXTVIEW("ContainsLooseProperties"));
 	static const FName NAME_VerseClass(ANSITEXTVIEW("VerseClass"));
diff --git a/Engine/Source/Runtime/CoreUObject/Private/UObject/PropertyBagRepository.cpp b/Engine/Source/Runtime/CoreUObject/Private/UObject/PropertyBagRepository.cpp
--- a/Engine/Source/Runtime/CoreUObject/Private/UObject/PropertyBagRepository.cpp
+++ b/Engine/Source/Runtime/CoreUObject/Private/UObject/PropertyBagRepository.cpp
@@ -11,6 +11,7 @@
 #include "Serialization/ObjectWriter.h"
 #include "Serialization/ArchiveProxy.h"
 #include "Serialization/ArchiveCountMem.h"
+#include "UObject/CommonFNames.h"
 #include "UObject/GarbageCollection.h"
 #include "UObject/InstanceDataObjectUtils.h"
 #include "UObject/InstanceDataTransforms.h"
@@ -37,7 +38,6 @@ DEFINE_LOG_CATEGORY_STATIC(LogPropertyBagRepository, Log, All);
 
 namespace UE
 {
-	static const FName NAME_IsLooseMetadata(ANSITEXTVIEW("IsLoose"));
 
 /** Defined in InstanceDataObjectUtils.cpp */
 void CopyTaggedProperties(const UObject* Source, UObject* Dest);
diff --git a/Engine/Source/Runtime/CoreUObject/Public/UObject/CommonFNames.h b/Engine/Source/Runtime/CoreUObject/Public/UObject/CommonFNames.h
new file mode 100644
--- /dev/null
+++ b/Engine/Source/Runtime/CoreUObject/Public/UObject/CommonFNames.h
@@ -0,0 +1,11 @@
+// Copyright Epic Games, Inc. All Rights Reserved.
+
+#pragma once
+
+#include "UObject/NameTypes.h"
+#include "Containers/StringView.h"
+
+namespace UE
+{
+	inline const FName NAME_IsLooseMetadata(ANSITEXTVIEW("IsLoose"));
+};
