[UObjectArray]
- Fix missing atomic read when UE_ENABLE_FUOBJECT_ITEM_PACKING is active

diff --git a/Engine/Source/Runtime/CoreUObject/Public/UObject/UObjectArray.h b/Engine/Source/Runtime/CoreUObject/Public/UObject/UObjectArray.h
--- a/Engine/Source/Runtime/CoreUObject/Public/UObject/UObjectArray.h
+++ b/Engine/Source/Runtime/CoreUObject/Public/UObject/UObjectArray.h
@@ -146,7 +146,7 @@ public:
 	inline class UObjectBase* GetObject() const
 	{
 #if UE_ENABLE_FUOBJECT_ITEM_PACKING
-		const uintptr_t Obj = (uintptr_t((FlagsAndRefCount >> 32) & PtrMask) << (32 + UObjectPtrTrailingZeroes)) | (uintptr_t(ObjectPtrLow) << UObjectPtrTrailingZeroes);
+		const uintptr_t Obj = (uintptr_t((FPlatformAtomics::AtomicRead_Relaxed((int64*)&FlagsAndRefCount) >> 32) & PtrMask) << (32 + UObjectPtrTrailingZeroes)) | (uintptr_t(ObjectPtrLow) << UObjectPtrTrailingZeroes);
 		return (class UObjectBase*)Obj;
 #else
 		PRAGMA_DISABLE_DEPRECATION_WARNINGS
