- Fixed a HybridPackageStore init issue when no zen snapshot exists
- Fixed some issues with getting PackageStore info when we don't have a package name (in ImportPackagesRecursiveInner sets bHasImportedPackageNames to false)

diff --git a/Engine/Source/Runtime/CoreUObject/Private/Serialization/PackageStore.cpp b/Engine/Source/Runtime/CoreUObject/Private/Serialization/PackageStore.cpp
--- a/Engine/Source/Runtime/CoreUObject/Private/Serialization/PackageStore.cpp
+++ b/Engine/Source/Runtime/CoreUObject/Private/Serialization/PackageStore.cpp
@@ -110,7 +110,8 @@ public:
 				if (Entry.LoaderType == EPackageLoader::LinkerLoad)
 				{
 					OutPackageStoreEntry = Entry;
-					if (OutPackagePath)
+					// can't get the package path without a name! (which is valid for cooked flows)
+					if (OutPackagePath && !PackageName.IsNone())
 					{
 						*OutPackagePath = FPackagePath::FromPackageNameUnchecked(PackageName);
 						OutPackagePath->SetHeaderExtension(OutPackageStoreEntry.PackageExtension);
@@ -268,16 +269,19 @@ EPackageStoreEntryStatus FHybridPackageStoreBackend::GetPackageStoreEntry(FPacka
 			// this way, we can just return the PackageEntry directly
 			if (LinkerLoadStore->GetPackageStoreEntryEx(PackageId, PackageName, TempLooseEntry, &TempPath) == EPackageStoreEntryStatus::Ok)
 			{
-				// now that we have determined the local path to the file, check if it's read only
-				const bool bIsReadOnly = FPlatformFileManager::Get().GetPlatformFile().IsReadOnly(*TempPath.GetLocalFullPath());
-				if (!bIsReadOnly)
+				// now that we have determined the local path to the file, check if it's read only (if a valid path was returned)
+				if (!TempPath.IsEmpty())
 				{
-					// if it's writeable, then we want to use the LinkerLoad version, so just return it
-					OutPackageStoreEntry = TempLooseEntry;
+					const bool bIsReadOnly = PackageName == NAME_None || FPlatformFileManager::Get().GetPlatformFile().IsReadOnly(*TempPath.GetLocalFullPath());
+					if (!bIsReadOnly)
+					{
+						// if it's writeable, then we want to use the LinkerLoad version, so just return it
+						OutPackageStoreEntry = TempLooseEntry;
 
-					// we are done!
-					InnerStatus = EPackageStoreEntryStatus::Ok;
-					bAllowZenLoad = false;
+						// we are done!
+						InnerStatus = EPackageStoreEntryStatus::Ok;
+						bAllowZenLoad = false;
+					}
 				}
 			}
 		}
@@ -523,8 +527,12 @@ void FPackageStore::Mount(TSharedRef<IPackageStoreBackend> Backend, int32 Priori
 			{
 				LinkerLoad = Backend;
 			}
-			TSharedPtr<IPackageStoreBackend> Zen = Backend;
-			if (!(Zen->GetSupportedLoaders() & EPackageLoader::Zen) && !!(Backends[0].Value->GetSupportedLoaders() & EPackageLoader::Zen))
+			TSharedPtr<IPackageStoreBackend> Zen;
+			if (!!(Backend->GetSupportedLoaders() & EPackageLoader::Zen))
+			{
+				Zen = Backend;
+			}
+			if (!Zen && Backends.Num() > 0 && !!(Backends[0].Value->GetSupportedLoaders() & EPackageLoader::Zen))
 			{
 				Zen = Backends[0].Value;
 			}
