Added asserts in CopyValuesInternal for corrupt maps or sets.
Added an exit condition when the destination is the same as the source.

#rnx

diff --git a/Engine/Source/Runtime/CoreUObject/Private/UObject/PropertyMap.cpp b/Engine/Source/Runtime/CoreUObject/Private/UObject/PropertyMap.cpp
--- a/Engine/Source/Runtime/CoreUObject/Private/UObject/PropertyMap.cpp
+++ b/Engine/Source/Runtime/CoreUObject/Private/UObject/PropertyMap.cpp
@@ -1259,6 +1259,11 @@ void FMapProperty::CopyValuesInternal(void* Dest, void const* Src, int32 Count)
 {
 	check(Count == 1);
 
+	if (Dest == Src)
+	{
+		return;
+	}
+
 	FScriptMapHelper SrcMapHelper (this, Src);
 	FScriptMapHelper DestMapHelper(this, Dest);
 
@@ -1270,6 +1275,8 @@ void FMapProperty::CopyValuesInternal(void* Dest, void const* Src, int32 Count)
 		return;
 	}
 
+	checkf(SrcMapHelper.GetMaxIndex() >= Num, TEXT("Corrupt source map: Num = %d, MaxIndex = %d"), Num, SrcMapHelper.GetMaxIndex());
+
 	for (int32 SrcIndex = 0; Num; ++SrcIndex)
 	{
 		if (SrcMapHelper.IsValidIndex(SrcIndex))
diff --git a/Engine/Source/Runtime/CoreUObject/Private/UObject/PropertySet.cpp b/Engine/Source/Runtime/CoreUObject/Private/UObject/PropertySet.cpp
--- a/Engine/Source/Runtime/CoreUObject/Private/UObject/PropertySet.cpp
+++ b/Engine/Source/Runtime/CoreUObject/Private/UObject/PropertySet.cpp
@@ -738,6 +738,11 @@ void FSetProperty::AddCppProperty(FProperty* Property)
 void FSetProperty::CopyValuesInternal(void* Dest, void const* Src, int32 Count) const
 {
 	check(Count == 1);
+	
+	if (Dest == Src)
+	{
+		return;
+	}
 
 	FScriptSetHelper SrcSetHelper (this, Src);
 	FScriptSetHelper DestSetHelper(this, Dest);
@@ -750,6 +755,8 @@ void FSetProperty::CopyValuesInternal(void* Dest, void const* Src, int32 Count)
 		return;
 	}
 
+	checkf(SrcSetHelper.GetMaxIndex() >= Num, TEXT("Corrupt source set: Num = %d, MaxIndex = %d"), Num, SrcSetHelper.GetMaxIndex());
+
 	for (int32 SrcIndex = 0; Num; ++SrcIndex)
 	{
 		if (SrcSetHelper.IsValidIndex(SrcIndex))
