[UE] Fix for stomping an annotation that has already been set during edit change propagation.

Notes:
- (Example) This was (incorrectly) marking child BP assets as needing to be saved if they happened to be open for edit alongside their parent in the BP editor, in cases where they would (correctly) not be dirtied if not also open for edit.

#Synced-CL 48038701

diff --git a/Engine/Source/Runtime/CoreUObject/Private/UObject/Obj.cpp b/Engine/Source/Runtime/CoreUObject/Private/UObject/Obj.cpp
--- a/Engine/Source/Runtime/CoreUObject/Private/UObject/Obj.cpp
+++ b/Engine/Source/Runtime/CoreUObject/Private/UObject/Obj.cpp
@@ -834,7 +834,7 @@ void UObject::PropagatePostEditChange( TArray<UObject*>& AffectedObjects, FPrope
 
 void UObject::SetEditChangePropagationFlags(EEditChangePropagationFlags InFlags)
 {
-	FPropagatedEditChangeAnnotation Annotation;
+	FPropagatedEditChangeAnnotation Annotation = PropagatedEditChangeAnnotation.GetAnnotation(this);
 	Annotation.bDeferredMarkAsDirty = !!(InFlags & EEditChangePropagationFlags::OnlyMarkRealignedInstancesAsDirty);
 	PropagatedEditChangeAnnotation.AddAnnotation(this, MoveTemp(Annotation));
 }
