Show localized default ToolTip text if none is specified in meta-data

#tests Run UE5 editor, change language to any but English, navigate to any menu, that has no specified tooltip text (e.g. RMB click Content Browser -> Animation -> Control Rig -> Control Rig Shape Library), observe the tooltip is now localized.
#Synced-CL 48141329

diff --git a/Engine/Source/Runtime/CoreUObject/Private/UObject/Class.cpp b/Engine/Source/Runtime/CoreUObject/Private/UObject/Class.cpp
--- a/Engine/Source/Runtime/CoreUObject/Private/UObject/Class.cpp
+++ b/Engine/Source/Runtime/CoreUObject/Private/UObject/Class.cpp
@@ -368,7 +368,7 @@ FText UField::GetToolTipText(bool bShortTooltip) const
 	{
 		if (NativeToolTip.IsEmpty())
 		{
-			NativeToolTip = FName::NameToDisplayString(FDisplayNameHelper::Get(*this), false);
+			return GetDisplayNameText();
 		}
 		else if (!bShortTooltip && IsNative())
 		{
