Fix problem with the same event sent multitple times through owners of subobjects.

#rnx
#Synced-CL 48133794

diff --git a/Engine/Source/Runtime/CoreUObject/Private/UObject/OverriddenPropertySet.cpp b/Engine/Source/Runtime/CoreUObject/Private/UObject/OverriddenPropertySet.cpp
--- a/Engine/Source/Runtime/CoreUObject/Private/UObject/OverriddenPropertySet.cpp
+++ b/Engine/Source/Runtime/CoreUObject/Private/UObject/OverriddenPropertySet.cpp
@@ -1018,13 +1018,14 @@ void FOverriddenPropertySet::NotifyPropertyChange(FOverriddenPropertyNode* Paren
 			}
 		}
 
-		uint8* FindOrAdd(TNotNull<const UObject*> ContainerOwner, TNotNull<const FProperty*> ContainerProperty)
+		uint8* FindOrAdd(TNotNull<const UObject*> ContainerOwner, TNotNull<const FProperty*> ContainerProperty, bool& bAlreadySaved)
 		{
 			if (uint8* Found = Find(ContainerOwner, ContainerProperty))
 			{
-				UE_LOG(LogOverridableObject, Warning, TEXT("This container owner:%s(0x%p) has already allocated memory for this property:%s"), *GetNameSafe(ContainerOwner), NotNullGet(ContainerOwner), *ContainerProperty->GetName());
+				bAlreadySaved = true;
 				return Found;
 			}
+			bAlreadySaved = false;
 			FContainerInfo Key = {ContainerOwner, ContainerProperty};
 			int32 I = Data.Add({ Key, (uint8*)FMemory::Malloc(ContainerProperty->GetSize(), ContainerProperty->GetMinAlignment()) });
 			ContainerProperty->InitializeValue(Data[I].Value);
@@ -1056,9 +1057,63 @@ void FOverriddenPropertySet::NotifyPropertyChange(FOverriddenPropertyNode* Paren
 			{
 				checkf(ArrayProperty->Inner, TEXT("Expecting an inner type for Arrays"));
 
+				auto ArrayReplace = [&]
+				{
+					if (SubPropertyNode)
+					{
+						// Overriding all entry in the array
+						SubPropertyNode->SetOperation(EOverriddenPropertyOperation::Replace);
+					}
+
+					// This is a case of the entire array is overridden
+					// Need to loop through every sub object and override all properties
+					for (int i = 0; i < ArrayHelper.Num(); ++i)
+					{
+						if (UObject* SubObject = TryGetInstancedSubObjectValue(InnerObjectProperty, ArrayHelper.GetElementPtr(i)))
+						{
+							OverridableManager.OverrideAllObjectProperties(SubObject);
+						}
+					}
+				};
+
+				// Only arrays flagged overridable logic can record deltas, for now just override entire array
+				if (!ArrayProperty->HasAnyPropertyFlags(CPF_ExperimentalOverridableLogic))
+				{
+					if (Notification == EPropertyNotificationType::PostEdit)
+					{
+						if (ChangeType == EPropertyChangeType::Unspecified && ArrayIndex == INDEX_NONE)
+						{
+							// Overriding all entry in the array + override instanced sub objects
+							ArrayReplace();
+						}
+						else if (SubPropertyNode)
+						{
+							// Overriding all entry in the array
+							SubPropertyNode->SetOperation(EOverriddenPropertyOperation::Replace);
+						}
+					}
+					return;
+				}
+
 				if (Notification == EPropertyNotificationType::PreEdit)
 				{
-					FScriptArrayHelper PreEditArrayHelper(ArrayProperty, SavedPreEditContainers.FindOrAdd(Owner, ArrayProperty));
+					bool bAlreadySaved = false;
+					uint8* PreEditArray = SavedPreEditContainers.FindOrAdd(Owner, ArrayProperty, bAlreadySaved);
+					if (bAlreadySaved)
+					{
+						// The saved values needs to be identical, otherwise warns about missing post edit and continue to override with correct values
+						// This happens when we receive multiple time the same event from subobject to owner and so on
+						if (ArrayProperty->Identical(SubValuePtr, PreEditArray, 0))
+						{
+							UE_LOG(LogOverridableObject, Verbose, TEXT("This owner:%s(0x%p) has already handled this property:%s"), 
+								*GetNameSafe(Owner), Owner.Get(), *Property->GetName());
+							return;
+						}
+						UE_LOG(LogOverridableObject, Warning, TEXT("This owner:%s(0x%p) has already a saved value for property(%s) and it is expected be identical, this is an indication of a missing post edit"),
+							*GetNameSafe(Owner), Owner.Get(), *Property->GetName());
+					}
+
+					FScriptArrayHelper PreEditArrayHelper(ArrayProperty, PreEditArray);
 					PreEditArrayHelper.EmptyAndAddValues(ArrayHelper.Num());
 					for (int32 i = 0; i < ArrayHelper.Num(); i++)
 					{
@@ -1075,26 +1130,14 @@ void FOverriddenPropertySet::NotifyPropertyChange(FOverriddenPropertyNode* Paren
 					}
 				};
 
-				FScriptArrayHelper PreEditArrayHelper(ArrayProperty, SavedPreEditContainers.Find(Owner, ArrayProperty));
-
-				auto ArrayReplace = [&]
+				uint8* SavedPreEditArray = SavedPreEditContainers.Find(Owner, ArrayProperty);
+				if (!SavedPreEditArray)
 				{
-					if (SubPropertyNode)
-					{
-						// Overriding all entry in the array
-						SubPropertyNode->SetOperation(EOverriddenPropertyOperation::Replace);
-					}
-
-					// This is a case of the entire array is overridden
-					// Need to loop through every sub object and override all properties
-					for (int i = 0; i < ArrayHelper.Num(); ++i)
-					{
-						if (UObject* SubObject = TryGetInstancedSubObjectValue(InnerObjectProperty, ArrayHelper.GetElementPtr(i)))
-						{
-							OverridableManager.OverrideAllObjectProperties(SubObject);
-						}
-					}
-				};
+					UE_LOG(LogOverridableObject, Verbose, TEXT("Owner:%s(0x%p) The property(%s) seem to already be handled, skipping this post edit notification"), 
+						*GetNameSafe(Owner), Owner.Get(), *Property->GetName());
+					return;
+				}
+				FScriptArrayHelper PreEditArrayHelper(ArrayProperty, SavedPreEditArray);
 
 				auto ArrayAddImpl = [&]()
 				{
@@ -1168,22 +1211,6 @@ void FOverriddenPropertySet::NotifyPropertyChange(FOverriddenPropertyNode* Paren
 					}
 				};
 
-				// Only arrays flagged overridable logic can record deltas, for now just override entire array
-				if (!ArrayProperty->HasAnyPropertyFlags(CPF_ExperimentalOverridableLogic))
-				{
-					if(ChangeType == EPropertyChangeType::Unspecified && ArrayIndex == INDEX_NONE)
-					{
-						// Overriding all entry in the array + override instanced sub objects
-						ArrayReplace();
-					}
-					else if (SubPropertyNode)
-					{
-						// Overriding all entry in the array
-						SubPropertyNode->SetOperation(EOverriddenPropertyOperation::Replace);
-					}
-					return;
-				}
-
 				// Note: Currently, if CPF_ExperimentalOverridableLogic is set, we also require the property to be explicitly marked as an instanced subobject.
 				checkf(InnerObjectProperty->HasAllPropertyFlags(CPF_PersistentInstance), TEXT("Only instanced array properties support experimental overridable logic"));
 
@@ -1249,6 +1276,11 @@ void FOverriddenPropertySet::NotifyPropertyChange(FOverriddenPropertyNode* Paren
 					return;
 				}
 			}
+			else
+			{
+				UE_LOG(LogOverridableObject, Warning, TEXT("Owner:%s(0x%p) Information is suggesting a subobject path but we are unable to forward the notification because we do not know the array index of %s, we might end up with missed match pre/post edit"),
+					*GetNameSafe(Owner), Owner.Get(), *Property->GetName());
+			}
 		}
 	}
 	// @todo support set in the overridable serialization
@@ -1289,7 +1321,23 @@ void FOverriddenPropertySet::NotifyPropertyChange(FOverriddenPropertyNode* Paren
 		{
 			if (Notification == EPropertyNotificationType::PreEdit)
 			{
-				FScriptMapHelper PreEditMapHelper(MapProperty, SavedPreEditContainers.FindOrAdd(Owner, MapProperty));
+				bool bAlreadySaved = false;
+				uint8* PreEditMap = SavedPreEditContainers.FindOrAdd(Owner, MapProperty, bAlreadySaved);
+				if (bAlreadySaved)
+				{
+					// The saved values needs to be identical, otherwise warns about missing post edit and continue to override with correct values
+					// This happens when we receive multiple time the same event from subobject to owner and so on
+					if (MapProperty->Identical(SubValuePtr, PreEditMap, 0))
+					{
+						UE_LOG(LogOverridableObject, Verbose, TEXT("This owner:%s(0x%p) has already handled this property:%s"),
+							*GetNameSafe(Owner), Owner.Get(), *Property->GetName());
+						return;
+					}
+
+					UE_LOG(LogOverridableObject, Warning, TEXT("This owner:%s(0x%p) has already a saved value for property(%s) and it is expected be identical, this is an indication of a missing post edit"),
+						*GetNameSafe(Owner), Owner.Get(), *Property->GetName());
+				}
+				FScriptMapHelper PreEditMapHelper(MapProperty, PreEditMap);
 				PreEditMapHelper.EmptyValues();
 				for (FScriptMapHelper::FIterator It(MapHelper); It; ++It)
 				{
@@ -1299,6 +1347,12 @@ void FOverriddenPropertySet::NotifyPropertyChange(FOverriddenPropertyNode* Paren
 			}
 
 			uint8* SavedPreEditMap = SavedPreEditContainers.Find(Owner, MapProperty);
+			if (!SavedPreEditMap)
+			{
+				UE_LOG(LogOverridableObject, Verbose, TEXT("Owner:%s(0x%p) The property(%s) seem to already be handled, skipping this post edit notification"),
+					*GetNameSafe(Owner), Owner.Get(), *Property->GetName());
+				return;
+			}
 			checkf(SavedPreEditMap, TEXT("Expecting the same property as the pre edit flow"));
 			FScriptMapHelper PreEditMapHelper(MapProperty, SavedPreEditMap);
 			// The logical should map directly to the pre edit map internal index as we skipped all of the invalid entries
@@ -1477,6 +1531,11 @@ void FOverriddenPropertySet::NotifyPropertyChange(FOverriddenPropertyNode* Paren
 				return;
 			}
 		}
+		else
+		{
+			UE_LOG(LogOverridableObject, Warning, TEXT("Owner:%s(0x%p) Information is suggesting a subobject path but we are unable to forward the notification because we do not know the map index of %s, we might end up with missed match pre/post edit"),
+				*GetNameSafe(Owner), Owner.Get(), *Property->GetName());
+		}
 	}
 	else if (const FStructProperty* StructProperty = CastField<FStructProperty>(Property))
 	{
@@ -1501,22 +1560,38 @@ void FOverriddenPropertySet::NotifyPropertyChange(FOverriddenPropertyNode* Paren
 	else if (const FObjectPropertyBase* ObjectProperty = CastField<FObjectPropertyBase>(Property))
 	{
 		UObject* SubObject = TryGetInstancedSubObjectValue(ObjectProperty, SubValuePtr);
-		if (Notification == EPropertyNotificationType::PreEdit)
-		{
-			uint8* PreEditValuePtr = SavedPreEditContainers.FindOrAdd(Owner, ObjectProperty);
-			ObjectProperty->SetObjectPtrPropertyValue(PreEditValuePtr, SubObject);
-		}
-
-		ON_SCOPE_EXIT
+		if (!NextPropertyIterator)
 		{
-			if (Notification == EPropertyNotificationType::PostEdit)
+			if (Notification == EPropertyNotificationType::PreEdit)
 			{
-				SavedPreEditContainers.Free(Owner, ObjectProperty);
+				bool bAlreadySaved = false;
+				uint8* PreEditValuePtr = SavedPreEditContainers.FindOrAdd(Owner, ObjectProperty, bAlreadySaved);
+				if (bAlreadySaved)
+				{
+					// The saved values needs to be identical, otherwise warns about missing post edit and continue to override with correct values
+					// This happens when we receive multiple time the same event from subobject to owner and so on
+					if (ObjectProperty->Identical(SubValuePtr, PreEditValuePtr, 0))
+					{
+						UE_LOG(LogOverridableObject, Verbose, TEXT("This owner:%s(0x%p) has already handled this property:%s"),
+							*GetNameSafe(Owner), Owner.Get(), *Property->GetName());
+						return;
+					}
+
+					UE_LOG(LogOverridableObject, Warning, TEXT("This owner:%s(0x%p) has already a saved value for property(%s) and it is expected be identical, this is an indication of a missing post edit"),
+						*GetNameSafe(Owner), Owner.Get(), *Property->GetName());
+				}
+
+				ObjectProperty->SetObjectPtrPropertyValue(PreEditValuePtr, SubObject);
 			}
-		};
 
-		if (!NextPropertyIterator)
-		{
+			ON_SCOPE_EXIT
+			{
+				if (Notification == EPropertyNotificationType::PostEdit)
+				{
+					SavedPreEditContainers.Free(Owner, ObjectProperty);
+				}
+			};
+
 			if (Notification == EPropertyNotificationType::PostEdit)
 			{
 				if (SubPropertyNode)
@@ -1525,7 +1600,12 @@ void FOverriddenPropertySet::NotifyPropertyChange(FOverriddenPropertyNode* Paren
 				}
 
 				uint8* PreEditValuePtr = SavedPreEditContainers.Find(Owner, ObjectProperty);
-				// PreEditValuePtr might be null on an unset TOptional<UObject*> as we never did a pre-edit on the inner object property as it was unset.
+				if (!PreEditValuePtr)
+				{
+					// PreEditValuePtr might be null on an unset TOptional<UObject*> as we never did a pre-edit on the inner object property as it was unset.
+					UE_LOG(LogOverridableObject, Verbose, TEXT("Owner:%s(0x%p) The property(%s) didn't have saved memory, might have inconsistent behavior"),
+						*GetNameSafe(Owner), Owner.Get(), *Property->GetName());
+				}
 				UObject* OldSubObject = PreEditValuePtr ? TryGetInstancedSubObjectValue(ObjectProperty, PreEditValuePtr) : nullptr;
 
 				// Handle the case where subobject was nulled out by clearing any previous edit as the reinstancer can still find this object as it might not be GC'd yet
