Explicitly pass NoDestroyConstEval to initialization of field ConstInitClass static member.
Fixes static analysis complaining that variable is initialized with a non-constexpr expression and is more correct anyway.

#Synced-CL Unknown

diff --git a/Engine/Source/Runtime/CoreUObject/Public/UObject/Field.h b/Engine/Source/Runtime/CoreUObject/Public/UObject/Field.h
--- a/Engine/Source/Runtime/CoreUObject/Public/UObject/Field.h
+++ b/Engine/Source/Runtime/CoreUObject/Public/UObject/Field.h
@@ -284,7 +284,7 @@ struct FRegisterFieldClass
 #define UE_DECLARE_GET_FIELD_CLASS_PRIVATE(TRequiredAPI) 				TRequiredAPI virtual FFieldClass* GetFieldClassPrivate() override;
 #define UE_DEFINE_GET_FIELD_CLASS_PRIVATE(TClass) 						FFieldClass* TClass::GetFieldClassPrivate() { return &TClass::ConstInitClass; }
 #define UE_DECLARE_FIELD_CLASS_STATIC_CONSTINIT(TClass, TRequiredAPI) 	TRequiredAPI static constinit TNoDestroy<TFieldClass<TClass>> ConstInitClass;
-#define UE_DEFINE_FIELD_CLASS_STATIC_CONSTINIT(TClass)					constinit TNoDestroy<TFieldClass<TClass>> TClass::ConstInitClass;
+#define UE_DEFINE_FIELD_CLASS_STATIC_CONSTINIT(TClass)					constinit TNoDestroy<TFieldClass<TClass>> TClass::ConstInitClass{NoDestroyConstEval};
 #else
 #define UE_DECLARE_GET_FIELD_CLASS_PRIVATE(TRequiredAPI) 
 #define UE_DEFINE_GET_FIELD_CLASS_PRIVATE(TClass) 
