[UObjectArray]
- Fix non-atomic GetRefCount

#Synced-CL Unknown

diff --git a/Engine/Source/Runtime/CoreUObject/Public/UObject/UObjectArray.h b/Engine/Source/Runtime/CoreUObject/Public/UObject/UObjectArray.h
--- a/Engine/Source/Runtime/CoreUObject/Public/UObject/UObjectArray.h
+++ b/Engine/Source/Runtime/CoreUObject/Public/UObject/UObjectArray.h
@@ -349,7 +349,7 @@ public:
 
 	UE_FORCEINLINE_HINT int32 GetRefCount() const
 	{
-		return (int32)(FlagsAndRefCount & EInternalObjectFlags_RefCountMask);
+		return GetRefCountInternal();
 	}
 
 	void AddRef()
