[ZenLoader]
- Reduce contention on global import table by using a lock at the FLoadedPackageRef level instead of the container
- Move the remaining Packages map to TStripedMap for a nice code cleanup and less danger to cause a deadlock in the future
- Put back FlushDeferredCleanupPackagesQueue where it belongs now that deadlocks are mitigated
- Fix LLM that didn't take the initial reserve into account and various other LLM fixes

#Synced-CL Unknown

diff --git a/Engine/Source/Runtime/CoreUObject/Private/Serialization/AsyncLoading2.cpp b/Engine/Source/Runtime/CoreUObject/Private/Serialization/AsyncLoading2.cpp
--- a/Engine/Source/Runtime/CoreUObject/Private/Serialization/AsyncLoading2.cpp
+++ b/Engine/Source/Runtime/CoreUObject/Private/Serialization/AsyncLoading2.cpp
@@ -127,6 +127,7 @@ CSV_DECLARE_CATEGORY_MODULE_EXTERN(CORE_API, FileIO);
 CSV_DEFINE_STAT(FileIO, FrameCompletedExportBundleLoadsKB);
 
 LLM_DEFINE_TAG(UObject_FAsyncPackage2);
+LLM_DEFINE_TAG(AsyncLoading_GlobalImportStore);
 
 // Scope helper class for region instrumentation
 struct FAsyncLoadingCsvRegionScope
@@ -897,6 +898,52 @@ struct FUnreachableObject
 
 using FUnreachableObjects = TArray<FUnreachableObject>;
 
+// Provides us with a way to profile contention.
+class FAsyncLoadingMutex
+{
+	UE::FSharedMutex Mutex;
+public:
+	void Lock()
+	{
+		if (!Mutex.TryLock())
+		{
+			TRACE_CPUPROFILER_EVENT_SCOPE(FAsyncLoadingMutexWriteLockContention);
+			Mutex.Lock();
+		}
+	}
+	
+	bool IsLocked() const
+	{
+		return Mutex.IsLocked();
+	}
+	
+	void Unlock()
+	{
+		Mutex.Unlock();
+	}
+
+	void LockShared()
+	{
+		if (!Mutex.TryLockShared())
+		{
+			TRACE_CPUPROFILER_EVENT_SCOPE(FAsyncLoadingMutexReadLockContention);
+			Mutex.LockShared();
+		}
+	}
+
+	bool IsLockShared() const
+	{
+		return Mutex.IsLockShared();
+	}
+
+	void UnlockShared()
+	{
+		Mutex.UnlockShared();
+	}
+};
+
+// This needs to stay below 48 bytes as much as possible
+// in the runtime version, the next bin size is currently 64 bytes.
 class FLoadedPackageRef
 {
 private:
@@ -904,6 +951,7 @@ private:
 	friend class FPublicExportMapBasicTest;
 	friend class FPublicExportMapIterationTest;
 	friend class FPublicExportMapPreInsertPublicExportsTest;
+	template <int32 Type> friend class FLockedPackageRef;
 
 	class FPublicExportMap
 	{
@@ -1358,74 +1406,134 @@ private:
 	FPublicExportMap PublicExportMap;
 	FName OriginalPackageName;
 	int32 PackageObjectIndex = -1;
-	int32 RefCount = 0;
+	uint16 PackageRefCount = 0;         // Package references during loading activity, should be reasonably low.
+	std::atomic<uint16> RefCount = 1;   // Only taken by thread activity, the max count should be proportional to the number of worker threads.
 	int32 ExportCount = -1;
+	// This mutex protects all the fields of this class via a FLockedPackageRef at the exception of RefCount, 
+	// which is a typical reference counting using an atomic for synchronization.
+	FAsyncLoadingMutex Mutex;
 	EPackageExtension PackageHeaderExtension = EPackageExtension::Unspecified;
 	EPackageLoader PackageLoader = EPackageLoader::Unknown;
-	bool bAreAllPublicExportsLoaded = false;
-	bool bIsNotInstalled = false;
-	bool bIsMissing = false;
-	bool bHasFailed = false;
-	bool bHasBeenLoadedDebug = false;
-
+	uint8 bAreAllPublicExportsLoaded : 1 = false;
+	uint8 bIsNotInstalled : 1 = false;
+	uint8 bIsMissing : 1 = false;
+	uint8 bHasFailed : 1 = false;
+	uint8 bHasBeenLoadedDebug : 1 = false;
+	
+	void Reset()
+	{
+		check(Mutex.IsLocked());
+		PublicExportMap = FPublicExportMap();
+		OriginalPackageName = FName();
+		PackageObjectIndex = -1;
+		PackageRefCount = 0;
+		ExportCount = -1;
+		PackageHeaderExtension = EPackageExtension::Unspecified;
+		PackageLoader = EPackageLoader::Unknown;
+		bAreAllPublicExportsLoaded = false;
+		bIsNotInstalled = false;
+		bIsMissing = false;
+		bHasFailed = false;
+		bHasBeenLoadedDebug = false;
+	}
 public:
-	FLoadedPackageRef() = default;
+	void AddRef()
+	{
+		check(GetRefCount() < TNumericLimits<decltype(RefCount.load())>::Max());
+		RefCount.fetch_add(1, std::memory_order_relaxed);
+	}
 
-	FLoadedPackageRef(const FLoadedPackageRef& Other) = delete;
+	void Release()
+	{
+		if (RefCount.fetch_sub(1, std::memory_order_relaxed) == 1)
+		{
+			delete this;
+		}
+	}
+	
+	int32 GetRefCount() const
+	{
+		return RefCount.load(std::memory_order_relaxed);
+	}
 
-	FLoadedPackageRef(FLoadedPackageRef&& Other) = default;
+	FLoadedPackageRef() = default;
 
+	FLoadedPackageRef(const FLoadedPackageRef& Other) = delete;
+	FLoadedPackageRef(FLoadedPackageRef&& Other) = delete;
 	FLoadedPackageRef& operator=(const FLoadedPackageRef& Other) = delete;
+	FLoadedPackageRef& operator=(FLoadedPackageRef&& Other) = delete;
+
+	void AddPackageRef()
+	{
+		check(Mutex.IsLocked());
+		check(PackageRefCount < TNumericLimits<decltype(PackageRefCount)>::Max());
+		++PackageRefCount;
+	}
 
-	FLoadedPackageRef& operator=(FLoadedPackageRef&& Other) = default;
+	void ReleasePackageRef()
+	{
+		check(Mutex.IsLocked());
+		check(PackageRefCount > 0);
+		--PackageRefCount;
+	}
 
-	inline int32 GetRefCount() const
+	inline int32 GetPackageRefCount() const
 	{
-		return RefCount;
+		check(Mutex.IsLocked() || Mutex.IsLockShared());
+		return PackageRefCount;
 	}
 
 	inline FName GetOriginalPackageName() const
 	{
+		check(Mutex.IsLocked() || Mutex.IsLockShared());
 		return OriginalPackageName;
 	}
 
 	inline bool HasPackage() const
 	{
+		check(Mutex.IsLocked() || Mutex.IsLockShared());
 		return PackageObjectIndex >= 0;
 	}
 
 	inline bool HasErrors() const
 	{
+		check(Mutex.IsLocked() || Mutex.IsLockShared());
 		return bHasFailed || bIsMissing || bIsNotInstalled;
 	}
 
 	inline bool IsFailed() const
 	{
+		check(Mutex.IsLocked() || Mutex.IsLockShared());
 		return bHasFailed;
 	}
 
 	inline bool IsMissing() const
 	{
+		check(Mutex.IsLocked() || Mutex.IsLockShared());
 		return bIsMissing;
 	}
 
 	inline bool IsNotInstalled() const
 	{
+		check(Mutex.IsLocked() || Mutex.IsLockShared());
 		return bIsNotInstalled;
 	}
 
 	inline EPackageLoader GetPackageLoader() const
 	{
+		check(Mutex.IsLocked() || Mutex.IsLockShared());
 		return PackageLoader;
 	}
 
 	inline EPackageExtension GetPackageHeaderExtension() const
 	{
+		check(Mutex.IsLocked() || Mutex.IsLockShared());
 		return PackageHeaderExtension;
 	}
 
 	inline UPackage* GetPackage() const
 	{
+		check(Mutex.IsLocked() || Mutex.IsLockShared());
 		if (HasPackage())
 		{
 			return static_cast<UPackage*>(GUObjectArray.IndexToObject(PackageObjectIndex)->GetObject());
@@ -1435,6 +1543,7 @@ public:
 
 	inline void SetPackage(UPackage* InPackage)
 	{
+		check(Mutex.IsLocked());
 		check(!bAreAllPublicExportsLoaded);
 		check(!bIsMissing);
 		check(!bIsNotInstalled);
@@ -1454,12 +1563,13 @@ public:
 
 	void RemoveUnreferencedObsoletePackage()
 	{
-		check(RefCount == 0);
-		*this = FLoadedPackageRef();
+		check(PackageRefCount == 0);
+		Reset();
 	}
 
 	void ReplaceReferencedRenamedPackage(UPackage* NewPackage)
 	{
+		check(Mutex.IsLocked());
 		// keep RefCount and PublicExportMap while resetting all other state,
 		// the public exports will be replaced one by one from StoreGlobalObject
 		bAreAllPublicExportsLoaded = false;
@@ -1472,11 +1582,13 @@ public:
 
 	inline bool AreAllPublicExportsLoaded() const
 	{
+		check(Mutex.IsLocked() || Mutex.IsLockShared());
 		return bAreAllPublicExportsLoaded && OriginalPackageName == GetPackage()->GetFName();
 	}
 
 	inline void SetAllPublicExportsLoaded(bool bSnapshotExportCount)
 	{
+		check(Mutex.IsLocked());
 		check(!bIsMissing);
 		check(!bIsNotInstalled);
 		check(!bHasFailed);
@@ -1494,6 +1606,7 @@ public:
 
 	inline void SetIsMissingPackage()
 	{
+		check(Mutex.IsLocked());
 		check(!bAreAllPublicExportsLoaded);
 		check(!HasPackage());
 		bIsMissing = true;
@@ -1502,6 +1615,7 @@ public:
 
 	inline void SetIsNotInstalledPackage()
 	{
+		check(Mutex.IsLocked());
 		check(!bAreAllPublicExportsLoaded);
 		check(!HasPackage());
 		bIsNotInstalled = true;
@@ -1510,6 +1624,7 @@ public:
 
 	inline void ClearErrorFlags()
 	{
+		check(Mutex.IsLocked());
 		bIsNotInstalled = false;
 		bIsMissing = false;
 		bHasFailed = false;
@@ -1517,27 +1632,32 @@ public:
 
 	inline void SetHasFailed()
 	{
+		check(Mutex.IsLocked());
 		bHasFailed = true;
 	}
 
 	FPublicExportMap::FValueIterator GetPublicExportObjectIndices()
 	{
+		check(Mutex.IsLocked() || Mutex.IsLockShared());
 		return PublicExportMap.CreateValueIterator();
 	}
 
 	/** Insert all the given exports into the map by their hash, storing a non-present/null value for any objects that are not already in the map */
 	void PreInsertPublicExports(TConstArrayView<TConstArrayView<FExportMapEntry>> Exports)
 	{
+		check(Mutex.IsLocked());
 		PublicExportMap.PreInsertPublicExports(Exports);
 	}
 
 	void StorePublicExport(uint64 ExportHash, UObject* Object)
 	{
+		check(Mutex.IsLocked());
 		PublicExportMap.Store(ExportHash, GUObjectArray.ObjectToIndex(Object));
 	}
 
 	void RemovePublicExport(uint64 ExportHash, FName ObjectName = NAME_None)
 	{
+		check(Mutex.IsLocked());
 		check(!bIsMissing);
 		check(!bIsNotInstalled);
 		check(HasPackage());
@@ -1550,6 +1670,7 @@ public:
 
 	UObject* GetPublicExport(uint64 ExportHash)
 	{
+		check(Mutex.IsLocked() || Mutex.IsLockShared());
 		int32 ObjectIndex = PublicExportMap.Find(ExportHash);
 		if (ObjectIndex >= 0)
 		{
@@ -1560,6 +1681,7 @@ public:
 
 	void PinPublicExportsForGC(TArray<int32>& OutUnreachableObjectIndices)
 	{
+		check(Mutex.IsLocked());
 		UPackage* Package = GetPackage();
 		UE_ASYNC_UPACKAGE_DEBUG(Package);
 
@@ -1589,6 +1711,7 @@ public:
 
 	void UnpinPublicExportsForGC()
 	{
+		check(Mutex.IsLocked());
 		UPackage* Package = GetPackage();
 		UE_ASYNC_UPACKAGE_DEBUG(Package);
 
@@ -1664,45 +1787,11 @@ public:
 	}
 };
 
-// Provides us with a way to profile contention.
-class FAsyncLoadingMutex
-{
-	UE::FSharedMutex Mutex;
-public:
-	void Lock()
-	{
-		if (!Mutex.TryLock())
-		{
-			TRACE_CPUPROFILER_EVENT_SCOPE(FAsyncLoadingMutexWriteLockContention);
-			Mutex.Lock();
-		}
-	}
-
-	void LockShared()
-	{
-		if (!Mutex.TryLockShared())
-		{
-			TRACE_CPUPROFILER_EVENT_SCOPE(FAsyncLoadingMutexReadLockContention);
-			Mutex.LockShared();
-		}
-	}
-
-	void UnlockShared()
-	{
-		Mutex.UnlockShared();
-	}
-
-	void Unlock()
-	{
-		Mutex.Unlock();
-	}
-};
-
 struct FAsyncLoadingSharedMutexStripedMapLockingPolicy
 {
 	typedef FAsyncLoadingMutex                  MutexType;
-	typedef UE::TUniqueLock<FAsyncLoadingMutex> ExclusiveLockType;
-	typedef UE::TSharedLock<FAsyncLoadingMutex> SharedLockType;
+	typedef UE::TUniqueLock<MutexType>          ExclusiveLockType;
+	typedef UE::TSharedLock<MutexType>          SharedLockType;
 };
 
 template <
@@ -1737,17 +1826,14 @@ template <int32 Type>
 class FLockedPackageRef
 {
 private:
-	FLoadedPackageRef*  PackageRef = nullptr;
-	FAsyncLoadingMutex* SharedMutex = nullptr;
+	TRefCountPtr<FLoadedPackageRef> PackageRef;
 public:
 	FLockedPackageRef(const FLockedPackageRef&) = delete;
 	FLockedPackageRef& operator=(const FLockedPackageRef&) = delete;
 
 	FLockedPackageRef(FLockedPackageRef&& Other)
 	{
-		SharedMutex = Other.SharedMutex;
 		PackageRef  = Other.PackageRef;
-		Other.SharedMutex = nullptr;
 		Other.PackageRef = nullptr;
 	}
 
@@ -1756,21 +1842,21 @@ public:
 		if (this != &Other)
 		{
 			Reset();
-			SharedMutex = Other.SharedMutex;
 			PackageRef = Other.PackageRef;
-			Other.SharedMutex = nullptr;
 			Other.PackageRef = nullptr;
 		}
 
 		return *this;
 	}
 
+	FLoadedPackageRef* Get() const { return PackageRef.GetReference(); }
+
 	FLockedPackageRef() = default;
 
 	FLoadedPackageRef& operator*() { return *PackageRef; }
 	FLoadedPackageRef* operator ->()
 	{
-		return PackageRef;
+		return PackageRef.GetReference();
 	}
 
 	operator bool() const
@@ -1780,27 +1866,37 @@ public:
 
 	void Reset()
 	{
-		PackageRef = nullptr;
-		if (SharedMutex)
+		if (PackageRef)
 		{
 			static_assert(Type >= 0 && Type <= 1);
 			if constexpr (Type == 0)
 			{
-				SharedMutex->UnlockShared();
+				PackageRef->Mutex.UnlockShared();
 			}
 			else
 			{
-				SharedMutex->Unlock();
+				PackageRef->Mutex.Unlock();
 			}
 
-			SharedMutex = nullptr;
+			PackageRef = nullptr;
 		}
 	}
 
-	FLockedPackageRef(FLoadedPackageRef* InPackageRef, FAsyncLoadingMutex& InSharedMutex)
+	FLockedPackageRef(TRefCountPtr<FLoadedPackageRef> InPackageRef)
 		: PackageRef(InPackageRef)
-		, SharedMutex(&InSharedMutex)
 	{
+		if (PackageRef)
+		{
+			static_assert(Type >= 0 && Type <= 1);
+			if constexpr (Type == 0)
+			{
+				PackageRef->Mutex.LockShared();
+			}
+			else
+			{
+				PackageRef->Mutex.Lock();
+			}
+		}
 	}
 
 	~FLockedPackageRef()
@@ -1821,17 +1917,16 @@ private:
 	FAsyncLoadingThread2& AsyncLoadingThread;
 	// Packages in active loading or completely loaded packages, with Desc.UPackageId as key.
 	// Does not track temp packages with custom UPackage names, since they are never imorted by other packages.
-	mutable FAsyncLoadingMutex PackagesLock;
-	TMap<FPackageId, FLoadedPackageRef> Packages;
+	TAsyncLoadingStripedMap<4, FPackageId, TRefCountPtr<FLoadedPackageRef>> Packages;
  	// Contention is currently negligible so only use TStripedMap as a thread-safe TMap with just 1 stripe.
 	// Packages in active loading that export VCells. Stored separately from FLoadedPackageRefs to save space.
-	TAsyncLoadingStripedMap<1, FPackageId, FLoadedPackageCellsRef> PackageCells;
+	TAsyncLoadingStripedMap<4, FPackageId, FLoadedPackageCellsRef> PackageCells;
 	// All native script objects (structs, enums, classes, CDOs and their subobjects) from the initial load phase
 	TAsyncLoadingStripedMap<4, FPackageObjectIndex, UObject*> ScriptObjects;
 	// Built-in VCells always defined by the Verse VM.
-	TAsyncLoadingStripedMap<1, FPackageObjectIndex, Verse::VCell*> ScriptCells;
+	TAsyncLoadingStripedMap<4, FPackageObjectIndex, Verse::VCell*> ScriptCells;
 	// All currently loaded public export objects from any loaded package
-	TAsyncLoadingStripedMap<1, int32, FPublicExportKey> ObjectIndexToPublicExport;
+	TAsyncLoadingStripedMap<4, int32, FPublicExportKey> ObjectIndexToPublicExport;
 #if ALT2_ENABLE_SCRIPT_IMPORTS_DEBUGGING
 public:
 	FCookedScriptObjectsDebug CookedScriptObjectsDebug;
@@ -1844,6 +1939,7 @@ public:
 	FGlobalImportStore(FAsyncLoadingThread2& InAsyncLoadingThread)
 		: AsyncLoadingThread(InAsyncLoadingThread)
 	{
+		LLM_SCOPE_BYTAG(AsyncLoading_GlobalImportStore);
 		Packages.Reserve(32768);
 		ScriptObjects.Reserve(32768);
 		ObjectIndexToPublicExport.Reserve(32768);
@@ -1851,7 +1947,6 @@ public:
 
 	int32 GetStoredPackagesCount() const
 	{
-		UE::TSharedLock ScopeLock(PackagesLock);
 		return Packages.Num();
 	}
 
@@ -1870,89 +1965,130 @@ public:
 		return ObjectIndexToPublicExport.Num();
 	}
 
-	inline FLoadedPackageRef* FindPackageRefNoLock(FPackageId PackageId)
+	inline FReadLockedPackageRef FindPackageRef(FPackageId PackageId) const
 	{
-		return Packages.Find(PackageId);
+		FReadLockedPackageRef PackageRef;
+		Packages.FindAndApply(PackageId, 
+			[&PackageRef](const TRefCountPtr<FLoadedPackageRef>& Value)
+			{
+				// Take the package lock while inside the stripe lock to avoid races with package removal.
+				PackageRef = Value;
+			}
+		);
+		return PackageRef;
+	}
 
+	inline FWriteLockedPackageRef FindPackageRefForWrite(FPackageId PackageId)
+	{
+		FWriteLockedPackageRef PackageRef;
+		Packages.FindAndApply(PackageId,
+			[&PackageRef](const TRefCountPtr<FLoadedPackageRef>& Value)
+			{
+				// Take the package lock while inside the stripe lock to avoid races with package removal.
+				PackageRef = Value;
+			}
+		);
+		return PackageRef;
 	}
-	inline FReadLockedPackageRef FindPackageRef(FPackageId PackageId)
+
+	inline FReadLockedPackageRef FindPackageRefChecked(FPackageId PackageId, FName Name = FName()) const
 	{
-		PackagesLock.LockShared();
-		FLoadedPackageRef* PackageRef = FindPackageRefNoLock(PackageId);
-		if (PackageRef)
+		FReadLockedPackageRef PackageRef = FindPackageRef(PackageId);
+		if (!PackageRef)
 		{
-			return FReadLockedPackageRef(PackageRef, PackagesLock);
+			UE_LOG(LogStreaming, Fatal, TEXT("FindPackageRefChecked: Package %s (0x%s) has been deleted"), *Name.ToString(), *LexToString(PackageId));
 		}
-		PackagesLock.UnlockShared();
-		return FReadLockedPackageRef();
+		return PackageRef;
 	}
 
-	inline FReadLockedPackageRef FindPackageRefChecked(FPackageId PackageId, FName Name = FName())
+	inline FWriteLockedPackageRef FindPackageRefCheckedForWrite(FPackageId PackageId, FName Name = FName())
 	{
-		PackagesLock.LockShared();
-		FLoadedPackageRef* PackageRef = FindPackageRefNoLock(PackageId);
-		if (PackageRef)
+		FWriteLockedPackageRef PackageRef = FindPackageRefForWrite(PackageId);
+		if (!PackageRef)
 		{
-			return FReadLockedPackageRef(PackageRef, PackagesLock);
+			UE_LOG(LogStreaming, Fatal, TEXT("FindPackageRefChecked: Package %s (0x%s) has been deleted"), *Name.ToString(), *LexToString(PackageId));
 		}
-		PackagesLock.UnlockShared();
-		UE_LOG(LogStreaming, Fatal, TEXT("FindPackageRefChecked: Package %s (0x%s) has been deleted"), *Name.ToString(), *LexToString(PackageId));
-		return FReadLockedPackageRef();
+		return PackageRef;
 	}
 
-	inline FWriteLockedPackageRef FindPackageRefCheckedForWrite(FPackageId PackageId, FName Name = FName())
+	FWriteLockedPackageRef FindOrAddPackageRef(FPackageId PackageId)
 	{
-		PackagesLock.Lock();
-		FLoadedPackageRef* PackageRef = FindPackageRefNoLock(PackageId);
-		if (PackageRef)
+		FWriteLockedPackageRef PackageRef;
+		Packages.FindOrProduceAndApply(PackageId,
+			[]() { return TRefCountPtr<FLoadedPackageRef>(new FLoadedPackageRef(), false); },
+			[&PackageRef](const TRefCountPtr<FLoadedPackageRef>& Value)
+			{
+				// Take the package lock while inside the stripe lock to avoid races with package removal.
+				PackageRef = Value;
+			}
+		);
+		return PackageRef;
+	}
+
+	FWriteLockedPackageRef TryRemovePackage(FPackageId PackageId)
+	{
+		FWriteLockedPackageRef RemovedPackageRef;
+		if (Packages.RemoveIf(PackageId,
+			[&RemovedPackageRef](TRefCountPtr<FLoadedPackageRef>& Value)
+			{
+				// Take the lock on the package while we're in the stripe lock to make sure no other thread is using that package during its removal.
+				RemovedPackageRef = Value;
+				return true;
+			}))
 		{
-			return FWriteLockedPackageRef(PackageRef, PackagesLock);
+			check(RemovedPackageRef->GetPackageRefCount() == 0);
+
+			// We should be the last remaining refcount
+			check(RemovedPackageRef->GetRefCount() == 1);
 		}
-		PackagesLock.Unlock();
-		UE_LOG(LogStreaming, Fatal, TEXT("FindPackageRefChecked: Package %s (0x%s) has been deleted"), *Name.ToString(), *LexToString(PackageId));
-		return FWriteLockedPackageRef();
+
+		return RemovedPackageRef;
 	}
 
 	inline FWriteLockedPackageRef AddPackageRef(FPackageId PackageId, FName PackageNameIfKnown, EPackageLoader PackageLoaderIfKnown, EPackageExtension PackageHeaderExtensionIfKnown)
 	{
-		LLM_SCOPE_BYNAME(TEXT("AsyncLoadPackageStore"));
+		LLM_SCOPE_BYTAG(AsyncLoading_GlobalImportStore);
 
-		FlushDeferredCleanupPackagesQueue();
-
-		PackagesLock.Lock();
-		FLoadedPackageRef& PackageRef = Packages.FindOrAdd(PackageId);
+		FWriteLockedPackageRef PackageRef = FindOrAddPackageRef(PackageId);
 
 		// Detect package that have been renamed but are still referenced and attempt to get rid of the references
 		// by flushing the deferred delete packages queue so that we can load the proper package instead.
-		if (PackageRef.RefCount > 0)
+		if (PackageRef->GetPackageRefCount() > 0)
 		{
-			if (UPackage* Package = PackageRef.GetPackage())
+			if (UPackage* Package = PackageRef->GetPackage())
 			{
-				if (PackageRef.GetOriginalPackageName() != Package->GetFName())
+				if (PackageRef->GetOriginalPackageName() != Package->GetFName())
 				{
+					// Avoid deadlock during the cleanup, we'll reacquire after.
+					PackageRef.Reset();
+
+					FlushDeferredCleanupPackagesQueue();
+
+					PackageRef = FindOrAddPackageRef(PackageId);
+
 					// Flush failed to clear the remaining references, warn since we want to know about this.
-					if (PackageRef.RefCount > 0)
+					if (PackageRef->GetPackageRefCount() > 0)
 					{
-						UE_LOG(LogStreaming, Warning, TEXT("Package %s was renamed to %s but is unexpectedly still being referenced by other packages being loaded"), *PackageRef.GetOriginalPackageName().ToString(), *Package->GetFName().ToString());
+						UE_LOG(LogStreaming, Warning, TEXT("Package %s was renamed to %s but is unexpectedly still being referenced by other packages being loaded"), *PackageRef->GetOriginalPackageName().ToString(), *Package->GetFName().ToString());
 					}
 				}
 			}
 		}
 
 		// is this the first reference to a package that already exists?
-		if (PackageRef.RefCount == 0)
+		if (PackageRef->GetPackageRefCount() == 0)
 		{
 			// Remove stale package before searching below as its possible a UEDPIE package got trashed and replaced by a new one
 			// and its important that we find the one that replaced it so we don't try to load it if its a PKG_InMemoryOnly package.
-			if (UPackage* Package = PackageRef.GetPackage())
+			if (UPackage* Package = PackageRef->GetPackage())
 			{
-				if (Package->IsUnreachable() || PackageRef.GetOriginalPackageName() != Package->GetFName())
+				if (Package->IsUnreachable() || PackageRef->GetOriginalPackageName() != Package->GetFName())
 				{
-					RemoveUnreferencedObsoletePackageNoLock(PackageRef);
+					RemoveUnreferencedObsoletePackage(PackageRef);
 				}
 			}
 #if WITH_EDITOR
-			if (!PackageRef.HasPackage() && !PackageNameIfKnown.IsNone())
+			if (!PackageRef->HasPackage() && !PackageNameIfKnown.IsNone())
 			{
 				UPackage* FoundPackage = FindObjectFast<UPackage>(nullptr, PackageNameIfKnown);
 				if (FoundPackage)
@@ -1967,37 +2103,35 @@ public:
 					FPackageId OldPackageId = FoundPackage->GetPackageId();
 					if (OldPackageId.IsValid() && OldPackageId != PackageId)
 					{
-						FLoadedPackageRef* OldPackageRef = Packages.Find(OldPackageId);
-						if (OldPackageRef)
+						if (FWriteLockedPackageRef RemovedPackageRef = TryRemovePackage(OldPackageId))
 						{
 							UE_LOG(LogStreaming, Log,
 								TEXT("FGlobalImportStore:AddPackageRef: Dropping stale reference to package %s (0x%s) that has been renamed to %s (0x%s)"),
-								*OldPackageRef->GetOriginalPackageName().ToString(),
+								*RemovedPackageRef->GetOriginalPackageName().ToString(),
 								*LexToString(OldPackageId),
 								*FoundPackage->GetName(),
 								*LexToString(PackageId)
 							);
 
-							check(OldPackageRef->GetRefCount() == 0);
-							RemoveUnreferencedObsoletePackageNoLock(*OldPackageRef);
-							RemovePackageNoLock(OldPackageId);
+							PackageCells.Remove(OldPackageId);
+							RemoveUnreferencedObsoletePackage(RemovedPackageRef);
 						}
 					}
 
-					if (PackageRef.bIsMissing)
+					if (PackageRef->bIsMissing)
 					{
-						check(!PackageRef.bIsNotInstalled);
-						PackageRef.bIsMissing = false;
+						check(!PackageRef->bIsNotInstalled);
+						PackageRef->bIsMissing = false;
 						UE_LOG(LogStreaming, Warning,
 							TEXT("FGlobalImportStore:AddPackageRef: Found reference to previously missing package %s (0x%s)"),
 							*FoundPackage->GetName(),
 							*LexToString(PackageId)
 						);
 					}
-					else if (PackageRef.bIsNotInstalled)
+					else if (PackageRef->bIsNotInstalled)
 					{
-						check(!PackageRef.bIsMissing);
-						PackageRef.bIsNotInstalled = false;
+						check(!PackageRef->bIsMissing);
+						PackageRef->bIsNotInstalled = false;
 						// Note, unlike missing packages, it is expected that we may later find a reference to a 
 						// previously not installed package so we do not emit a warning
 						UE_LOG(LogStreaming, Log,
@@ -2007,36 +2141,36 @@ public:
 						);
 					}
 
-					PackageRef.SetPackage(FoundPackage);
+					PackageRef->SetPackage(FoundPackage);
 					FoundPackage->SetCanBeImportedFlag(true);
 					FoundPackage->SetPackageId(PackageId);
 				}
 			}
-			if (PackageRef.HasPackage())
+			if (PackageRef->HasPackage())
 			{
-				if (PackageRef.GetPackage()->bHasBeenFullyLoaded)
+				if (PackageRef->GetPackage()->bHasBeenFullyLoaded)
 				{
 					const bool bSnapshotExportCount = false;
-					PackageRef.SetAllPublicExportsLoaded(bSnapshotExportCount);
+					PackageRef->SetAllPublicExportsLoaded(bSnapshotExportCount);
 				}
 			}
 #endif
-			if (UPackage* Package = PackageRef.GetPackage())
+			if (UPackage* Package = PackageRef->GetPackage())
 			{
-				if (Package->IsUnreachable() || PackageRef.GetOriginalPackageName() != Package->GetFName())
+				if (Package->IsUnreachable() || PackageRef->GetOriginalPackageName() != Package->GetFName())
 				{
 					UE_CLOG(!Package->IsUnreachable(), LogStreaming, Log,
 						TEXT("FGlobalImportStore:AddPackageRef: Dropping renamed package %s before reloading %s (0x%s)"),
 						*Package->GetName(),
-						*PackageRef.GetOriginalPackageName().ToString(),
+						*PackageRef->GetOriginalPackageName().ToString(),
 						*LexToString(Package->GetPackageId()));
 
-					RemoveUnreferencedObsoletePackageNoLock(PackageRef);
+					RemoveUnreferencedObsoletePackage(PackageRef);
 				}
 				else
 				{
 					TArray<int32> UnreachableObjectIndices;
-					PackageRef.PinPublicExportsForGC(UnreachableObjectIndices);
+					PackageRef->PinPublicExportsForGC(UnreachableObjectIndices);
 
 					for (int32 ObjectIndex : UnreachableObjectIndices)
 					{
@@ -2053,8 +2187,8 @@ public:
 			}
 
 			// Must be done last since we may reset PackageRef during RemoveUnreferencedObsoletePackage above
-			PackageRef.PackageHeaderExtension = PackageHeaderExtensionIfKnown;
-			PackageRef.PackageLoader = PackageLoaderIfKnown;
+			PackageRef->PackageHeaderExtension = PackageHeaderExtensionIfKnown;
+			PackageRef->PackageLoader = PackageLoaderIfKnown;
 		}
 		else
 		{
@@ -2063,25 +2197,24 @@ public:
 			// Since our package has a RefCount > 0, we should only further specify/promote the extension and load
 			// location, and not allow lowering the extension or location to an unknown value (since it was
 			// either unknown upon creation or it was something of actual meaning that needs to remain the same)
-			if (PackageRef.PackageHeaderExtension == EPackageExtension::Unspecified)
+			if (PackageRef->PackageHeaderExtension == EPackageExtension::Unspecified)
 			{
-				PackageRef.PackageHeaderExtension = PackageHeaderExtensionIfKnown;
+				PackageRef->PackageHeaderExtension = PackageHeaderExtensionIfKnown;
 			}
-			if (PackageRef.PackageLoader == EPackageLoader::Unknown)
+			if (PackageRef->PackageLoader == EPackageLoader::Unknown)
 			{
-				PackageRef.PackageLoader = PackageLoaderIfKnown;
+				PackageRef->PackageLoader = PackageLoaderIfKnown;
 			}
 		}
-		++PackageRef.RefCount;
-		return FWriteLockedPackageRef(&PackageRef, PackagesLock);
+		PackageRef->AddPackageRef();
+		return PackageRef;
 	}
 
 	inline void ReleasePackageRef(FPackageId PackageId, FPackageId FromPackageId = FPackageId())
 	{
-		FReadLockedPackageRef PackageRef = FindPackageRefChecked(PackageId);
+		FWriteLockedPackageRef PackageRef = FindPackageRefCheckedForWrite(PackageId);
 
-		check(PackageRef->RefCount > 0);
-		--PackageRef->RefCount;
+		PackageRef->ReleasePackageRef();
 
 #if DO_CHECK
 		ensureMsgf(!PackageRef->bHasBeenLoadedDebug || PackageRef->bAreAllPublicExportsLoaded || PackageRef->bIsMissing || PackageRef->bIsNotInstalled || PackageRef->bHasFailed,
@@ -2090,7 +2223,7 @@ public:
 			*LexToString(FromPackageId),
 			*PackageRef->GetOriginalPackageName().ToString(),
 			*LexToString(PackageId),
-			PackageRef->RefCount,
+			PackageRef->GetPackageRefCount(),
 			PackageRef->bAreAllPublicExportsLoaded,
 			PackageRef->bIsMissing,
 			PackageRef->bIsNotInstalled,
@@ -2107,7 +2240,7 @@ public:
 		}
 #endif
 		// is this the last reference to a loaded package?
-		if (PackageRef->RefCount == 0 && PackageRef->HasPackage())
+		if (PackageRef->GetPackageRefCount() == 0 && PackageRef->HasPackage())
 		{
 			PackageRef->UnpinPublicExportsForGC();
 
@@ -2120,37 +2253,43 @@ public:
 		}
 	}
 
-	void VerifyLoadedPackages()
+	void VerifyLoadedPackages() const
 	{
 		TRACE_CPUPROFILER_EVENT_SCOPE(VerifyLoadedPackages);
-		UE::TSharedLock ScopeLock(PackagesLock);
-		for (TPair<FPackageId, FLoadedPackageRef>& Pair : Packages)
+		Packages.ParallelForEach(
+			[](const TPair<FPackageId, TRefCountPtr<FLoadedPackageRef>>& Pair)
+			{
+				const FPackageId& PackageId = Pair.Key;
+				FReadLockedPackageRef PackageRef = Pair.Value;
+				ensureMsgf(PackageRef->GetPackageRefCount() == 0,
+					TEXT("PackageId '%s' with ref count %d should not have a ref count now")
+					TEXT(", or this check is incorrectly reached during active loading."),
+					*FormatPackageId(PackageId),
+					PackageRef->GetPackageRefCount());
+			}
+		);
+	}
+
+	void CleanupPublicExportMap(FWriteLockedPackageRef& PackageRef)
+	{
+		for (auto It = PackageRef->PublicExportMap.CreateValueIterator(); It; ++It)
 		{
-			FPackageId& PackageId = Pair.Key;
-			FLoadedPackageRef& Ref = Pair.Value;
-			ensureMsgf(Ref.GetRefCount() == 0,
-				TEXT("PackageId '%s' with ref count %d should not have a ref count now")
-				TEXT(", or this check is incorrectly reached during active loading."),
-				*FormatPackageId(PackageId),
-				Ref.GetRefCount());
+			const int32 ObjectIndex = *It;
+			ObjectIndexToPublicExport.Remove(ObjectIndex);
 		}
 	}
 
-	void RemoveUnreferencedObsoletePackageNoLock(FLoadedPackageRef& PackageRef)
+	void RemoveUnreferencedObsoletePackage(FWriteLockedPackageRef& PackageRef)
 	{
-		UPackage* OldPackage = PackageRef.GetPackage();
+		UPackage* OldPackage = PackageRef->GetPackage();
 		UE_ASYNC_UPACKAGE_DEBUG(OldPackage);
 
 		if (GVerifyUnreachableObjects)
 		{
-			VerifyPackageForRemovalNoLock(PackageRef);
+			VerifyPackageForRemoval(*PackageRef);
 		}
 
-		for (auto It = PackageRef.PublicExportMap.CreateValueIterator(); It; ++It)
-		{
-			const int32 ObjectIndex = *It;
-			ObjectIndexToPublicExport.Remove(ObjectIndex);
-		}
+		CleanupPublicExportMap(PackageRef);
 
 		PackageCells.FindAndApply(OldPackage->GetPackageId(),
 			[](FLoadedPackageCellsRef& PackageCellsRef)
@@ -2159,17 +2298,11 @@ public:
 			}
 		);
 
-		PackageRef.RemoveUnreferencedObsoletePackage();
+		PackageRef->RemoveUnreferencedObsoletePackage();
 		// Reset PackageId to prevent a double remove from GC NotifyUnreachableObjects
 		OldPackage->SetPackageId(FPackageId());
 	}
 
-	void RemoveUnreferencedObsoletePackage(FLoadedPackageRef& PackageRef)
-	{
-		UE::TUniqueLock ScopeLock(PackagesLock);
-		RemoveUnreferencedObsoletePackageNoLock(PackageRef);
-	}
-
 	void ReplaceReferencedRenamedPackageNoLock(FLoadedPackageRef& PackageRef, UPackage* NewPackage)
 	{
 		UPackage* OldPackage = PackageRef.GetPackage();
@@ -2185,44 +2318,33 @@ public:
 	void RemovePackages(const FUnreachableObjects& ObjectsToRemove)
 	{
 		TRACE_CPUPROFILER_EVENT_SCOPE(RemovePackages);
-		UE::TUniqueLock ScopeLock(PackagesLock);
 		for (const FUnreachableObject& Item : ObjectsToRemove)
 		{
 			const FPackageId& PackageId = Item.PackageId;
 			if (PackageId.IsValid())
 			{
-				RemovePackageNoLock(PackageId);
+				RemovePackage(PackageId);
 			}
 		}
 	}
 
-	void RemovePackageNoLock(FPackageId PackageId)
+	void RemovePackage(FPackageId PackageId)
 	{
 		UE_ASYNC_PACKAGEID_DEBUG(PackageId);
-		FLoadedPackageRef PackageRef;
-		bool bRemoved = Packages.RemoveAndCopyValue(PackageId, PackageRef);
 
+		bool bRemoved = false;
+		if (FWriteLockedPackageRef RemovedPackageRef = TryRemovePackage(PackageId))
+		{
+			bRemoved = true;
+			CleanupPublicExportMap(RemovedPackageRef);
+		}
+		
 		{
 			bool bRemovedCells = PackageCells.Remove(PackageId) > 0;
 #if DO_CHECK
 			ensureMsgf(!bRemovedCells || bRemoved, TEXT("Removed %s from cell package map when there should have been nothing to remove"), *LexToString(PackageId));
 #endif
 		}
-
-		if (bRemoved)
-		{
-			for (auto It = PackageRef.PublicExportMap.CreateValueIterator(); It; ++It)
-			{
-				const int32 ObjectIndex = *It;
-				ObjectIndexToPublicExport.Remove(ObjectIndex);
-			}
-		}
-	}
-
-	void RemovePackage(FPackageId PackageId)
-	{
-		UE::TUniqueLock ScopeLock(PackagesLock);
-		RemovePackageNoLock(PackageId);
 	}
 
 	void RemovePublicExports(const FUnreachableObjects& ObjectsToRemove)
@@ -2230,7 +2352,7 @@ public:
 		TRACE_CPUPROFILER_EVENT_SCOPE(RemovePublicExports);
 
 		FPackageId LastPackageId;
-		FReadLockedPackageRef PackageRef;
+		FWriteLockedPackageRef PackageRef;
 		for (const FUnreachableObject& Item : ObjectsToRemove)
 		{
 			int32 ObjectIndex = Item.ObjectIndex;
@@ -2251,7 +2373,7 @@ public:
 					PackageRef.Reset();
 					UE_ASYNC_PACKAGEID_DEBUG(PackageId);
 					LastPackageId = PackageId;
-					PackageRef = FindPackageRef(PackageId);
+					PackageRef = FindPackageRefForWrite(PackageId);
 				}
 				if (PackageRef)
 				{
@@ -2261,68 +2383,81 @@ public:
 		}
 	}
 
-	void VerifyObjectForRemovalNoLock(UObject* GCObject)
+	void VerifyObjectForRemoval(UObject* GCObject, FPackageId* InPackageId = nullptr, FLoadedPackageRef* InPackageRef = nullptr) const
 	{
 		int32 ObjectIndex = GUObjectArray.ObjectToIndex(GCObject);
-		ObjectIndexToPublicExport.FindAndApply(ObjectIndex,
-			[this, GCObject](const FPublicExportKey& PublicExportKey)
+
+		FPublicExportKey PublicExportKey;
+		if (ObjectIndexToPublicExport.FindAndApply(ObjectIndex,
+			[&PublicExportKey](const FPublicExportKey& InPublicExportKey)
 			{
-				FPackageId PackageId = PublicExportKey.GetPackageId();
-				FLoadedPackageRef* PackageRef = FindPackageRefNoLock(PackageId);
-				if (PackageRef)
-				{
-					UObject* ExistingObject = FindPublicExportObjectUncheckedNoLock(PublicExportKey);
-					UE_CLOG(!ExistingObject, LogStreaming, Fatal,
-						TEXT("FGlobalImportStore::VerifyObjectForRemoval: The loaded public export object '%s' with flags (ObjectFlags=%x, InternalObjectFlags=%x) and id %s:0x%llX is missing in GlobalImportStore. ")
-						TEXT("Reason unknown. Double delete? Bug or hash collision?"),
-						*GCObject->GetFullName(),
-						GCObject->GetFlags(),
-						int(GCObject->GetInternalFlags()),
-						*FormatPackageId(PackageId),
-						PublicExportKey.GetExportHash());
-
-					UE_CLOG(ExistingObject != GCObject, LogStreaming, Fatal,
-						TEXT("FGlobalImportStore::VerifyObjectForRemoval: The loaded public export object '%s' with flags (ObjectFlags=%x, InternalObjectFlags=%x) and id %s:0x%llX is not matching the object '%s' in GlobalImportStore. ")
-						TEXT("Reason unknown. Overwritten after it was added? Bug or hash collision?"),
-						*GCObject->GetFullName(),
-						GCObject->GetFlags(),
-						int(GCObject->GetInternalFlags()),
-						*FormatPackageId(PackageId),
-						PublicExportKey.GetExportHash(),
-						*ExistingObject->GetFullName());
-				}
-				else
-				{
-					UE_LOG(LogStreaming, Warning,
-						TEXT("FGlobalImportStore::VerifyObjectForRemoval: The package for the serialized GC object '%s' with flags (ObjectFlags=%x, InternalObjectFlags=%x) and id %s:0x%llX is missing in GlobalImportStore. ")
-						TEXT("Most likely this object has been moved into this package after it was loaded, while the original package is still around."),
-						*GCObject->GetFullName(),
-						GCObject->GetFlags(),
-						int(GCObject->GetInternalFlags()),
-						*FormatPackageId(PackageId),
-						PublicExportKey.GetExportHash());
-				}
+				PublicExportKey = InPublicExportKey;
 			}
-		);
-	}
-
-	void VerifyObjectForRemoval(UObject* GCObject)
-	{
-		UE::TSharedLock ScopeLock(PackagesLock);
-		return VerifyObjectForRemovalNoLock(GCObject);
+		))
+		{
+			// This is a bit messy but exist to avoid locking the same package twice
+			// in the same callstack since the locks aren't recursive and making them
+			// recursive just for this use case would impact both memory and performance and is not worth it.
+			FPackageId PackageId = PublicExportKey.GetPackageId();
+			FReadLockedPackageRef ReadLockedPackageRef;
+			FLoadedPackageRef* PackageRef = nullptr;
+			if (InPackageId && InPackageRef && PackageId == *InPackageId)
+			{
+				PackageRef = InPackageRef;
+			}
+			else
+			{
+				ReadLockedPackageRef = FindPackageRef(PackageId);
+				PackageRef = ReadLockedPackageRef.Get();
+			}
+			
+			if (PackageRef)
+			{
+				UObject* ExistingObject = PackageRef->GetPublicExport(PublicExportKey.GetExportHash());
+				UE_CLOG(!ExistingObject, LogStreaming, Fatal,
+					TEXT("FGlobalImportStore::VerifyObjectForRemoval: The loaded public export object '%s' with flags (ObjectFlags=%x, InternalObjectFlags=%x) and id %s:0x%llX is missing in GlobalImportStore. ")
+					TEXT("Reason unknown. Double delete? Bug or hash collision?"),
+					*GCObject->GetFullName(),
+					GCObject->GetFlags(),
+					int(GCObject->GetInternalFlags()),
+					*FormatPackageId(PackageId),
+					PublicExportKey.GetExportHash());
+
+				UE_CLOG(ExistingObject != GCObject, LogStreaming, Fatal,
+					TEXT("FGlobalImportStore::VerifyObjectForRemoval: The loaded public export object '%s' with flags (ObjectFlags=%x, InternalObjectFlags=%x) and id %s:0x%llX is not matching the object '%s' in GlobalImportStore. ")
+					TEXT("Reason unknown. Overwritten after it was added? Bug or hash collision?"),
+					*GCObject->GetFullName(),
+					GCObject->GetFlags(),
+					int(GCObject->GetInternalFlags()),
+					*FormatPackageId(PackageId),
+					PublicExportKey.GetExportHash(),
+					*ExistingObject->GetFullName());
+			}
+			else
+			{
+				UE_LOG(LogStreaming, Warning,
+					TEXT("FGlobalImportStore::VerifyObjectForRemoval: The package for the serialized GC object '%s' with flags (ObjectFlags=%x, InternalObjectFlags=%x) and id %s:0x%llX is missing in GlobalImportStore. ")
+					TEXT("Most likely this object has been moved into this package after it was loaded, while the original package is still around."),
+					*GCObject->GetFullName(),
+					GCObject->GetFlags(),
+					int(GCObject->GetInternalFlags()),
+					*FormatPackageId(PackageId),
+					PublicExportKey.GetExportHash());
+			}
+		}
 	}
-
-	void VerifyPackageForRemovalNoLock(FLoadedPackageRef& PackageRef)
+	
+	void VerifyPackageForRemoval(FLoadedPackageRef& PackageRef) const
 	{
 		UPackage* Package = PackageRef.GetPackage();
 		FPackageId PackageId = Package->GetPackageId();
 
-		UE_CLOG(PackageRef.GetRefCount() > 0, LogStreaming, Fatal,
+		UE_CLOG(PackageRef.GetPackageRefCount() > 0, LogStreaming, Fatal,
 			TEXT("FGlobalImportStore::VerifyPackageForRemoval: %s (0x%s) - ")
 			TEXT("Package removed while still being referenced, RefCount %d > 0."),
 			*Package->GetName(),
 			*LexToString(PackageId),
-			PackageRef.GetRefCount());
+			PackageRef.GetPackageRefCount());
 
 		for (auto It = PackageRef.GetPublicExportObjectIndices(); It; ++It)
 		{
@@ -2357,19 +2492,13 @@ public:
 						*FormatPackageId(PackageId),
 						*Object->GetFullName());
 
-			VerifyObjectForRemovalNoLock(Object);
+			VerifyObjectForRemoval(Object, &PackageId, &PackageRef);
 		}
 	}
 
-	void VerifyPackageForRemoval(FLoadedPackageRef& PackageRef)
+	inline UObject* FindPublicExportObjectUnchecked(const FPublicExportKey& Key) const
 	{
-		UE::TSharedLock ScopeLock(PackagesLock);
-		VerifyPackageForRemovalNoLock(PackageRef);
-	}
-
-	inline UObject* FindPublicExportObjectUncheckedNoLock(const FPublicExportKey& Key)
-	{
-		FLoadedPackageRef* PackageRef = FindPackageRefNoLock(Key.GetPackageId());
+		FReadLockedPackageRef PackageRef = FindPackageRef(Key.GetPackageId());
 		if (!PackageRef)
 		{
 			return nullptr;
@@ -2377,20 +2506,14 @@ public:
 		return PackageRef->GetPublicExport(Key.GetExportHash());
 	}
 
-	inline UObject* FindPublicExportObjectUnchecked(const FPublicExportKey& Key)
-	{
-		UE::TSharedLock ScopeLock(PackagesLock);
-		return FindPublicExportObjectUncheckedNoLock(Key);
-	}
-
-	inline UObject* FindPublicExportObject(const FPublicExportKey& Key)
+	inline UObject* FindPublicExportObject(const FPublicExportKey& Key) const
 	{
 		UObject* Object = FindPublicExportObjectUnchecked(Key);
 		checkf(!Object || !Object->IsUnreachable(), TEXT("%s"), Object ? *Object->GetFullName() : TEXT("null"));
 		return Object;
 	}
 
-	inline UObject* FindScriptImportObject(FPackageObjectIndex GlobalIndex)
+	inline UObject* FindScriptImportObject(FPackageObjectIndex GlobalIndex) const
 	{
 		check(GlobalIndex.IsScriptImport());
 		UObject* Object = nullptr;
@@ -2398,7 +2521,7 @@ public:
 		return Object;
 	}
 
-	Verse::VCell* FindPublicExportCell(const FPublicExportKey& Key)
+	Verse::VCell* FindPublicExportCell(const FPublicExportKey& Key) const
 	{
 		Verse::VCell* Cell = nullptr;
 		PackageCells.FindAndApply(Key.GetPackageId(),
@@ -2410,7 +2533,7 @@ public:
 		return Cell;
 	}
 
-	Verse::VCell* FindScriptImportCell(FPackageObjectIndex GlobalIndex)
+	Verse::VCell* FindScriptImportCell(FPackageObjectIndex GlobalIndex) const
 	{
 		check(GlobalIndex.IsScriptImport());
 		return ScriptCells.FindRef(GlobalIndex);
@@ -2418,38 +2541,54 @@ public:
 
 	void StoreGlobalObject(FPackageId PackageId, uint64 ExportHash, UObject* Object)
 	{
+		LLM_SCOPE_BYTAG(AsyncLoading_GlobalImportStore);
+
 		check(PackageId.IsValid());
 		check(ExportHash != 0);
 		int32 ObjectIndex = GUObjectArray.ObjectToIndex(Object);
 		FPublicExportKey Key = FPublicExportKey::MakeKey(PackageId, ExportHash);
+		FPublicExportKey ExistingKey;
 
-		UObject* ExistingObject = FindPublicExportObjectUnchecked(Key);
-		if (ExistingObject && ExistingObject != Object)
+		bool bPreviousExportRemoved = false;
 		{
-			int32 ExistingObjectIndex = GUObjectArray.ObjectToIndex(ExistingObject);
+			FWriteLockedPackageRef PackageRef = FindPackageRefCheckedForWrite(Key.GetPackageId());
 
-			UE_LOG(LogStreaming, Verbose,
-				TEXT("FGlobalImportStore::StoreGlobalObject: The constructed public export object '%s' with index %d and id %s:0x%llX collides with object '%s' (ObjectFlags=%X, InternalObjectFlags=%x) with index %d in GlobalImportStore. ")
-				TEXT("The existing object will be replaced since it or its package was most likely renamed after it was loaded the first time."),
-				Object ? *Object->GetFullName() : TEXT("null"),
-				ObjectIndex,
-				*FormatPackageId(Key.GetPackageId()),
-				Key.GetExportHash(),
-				*ExistingObject->GetFullName(),
-				ExistingObject->GetFlags(),
-				int(ExistingObject->GetInternalFlags()),
-				ExistingObjectIndex);
+			UObject* ExistingObject = PackageRef->GetPublicExport(Key.GetExportHash());
+			if (ExistingObject && ExistingObject != Object)
+			{
+				int32 ExistingObjectIndex = GUObjectArray.ObjectToIndex(ExistingObject);
 
-			ExistingObject->AtomicallyClearInternalFlags(EInternalObjectFlags::LoaderImport);
-			ObjectIndexToPublicExport.Remove(ExistingObjectIndex);
-		}
+				UE_LOG(LogStreaming, Verbose,
+					TEXT("FGlobalImportStore::StoreGlobalObject: The constructed public export object '%s' with index %d and id %s:0x%llX collides with object '%s' (ObjectFlags=%X, InternalObjectFlags=%x) with index %d in GlobalImportStore. ")
+					TEXT("The existing object will be replaced since it or its package was most likely renamed after it was loaded the first time."),
+					Object ? *Object->GetFullName() : TEXT("null"),
+					ObjectIndex,
+					*FormatPackageId(Key.GetPackageId()),
+					Key.GetExportHash(),
+					*ExistingObject->GetFullName(),
+					ExistingObject->GetFlags(),
+					int(ExistingObject->GetInternalFlags()),
+					ExistingObjectIndex);
 
-		ObjectIndexToPublicExport.RemoveIf(ObjectIndex,
-			[this, ObjectIndex, Object, &Key](const FPublicExportKey& ExistingKey)
-			{
-				if (ExistingKey != Key)
+				ExistingObject->AtomicallyClearInternalFlags(EInternalObjectFlags::LoaderImport);
+				ObjectIndexToPublicExport.Remove(ExistingObjectIndex);
+			}
+		
+			bPreviousExportRemoved = ObjectIndexToPublicExport.RemoveIf(ObjectIndex,
+				[this, ObjectIndex, Object, &Key, &ExistingKey](const FPublicExportKey& InExistingKey)
 				{
-					UE_LOG(LogStreaming, Verbose,
+					ExistingKey = InExistingKey;
+					return InExistingKey != Key;
+				}
+			) > 0;
+
+			PackageRef->StorePublicExport(ExportHash, Object);
+			ObjectIndexToPublicExport.Add(ObjectIndex, Key);
+		}
+
+		if (bPreviousExportRemoved)
+		{
+			UE_LOG(LogStreaming, Verbose,
 						TEXT("FGlobalImportStore::StoreGlobalObject: The constructed public export object '%s' with index %d and id %s:0x%llX already exists in GlobalImportStore but with a different key %s:0x%llX.")
 						TEXT("The existing object will be replaced since it or its package was most likely renamed after it was loaded the first time."),
 						Object ? *Object->GetFullName() : TEXT("null"),
@@ -2457,29 +2596,18 @@ public:
 						*FormatPackageId(Key.GetPackageId()), Key.GetExportHash(),
 						*FormatPackageId(ExistingKey.GetPackageId()), ExistingKey.GetExportHash());
 
-					// Break the link with the old package now because otherwise, we wouldn't be able to remove the export
-					// during GC since the ObjectIndex can only be linked to a single package ref.
-					if (FReadLockedPackageRef ExistingPackageRef = FindPackageRef(ExistingKey.GetPackageId()))
-					{
-						ExistingPackageRef->RemovePublicExport(ExistingKey.GetExportHash());
-					}
-
-					return true; // Remove
-				}
-
-				return false;
+			// Break the link with the old package now because otherwise, we wouldn't be able to remove the export
+			// during GC since the ObjectIndex can only be linked to a single package ref.
+			if (FWriteLockedPackageRef ExistingPackageRef = FindPackageRefForWrite(ExistingKey.GetPackageId()))
+			{
+				ExistingPackageRef->RemovePublicExport(ExistingKey.GetExportHash());
 			}
-		);
-
-		{
-			FReadLockedPackageRef PackageRef = FindPackageRefChecked(Key.GetPackageId());
-			PackageRef->StorePublicExport(ExportHash, Object);
 		}
-		ObjectIndexToPublicExport.Add(ObjectIndex, Key);
 	}
 
 	void StoreGlobalCell(FPackageId PackageId, uint64 ExportHash, Verse::VCell* Cell)
 	{
+		LLM_SCOPE_BYTAG(AsyncLoading_GlobalImportStore);
 		PackageCells.FindOrProduceAndApplyForWrite(PackageId,
 			[]()
 			{
@@ -2497,6 +2625,7 @@ public:
 
 	void AddScriptObject(FStringView PackageName, FStringView Name, UObject* Object)
 	{
+		LLM_SCOPE_BYTAG(AsyncLoading_GlobalImportStore);
 		TRACE_CPUPROFILER_EVENT_SCOPE(FGlobalImportStore::AddScriptObject);
 
 		TStringBuilder<FName::StringBufferSize> FullName;
@@ -2549,6 +2678,7 @@ public:
 #if WITH_VERSE_VM || defined(__INTELLISENSE__)
 	void AddScriptCellPackage(Verse::VPackage* Package)
 	{
+		LLM_SCOPE_BYTAG(AsyncLoading_GlobalImportStore);
 		for (int32 I = 0; I < Package->NumDefinitions(); ++I)
 		{
 			FUtf8StringView VersePath = Package->GetDefinitionName(I).AsStringView();
@@ -6202,6 +6332,7 @@ void FEventLoadNode2::ParallelLoadingWork()
 	UE::Sanitizer::RaceDetector::FRaceDetectorScope RaceDetectorScope(GDetectRaceDuringLoading, GDetectRaceGracePeriod);
 #endif
 
+	LLM_SCOPE(ELLMTag::AsyncLoading);
 	FTaskTagScope Scope(ETaskTag::EParallelLoadingThread);
 	FAsyncLoadingThreadState2 LocalThreadState{ .bIsWorkerThread = true };
 
@@ -11279,7 +11410,7 @@ void FAsyncLoadingThread2::CollectUnreachableObjects(
 					FReadLockedPackageRef PackageRef = GlobalImportStore.FindPackageRef(PackageId);
 					if (PackageRef)
 					{
-						GlobalImportStore.VerifyPackageForRemovalNoLock(*PackageRef);
+						GlobalImportStore.VerifyPackageForRemoval(*PackageRef);
 					}
 				}
 			}
@@ -11871,7 +12002,7 @@ void FAsyncPackage2::CreateUPackage()
 			UE_ASYNC_PACKAGE_LOG(Warning, Desc, TEXT("CreateUPackage: ReplacePackage"),
 				TEXT("Replacing renamed package %s (0x%s) while being referenced by the loader, RefCount=%d"),
 				*ExistingPackage->GetName(), *LexToString(ExistingPackage->GetPackageId()),
-				PackageRef->GetRefCount());
+				PackageRef->GetPackageRefCount());
 
 			AsyncLoadingThread.GlobalImportStore.ReplaceReferencedRenamedPackageNoLock(*PackageRef, LinkerRoot);
 		}
@@ -13304,4 +13435,4 @@ bool FPublicExportMapPreInsertPublicExportsTest::RunTest(const FString& Paramete
 	return true;
 }
 
-#endif // WITH_LOW_LEVEL_TESTS
+#endif // WITH_LOW_LEVEL_TESTS
\ No newline at end of file
