Add "ConfigRequiresRestart" to the ObjectMacros.h file for better documentation.

#rnx
#Synced-CL 48222653

diff --git a/Engine/Source/Runtime/CoreUObject/Public/UObject/ObjectMacros.h b/Engine/Source/Runtime/CoreUObject/Public/UObject/ObjectMacros.h
--- a/Engine/Source/Runtime/CoreUObject/Public/UObject/ObjectMacros.h
+++ b/Engine/Source/Runtime/CoreUObject/Public/UObject/ObjectMacros.h
@@ -1366,6 +1366,16 @@ namespace UM
 		/// [PropertyMetadata] Property is serialized to config and we should be able to set it anywhere along the config hierarchy.
 		ConfigHierarchyEditable,
 
+		/// [PropertyMetadata] Use by "Config" properties. Indicates that changing the value of this property will not take effect until a restart of the program.
+		/// If in the editor, then the user will be prompted to reboot now or later.
+		///
+		/// This is particularly useful on project settings and for state that is read very early on the engine's boot and cannot be changed once it is running.
+		/// 
+		/// Example:
+		///		UPROPERTY(EditAnywhere, Config, Category = "YourCategory", meta=(ConfigRestartRequired=true))
+		///		bool bEnableSomeFeature = false;
+		ConfigRestartRequired,
+
 		/// [PropertyMetadata] Used by FDirectoryPath properties. Indicates that the path will be picked using the Slate-style directory picker inside the game Content dir.
 		ContentDir,
 
