[Linker]
- Remove deprecations

#Synced-CL Unknown

diff --git a/Engine/Source/Runtime/CoreUObject/Private/UObject/Linker.cpp b/Engine/Source/Runtime/CoreUObject/Private/UObject/Linker.cpp
--- a/Engine/Source/Runtime/CoreUObject/Private/UObject/Linker.cpp
+++ b/Engine/Source/Runtime/CoreUObject/Private/UObject/Linker.cpp
@@ -799,20 +799,6 @@ FLinkerLoad* GetPackageLinker
 	return Result;
 }
 
-FLinkerLoad* LoadPackageLinker(UPackage* InOuter, const FPackagePath& PackagePath, uint32 LoadFlags, UPackageMap* Sandbox, FArchive* InReaderOverride, TFunctionRef<void(FLinkerLoad* LoadedLinker)> LinkerLoadedCallback)
-{
-	FLinkerLoad* Linker = GetPackageLinker(InOuter, PackagePath, LoadFlags, Sandbox, InReaderOverride);
-	LinkerLoadedCallback(Linker);
-	return Linker;
-}
-
-FLinkerLoad* LoadPackageLinker(UPackage* InOuter, const FPackagePath& PackagePath, uint32 LoadFlags, UPackageMap* Sandbox, FArchive* InReaderOverride)
-{
-	PRAGMA_DISABLE_DEPRECATION_WARNINGS
-	return LoadPackageLinker(InOuter, PackagePath, LoadFlags, Sandbox, InReaderOverride, [](FLinkerLoad* InLinker) {});
-	PRAGMA_ENABLE_DEPRECATION_WARNINGS
-}
-
 void ResetLinkerExports(UPackage* InPackage)
 {
 	FLinkerManager::Get().ResetLinkerExports(InPackage);
diff --git a/Engine/Source/Runtime/CoreUObject/Public/UObject/Linker.h b/Engine/Source/Runtime/CoreUObject/Public/UObject/Linker.h
--- a/Engine/Source/Runtime/CoreUObject/Public/UObject/Linker.h
+++ b/Engine/Source/Runtime/CoreUObject/Public/UObject/Linker.h
@@ -646,22 +646,6 @@ typedef uint32 ELazyLoaderFlags;
 	Global functions
 -----------------------------------------------------------------------------*/
 
-/** 
- * Loads a linker for a package and returns it without loading any objects.
- * @param InOuter Package if known, can be null
- * @param PackagePath Package resource to load, must not be empty
- * @param LoadFlags Flags to pass to the new linker
- * @param Sandbox Additional sandbox for loading
- * @param CompatibleGuid Net GUID
- * @param InReaderOverride Optional archive to use for reading package data
- * @param LinkerLoadedCallback Callback when the linker is loaded (or not found)
- * @return Pointer to the loaded linker or null if the file didn't exist
- */
-UE_DEPRECATED(5.4, "Use GetPackageLinker instead")
-COREUOBJECT_API FLinkerLoad* LoadPackageLinker(UPackage* InOuter, const FPackagePath& PackagePath, uint32 LoadFlags, UPackageMap* Sandbox, FArchive* InReaderOverride, TFunctionRef<void(FLinkerLoad* LoadedLinker)> LinkerLoadedCallback);
-UE_DEPRECATED(5.4, "Use GetPackageLinker instead")
-COREUOBJECT_API FLinkerLoad* LoadPackageLinker(UPackage* InOuter, const FPackagePath& PackagePath, uint32 LoadFlags = LOAD_None, UPackageMap* Sandbox = nullptr, FArchive* InReaderOverride = nullptr);
-
 /** 
  * Gets a linker for a package and returns it without loading any objects.
  * @param InOuter Package if known, can be null
