[UObjectHash]
- Fix race on UClass still being loaded by filtering them from GetAllDerivedClasses
- Fix missing filtering for object still being loaded from various functions

#Synced-CL Unknown

diff --git a/Engine/Source/Runtime/CoreUObject/Private/UObject/UObjectHash.cpp b/Engine/Source/Runtime/CoreUObject/Private/UObject/UObjectHash.cpp
--- a/Engine/Source/Runtime/CoreUObject/Private/UObject/UObjectHash.cpp
+++ b/Engine/Source/Runtime/CoreUObject/Private/UObject/UObjectHash.cpp
@@ -55,9 +55,14 @@ LLM_DEFINE_TAG(UObjectHash, TEXT("UObject hashtables"));
 // Global UObject array instance
 FUObjectArray GUObjectArray;
 
-// Common EInternalObjectFlags used in this file to filter out objects.
-static constexpr EInternalObjectFlags DefaultInternalExclusionFlags =
-	EInternalObjectFlags::Unreachable | EInternalObjectFlags::AutoRTFMConstructionAborted;
+static EInternalObjectFlags GetMandatoryInternalExclusionFlags()
+{
+	// Common EInternalObjectFlags used in this file to filter out objects.
+	static constexpr EInternalObjectFlags MandatoryInternalExclusionFlags =
+		EInternalObjectFlags::Unreachable | EInternalObjectFlags::AutoRTFMConstructionAborted;
+
+	return MandatoryInternalExclusionFlags | UE::GetAsyncLoadingInternalFlagsExclusion();
+}
 
 static FORCENOINLINE void OnHashFailure(UObjectBaseUtility* Object, const TCHAR* HashName, const TCHAR* FailureKind)
 {
@@ -874,7 +879,7 @@ static int32 GetObjectOuterHash(FName ObjName, PTRINT Outer)
 
 UObject* StaticFindObjectFastExplicitThreadSafe(FUObjectHashTables& ThreadHash, const UClass* ObjectClass, FName ObjectName, const FString& ObjectPathName, EFindObjectFlags Flags, EObjectFlags ExcludeFlags/*=0*/)
 {
-	const EInternalObjectFlags ExclusiveInternalFlags = DefaultInternalExclusionFlags;
+	const EInternalObjectFlags ExclusiveInternalFlags = GetMandatoryInternalExclusionFlags();
 
 	// Find an object with the specified name and (optional) class, in any package; if bAnyPackage is false, only matches top-level packages
 	int32 Hash = GetObjectHash(ObjectName);
@@ -1005,7 +1010,7 @@ struct FObjectSearchPath
 
 UObject* StaticFindObjectInPackageInternal(FUObjectHashTables& ThreadHash, const UClass* ObjectClass, const UPackage* ObjectPackage, FName ObjectName, EFindObjectFlags Flags, EObjectFlags ExcludeFlags, EInternalObjectFlags ExclusiveInternalFlags)
 {
-	ExclusiveInternalFlags |= DefaultInternalExclusionFlags;
+	ExclusiveInternalFlags |= GetMandatoryInternalExclusionFlags();
 	UObject* Result = nullptr;
 	if (FHashBucket* Inners = ThreadHash.PackageToObjectListMap.Find(ObjectPackage))
 	{
@@ -1040,7 +1045,7 @@ UObject* StaticFindObjectInPackageInternal(FUObjectHashTables& ThreadHash, const
 
 UObject* StaticFindObjectFastInternalThreadSafe(FUObjectHashTables& ThreadHash, const UClass* ObjectClass, const UObject* ObjectPackage, FName ObjectName, EFindObjectFlags Flags, EObjectFlags ExcludeFlags, EInternalObjectFlags ExclusiveInternalFlags)
 {
-	ExclusiveInternalFlags |= DefaultInternalExclusionFlags;
+	ExclusiveInternalFlags |= GetMandatoryInternalExclusionFlags();
 
 	// If they specified an outer use that during the hashing
 	UObject* Result = nullptr;
@@ -1183,7 +1188,7 @@ bool StaticFindAllObjectsFastInternal(TArray<UObject*>& OutFoundObjects, const U
 {
 	INC_DWORD_STAT(STAT_FindObjectFast);
 
-	ExclusiveInternalFlags |= DefaultInternalExclusionFlags;
+	ExclusiveInternalFlags |= GetMandatoryInternalExclusionFlags();
 
 	FObjectSearchPath SearchPath(ObjectName);
 	const int32 Hash = GetObjectHash(SearchPath.Inner);
@@ -1228,7 +1233,7 @@ UObject* StaticFindFirstObjectFastInternal(const UClass* ObjectClass, FName Obje
 {
 	INC_DWORD_STAT(STAT_FindObjectFast);
 
-	ExclusiveInternalFlags |= DefaultInternalExclusionFlags;
+	ExclusiveInternalFlags |= GetMandatoryInternalExclusionFlags();
 
 	UObject* Result = nullptr;
 	FObjectSearchPath SearchPath(ObjectName);
@@ -1563,7 +1568,7 @@ void GetObjectsWithOuter(const class UObjectBase* Outer, TArray<UObject *>& Resu
 #endif
 
 	// We don't want to return any objects that are currently being background loaded unless we're using the object iterator during async loading.
-	ExclusionInternalFlags |= DefaultInternalExclusionFlags | UE::GetAsyncLoadingInternalFlagsExclusion();
+	ExclusionInternalFlags |= GetMandatoryInternalExclusionFlags();
 
 	int32 StartNum = Results.Num();
 	FUObjectHashTables& ThreadHash = FUObjectHashTables::Get();
@@ -1627,7 +1632,7 @@ void ForEachObjectWithOuterBreakable(const class UObjectBase* Outer, TFunctionRe
 #endif
 
 	// We don't want to return any objects that are currently being background loaded unless we're using the object iterator during async loading.
-	ExclusionInternalFlags |= DefaultInternalExclusionFlags | UE::GetAsyncLoadingInternalFlagsExclusion();
+	ExclusionInternalFlags |= GetMandatoryInternalExclusionFlags();
 
 	FUObjectHashTables& ThreadHash = FUObjectHashTables::Get();
 	FHashTableLock HashLock(ThreadHash);
@@ -1672,7 +1677,7 @@ UObjectBase* FindObjectWithOuter(const class UObjectBase* Outer, const class UCl
 	UObject* Result = nullptr;
 	check( Outer );
 	// We don't want to return any objects that are currently being background loaded unless we're using the object iterator during async loading.
-	EInternalObjectFlags ExclusionInternalFlags = DefaultInternalExclusionFlags | UE::GetAsyncLoadingInternalFlagsExclusion();
+	EInternalObjectFlags ExclusionInternalFlags = GetMandatoryInternalExclusionFlags();
 
 	if( NameToLookFor != NAME_None )
 	{
@@ -1722,7 +1727,7 @@ void ForEachObjectWithPackage(const class UPackage* Package, TFunctionRef<bool(U
 	check(Package != nullptr);
 
 	// We don't want to return any objects that are currently being background loaded unless we're using the object iterator during async loading.
-	ExclusionInternalFlags |= DefaultInternalExclusionFlags | UE::GetAsyncLoadingInternalFlagsExclusion();
+	ExclusionInternalFlags |= GetMandatoryInternalExclusionFlags();
 
 	FUObjectHashTables& ThreadHash = FUObjectHashTables::Get();
 	FHashTableLock HashLock(ThreadHash);
@@ -1778,7 +1783,7 @@ void ForEachObjectWithPackage(const class UPackage* Package, TFunctionRef<bool(U
 
 /** Helper function that returns all the children of the specified class recursively */
 template<typename ClassType, typename ArrayAllocator>
-static void RecursivelyPopulateDerivedClasses(FUObjectHashTables& ThreadHash, const UClass* ParentClass, TArray<ClassType, ArrayAllocator>& OutAllDerivedClass)
+static void RecursivelyPopulateDerivedClasses(FUObjectHashTables& ThreadHash, const UClass* ParentClass, TArray<ClassType, ArrayAllocator>& OutAllDerivedClass, const EInternalObjectFlags ExclusionFlags)
 {
 	// Start search with the parent class at virtual index Num-1, then continue searching from index Num as things are added
 	int32 SearchIndex = OutAllDerivedClass.Num() - 1;
@@ -1792,7 +1797,10 @@ static void RecursivelyPopulateDerivedClasses(FUObjectHashTables& ThreadHash, co
 			OutAllDerivedClass.Reserve(OutAllDerivedClass.Num() + ChildSet->Num());
 			for (UClass* ChildClass : *ChildSet)
 			{
-				OutAllDerivedClass.Add(ChildClass);
+				if (ExclusionFlags == EInternalObjectFlags::None || !ChildClass->HasAnyInternalFlags(ExclusionFlags))
+				{
+					OutAllDerivedClass.Add(ChildClass);
+				}
 			}
 		}
 
@@ -1801,7 +1809,7 @@ static void RecursivelyPopulateDerivedClasses(FUObjectHashTables& ThreadHash, co
 
 		if (SearchIndex < OutAllDerivedClass.Num())
 		{
-			SearchClass = OutAllDerivedClass[SearchIndex];			
+			SearchClass = OutAllDerivedClass[SearchIndex];
 		}
 		else
 		{
@@ -1829,7 +1837,7 @@ FORCEINLINE void ForEachObjectOfClasses_Implementation(FUObjectHashTables& Threa
 	TRACE_CPUPROFILER_EVENT_SCOPE(ForEachObjectOfClasses_Implementation);
 
 	// We don't want to return any objects that are currently being background loaded unless we're using the object iterator during async loading.
-	ExclusionInternalFlags |= DefaultInternalExclusionFlags | UE::GetAsyncLoadingInternalFlagsExclusion();
+	ExclusionInternalFlags |= GetMandatoryInternalExclusionFlags();
 
 	TBucketMapLock ClassToObjectListMapLock(ThreadHash.ClassToObjectListMap);
 
@@ -1858,6 +1866,8 @@ void ForEachObjectOfClass(const UClass* ClassToLookFor, TFunctionRef<void(UObjec
 {
 	TRACE_CPUPROFILER_EVENT_SCOPE(ForEachObjectOfClass);
 
+	ExclusionInternalFlags |= GetMandatoryInternalExclusionFlags();
+
 	// Most classes searched for have around 10 subclasses, some have hundreds
 	TArray<const UClass*, TInlineAllocator<16>> ClassesToSearch;
 	ClassesToSearch.Add(ClassToLookFor);
@@ -1867,7 +1877,7 @@ void ForEachObjectOfClass(const UClass* ClassToLookFor, TFunctionRef<void(UObjec
 
 	if (bIncludeDerivedClasses)
 	{
-		RecursivelyPopulateDerivedClasses(ThreadHash, ClassToLookFor, ClassesToSearch);
+		RecursivelyPopulateDerivedClasses(ThreadHash, ClassToLookFor, ClassesToSearch, ExclusionInternalFlags);
 	}
 
 	ForEachObjectOfClasses_Implementation(ThreadHash, ClassesToSearch, Operation, ExclusionFlags, ExclusionInternalFlags);
@@ -1886,9 +1896,10 @@ void GetDerivedClasses(const UClass* ClassToLookFor, TArray<UClass*>& Results, b
 	FUObjectHashTables& ThreadHash = FUObjectHashTables::Get();
 	FHashTableLock HashLock(ThreadHash);
 
+	const EInternalObjectFlags ExclusionFlags = GetMandatoryInternalExclusionFlags();
 	if (bRecursive)
 	{
-		RecursivelyPopulateDerivedClasses(ThreadHash, ClassToLookFor, Results);
+		RecursivelyPopulateDerivedClasses(ThreadHash, ClassToLookFor, Results, ExclusionFlags);
 	}
 	else
 	{
@@ -1898,7 +1909,10 @@ void GetDerivedClasses(const UClass* ClassToLookFor, TArray<UClass*>& Results, b
 			Results.Reserve(Results.Num() + DerivedClasses->Num());
 			for (UClass* DerivedClass : *DerivedClasses)
 			{
-				Results.Add(DerivedClass);
+				if (!DerivedClass->HasAnyInternalFlags(ExclusionFlags))
+				{
+					Results.Add(DerivedClass);
+				}
 			}
 		}
 	}
@@ -1906,9 +1920,44 @@ void GetDerivedClasses(const UClass* ClassToLookFor, TArray<UClass*>& Results, b
 
 TMap<UClass*, TSet<UClass*>> GetAllDerivedClasses()
 {
-	FUObjectHashTables& ThreadHash = FUObjectHashTables::Get();
-	FHashTableLock HashLock(ThreadHash);
-	return ThreadHash.ClassToChildListMap;
+	TRACE_CPUPROFILER_EVENT_SCOPE(GetAllDerivedClasses);
+
+	TMap<UClass*, TSet<UClass*>> Classes;
+	bool bIsAsyncLoading = false;
+	{
+		FUObjectHashTables& ThreadHash = FUObjectHashTables::Get();
+		FHashTableLock HashLock(ThreadHash);
+		Classes = ThreadHash.ClassToChildListMap;
+		bIsAsyncLoading = IsAsyncLoading();
+	}
+
+	// We don't need any filtering if the loader is currently idle
+	if (bIsAsyncLoading)
+	{
+		TRACE_CPUPROFILER_EVENT_SCOPE(GetAllDerivedClasses_AsyncLoadingFilter);
+
+		const EInternalObjectFlags ExclusionFlags = GetMandatoryInternalExclusionFlags();
+		for (auto It = Classes.CreateIterator(); It; ++It)
+		{
+			if (It->Key->HasAnyInternalFlags(ExclusionFlags))
+			{
+				It.RemoveCurrent();
+			}
+			else
+			{
+				for (auto InnerIt = It->Value.CreateIterator(); InnerIt; ++InnerIt)
+				{
+					UClass* Class = *InnerIt;
+					if (Class->HasAnyInternalFlags(ExclusionFlags))
+					{
+						InnerIt.RemoveCurrent();
+					}
+				}
+			}
+		}
+	}
+
+	return Classes;
 }
 
 bool ClassHasInstancesAsyncLoading(const UClass* ClassToLookFor)
@@ -1919,7 +1968,8 @@ bool ClassHasInstancesAsyncLoading(const UClass* ClassToLookFor)
 	FUObjectHashTables& ThreadHash = FUObjectHashTables::Get();
 	FHashTableLock HashLock(ThreadHash);
 
-	RecursivelyPopulateDerivedClasses(ThreadHash, ClassToLookFor, ClassesToSearch);
+	// No exclusion flags since the goal of this function is to test those flags presence
+	RecursivelyPopulateDerivedClasses(ThreadHash, ClassToLookFor, ClassesToSearch, EInternalObjectFlags::None);
 
 	for (const UClass* SearchClass : ClassesToSearch)
 	{
