[OldLoader]

diff --git a/Engine/Source/Runtime/CoreUObject/Private/Serialization/AsyncLoading.h b/Engine/Source/Runtime/CoreUObject/Private/Serialization/AsyncLoading.h
--- a/Engine/Source/Runtime/CoreUObject/Private/Serialization/AsyncLoading.h
+++ b/Engine/Source/Runtime/CoreUObject/Private/Serialization/AsyncLoading.h
@@ -21,968 +21,3 @@
 
 #include "Async/AsyncFileHandle.h"
 
-struct FAsyncPackage;
-struct FFlushRequest;
-class FAsyncLoadingThread;
-
-
-/** [EDL] Async Package Loading State */
-enum class EAsyncPackageLoadingState : uint8
-{
-	NewPackage,
-	WaitingForSummary,
-	StartImportPackages,
-	WaitingForImportPackages,
-	SetupImports,
-	SetupExports,
-	ProcessNewImportsAndExports,
-	WaitingForPostLoad,
-	ReadyForPostLoad,
-	PostLoad_Etc,
-	PackageComplete,
-};
-
-/** [EDL] This version is an ordinary pointer. We can swap in the safer version to verify assumptions */
-struct FUnsafeWeakAsyncPackagePtr
-{
-	FAsyncPackage* Package;
-	FUnsafeWeakAsyncPackagePtr(FAsyncPackage* InPackage = nullptr)
-		: Package(InPackage)
-	{
-	}
-	FORCEINLINE FAsyncPackage& GetPackage() const
-	{
-		check(Package);
-		return *Package;
-	}
-	FORCEINLINE friend uint32 GetTypeHash(const FUnsafeWeakAsyncPackagePtr& WeakPtr)
-	{
-		return GetTypeHash(WeakPtr.Package);
-	}
-	FORCEINLINE bool operator==(const FUnsafeWeakAsyncPackagePtr& Other) const
-	{
-		return Package == Other.Package;
-	}
-	FName HumanReadableStringForDebugging() const;
-};
-
-/**  [EDL] Weak pointer to the async package */
-struct FWeakAsyncPackagePtr
-{
-	FName PackageName;
-	int32 SerialNumber;
-	FWeakAsyncPackagePtr(FAsyncPackage* Package = nullptr);
-	FAsyncPackage& GetPackage() const;
-	FORCEINLINE friend uint32 GetTypeHash(const FWeakAsyncPackagePtr& WeakPtr)
-	{
-		return GetTypeHash(WeakPtr.PackageName);
-	}
-	FORCEINLINE bool operator==(const FWeakAsyncPackagePtr& Other) const
-	{
-		return SerialNumber == Other.SerialNumber;
-	}
-	FName HumanReadableStringForDebugging() const
-	{
-		return PackageName;
-	}
-};
-
-#define VERIFY_WEAK_ASYNC_PACKAGE_PTRS (0) //(DO_CHECK)
-
-#if VERIFY_WEAK_ASYNC_PACKAGE_PTRS
-typedef FWeakAsyncPackagePtr FCheckedWeakAsyncPackagePtr;
-#else
-typedef FUnsafeWeakAsyncPackagePtr FCheckedWeakAsyncPackagePtr;
-#endif
-
-/** [EDL] Event Load Node */
-enum class EEventLoadNode
-{
-	Package_LoadSummary,
-	Package_SetupImports,
-	Package_ExportsSerialized,
-	Package_NumPhases,
-
-	ImportOrExport_Create = 0,
-	ImportOrExport_Serialize,
-	Import_NumPhases,
-
-	Export_StartIO = Import_NumPhases,
-	Export_NumPhases,
-
-	MAX_NumPhases = Package_NumPhases,
-	Invalid_Value = -1
-};
-
-#define USE_IMPLICIT_ARCS (1) // saves time and memory by not actually adding the arcs that are always present, they are:
-
-// Import: EEventLoadNode::ImportOrExport_Create -> EEventLoadNode::ImportOrExport_Serialize: can't consider this import serialized until we hook it up after creation
-// Import: EEventLoadNode::ImportOrExport_Serialize -> EEventLoadNode::Package_ExportsSerialized: can't consider the package done with event driven loading until all imports are serialized
-
-// Export: EEventLoadNode::ImportOrExport_Create -> EEventLoadNode::Export_StartIO: can't do the IO request until it is created
-// Export: EEventLoadNode::Export_StartIO -> EEventLoadNode::ImportOrExport_Serialize: can't serialize until the IO request is ready
-// Import: EEventLoadNode::ImportOrExport_Serialize -> EEventLoadNode::Package_ExportsSerialized: can't consider the package done with event driven loading until all exports are serialized
-
-/** [EDL] Event Load Node Pointer */
-struct FEventLoadNodePtr
-{
-	FCheckedWeakAsyncPackagePtr WaitingPackage;
-	FPackageIndex ImportOrExportIndex;  // IsNull()==true for PACKAGE_ Phases
-	EEventLoadNode Phase;
-
-	FEventLoadNodePtr()
-		: Phase(EEventLoadNode::Invalid_Value)
-	{
-	}
-	FORCEINLINE friend uint32 GetTypeHash(const FEventLoadNodePtr& NodePtr)
-	{
-		return GetTypeHash(NodePtr.WaitingPackage) ^ (GetTypeHash(NodePtr.ImportOrExportIndex) << 13) ^ (uint32(NodePtr.Phase) << 27);
-	}
-
-#if 0
-	FORCEINLINE bool operator<(const FEventLoadNodePtr& Other) const // not a total order
-	{
-		return WaitingPackage.SerialNumber < Other.WaitingPackage.SerialNumber;
-	}
-#endif
-	FORCEINLINE bool operator==(const FEventLoadNodePtr& Other) const
-	{
-		return WaitingPackage == Other.WaitingPackage && 
-			ImportOrExportIndex == Other.ImportOrExportIndex && 
-			Phase == Other.Phase;
-	}
-#if USE_IMPLICIT_ARCS
-	FORCEINLINE int32 NumImplicitArcs()
-	{
-		return ImportOrExportIndex.IsNull() ? 0 : 1; // only import and export nodes have implicit arcs, and every one of them has exactly one
-	}
-	FORCEINLINE FEventLoadNodePtr GetImplicitArc()
-	{
-		check(!ImportOrExportIndex.IsNull()); // package nodes don't have implicit arcs
-		FEventLoadNodePtr Result;
-		Result.WaitingPackage = WaitingPackage;
-		if (Phase == EEventLoadNode::ImportOrExport_Serialize)
-		{
-			Result.Phase = EEventLoadNode::Package_ExportsSerialized;
-			check(Result.ImportOrExportIndex.IsNull());
-		}
-		else
-		{
-			Result.ImportOrExportIndex = ImportOrExportIndex;
-			if (Phase == EEventLoadNode::ImportOrExport_Create)
-			{
-				Result.Phase = ImportOrExportIndex.IsImport() ? EEventLoadNode::ImportOrExport_Serialize : EEventLoadNode::Export_StartIO;
-			}
-			else
-			{
-				check(Phase == EEventLoadNode::Export_StartIO);
-				Result.Phase = EEventLoadNode::ImportOrExport_Serialize;
-			}
-		}
-		return Result;
-	}
-
-#endif
-
-	FString HumanReadableStringForDebugging() const;
-};
-
-/** [EDL] Event Load Node */
-struct FEventLoadNode
-{
-#if USE_IMPLICIT_ARCS
-	typedef TArray<FEventLoadNodePtr> TNodesWaitingForMeArray;
-#else
-	typedef TArray<FEventLoadNodePtr, TInlineAllocator<1> > TNodesWaitingForMeArray;
-#endif
-
-	TNodesWaitingForMeArray NodesWaitingForMe;
-	int32 NumPrerequistes;
-	bool bFired;
-	bool bAddedToGraph;
-	FORCEINLINE FEventLoadNode()
-		: NumPrerequistes(0)
-		, bFired(false)
-		, bAddedToGraph(false)
-	{
-	}
-};
-
-/** [EDL] Event Load Node Array */
-struct FEventLoadNodeArray
-{
-	FEventLoadNode PackageNodes[int(EEventLoadNode::Package_NumPhases)];
-	FEventLoadNode* Array;
-	int32 TotalNumberOfImportExportNodes;
-	int32 TotalNumberOfNodesAdded;
-	int32 NumImports;
-	int32 NumExports;
-	int32 OffsetToImports;
-	int32 OffsetToExports;
-
-	FEventLoadNodeArray()
-		: Array(nullptr)
-		, TotalNumberOfImportExportNodes(0)
-		, TotalNumberOfNodesAdded(0)
-		, NumImports(0)
-		, NumExports(0)
-		, OffsetToImports(0)
-		, OffsetToExports(0)
-	{
-	}
-
-
-	bool AddNode(FEventLoadNodePtr Node)
-	{
-		FEventLoadNode& NodeRef(PtrToNode(Node));
-		check(!NodeRef.bAddedToGraph);
-		NodeRef.bAddedToGraph = true;
-		return (++TotalNumberOfNodesAdded) == 1; // first add, the caller will add this to a list of things with outstanding nodes
-	}
-	bool RemoveNode(FEventLoadNodePtr Node)
-	{
-		FEventLoadNode& NodeRef(PtrToNode(Node));
-		check(NodeRef.bAddedToGraph);
-		NodeRef.bAddedToGraph = false;
-		return (--TotalNumberOfNodesAdded) == 0; // first add, the caller will add this to a list of things with outstanding nodes
-	}
-	FORCEINLINE FEventLoadNode& GetNode(FEventLoadNodePtr Node, bool bCheckAdded = true)
-	{
-		FEventLoadNode& NodeRef(PtrToNode(Node));
-		check(!bCheckAdded || NodeRef.bAddedToGraph);
-		return NodeRef;
-	}
-
-	void Init(int32 InNumImports, int32 InNumExports);
-	void Shutdown();
-	void GetAddedNodes(TArray<FEventLoadNodePtr>& OutAddedNodes, FAsyncPackage* Owner);
-private:
-	FORCEINLINE FEventLoadNode& PtrToNode(FEventLoadNodePtr Node)
-	{
-		int32 Index;
-		if (Node.ImportOrExportIndex.IsNull())
-		{
-			Index = int32(Node.Phase);
-			check(Index < int32(EEventLoadNode::Package_NumPhases) && Index >= 0);
-			return PackageNodes[Index];
-		}
-		check(TotalNumberOfImportExportNodes); // otherwise Init was not called yet
-		if (Node.ImportOrExportIndex.IsImport())
-		{
-			check(int32(Node.Phase) < int32(EEventLoadNode::Import_NumPhases));
-			Index = OffsetToImports + Node.ImportOrExportIndex.ToImport() * int32(EEventLoadNode::Import_NumPhases) + int32(Node.Phase);
-			check(Index >= OffsetToImports && Index < OffsetToExports);
-		}
-		else
-		{
-			check(int32(Node.Phase) < int32(EEventLoadNode::Export_NumPhases));
-			Index = OffsetToExports + Node.ImportOrExportIndex.ToExport() * int32(EEventLoadNode::Export_NumPhases) + int32(Node.Phase);
-			check(Index >= OffsetToExports);
-		}
-		check(Array && Index < TotalNumberOfImportExportNodes && Index >= 0);
-		return Array[Index];
-	}
-
-};
-
-/** [EDL] Event Load Graph */
-struct FEventLoadGraph
-{
-	TSet<FCheckedWeakAsyncPackagePtr> PackagesWithNodes;
-	TArray<int32> IndicesToFire;
-
-	FEventLoadNodeArray& GetArray(FEventLoadNodePtr& Node);
-
-	FEventLoadNode& GetNode(FEventLoadNodePtr& NodeToGet);
-
-	void AddNode(FEventLoadNodePtr& NewNode, bool bHoldForLater = false, int32 NumImplicitPrereqs = 0);
-	void DoneAddingPrerequistesFireIfNone(FEventLoadNodePtr& NewNode, bool bWasHeldForLater = false);
-	void AddArc(FEventLoadNodePtr& PrereqisiteNode, FEventLoadNodePtr& DependentNode);
-	void RemoveNode(FEventLoadNodePtr& NodeToRemove);
-	void NodeWillBeFiredExternally(FEventLoadNodePtr& NodeThatWasFired);
-	void CheckForCycles(bool bDoSlowTests = (!UE_BUILD_SHIPPING && !UE_BUILD_TEST));
-	bool CheckForCyclesInner(const TMultiMap<FEventLoadNodePtr, FEventLoadNodePtr>& Arcs, TSet<FEventLoadNodePtr>& Visited, TSet<FEventLoadNodePtr>& Stack, const FEventLoadNodePtr& Visit);
-};
-
-/** [EDL] Event Load Graph */
-struct FAsyncLoadEventArgs
-{
-	double TickStartTime;
-	const TCHAR* OutLastTypeOfWorkPerformed;
-	UObject* OutLastObjectWorkWasPerformedOn;
-	float TimeLimit;
-	bool bUseTimeLimit;
-	bool bUseFullTimeLimit;
-
-	FAsyncLoadEventArgs()
-		: TickStartTime(0.0)
-		, OutLastTypeOfWorkPerformed(nullptr)
-		, OutLastObjectWorkWasPerformedOn(nullptr)
-		, TimeLimit(0.0f)
-		, bUseTimeLimit(false)
-		, bUseFullTimeLimit(true)
-	{
-	}
-};
-
-/** [EDL] this is a little wrapper that does random pops for further randomization */
-class FImportOrImportIndexArray : public TArray<int32>
-{
-public:
-	void HeapPop(int32& OutItem, EAllowShrinking AllowShrinking = EAllowShrinking::Default);
-};
-
-/**
-* Structure containing intermediate data required for async loading of all imports and exports of a
-* FLinkerLoad.
-*/
-struct FAsyncPackage : public FGCObject
-{
-	friend struct FScopedAsyncPackageEvent;
-	friend struct FAsyncPackageScope;
-	/**
-	 * Constructor
-	 */
-	FAsyncPackage(FAsyncLoadingThread& InThread, const FAsyncPackageDesc& InDesc, struct FEDLBootNotificationManager& InEDLBootNotificationManager);
-	~FAsyncPackage();
-
-	/**
-	 * Ticks the async loading code.
-	 *
-	 * @param	InbUseTimeLimit		Whether to use a time limit
-	 * @param	InbUseFullTimeLimit	If true use the entire time limit, even if you have to block on IO
-	 * @param	InTimeLimit			Soft limit to time this function may take
-	 *
-	 * @return	true if package has finished loading, false otherwise
-	 */
-	EAsyncPackageState::Type TickAsyncPackage(bool bUseTimeLimit, bool bInbUseFullTimeLimit, float& InOutTimeLimit);
-	
-	/** Marks a specific request as complete */
-	void MarkRequestIDsAsComplete();
-
-	/**
-	 * @return Estimated load completion percentage.
-	 */
-	FORCEINLINE float GetLoadPercentage() const
-	{
-		return LoadPercentage;
-	}
-
-	/**
-	 * @return Time load begun. This is NOT the time the load was requested in the case of other pending requests.
-	 */
-	double GetLoadStartTime() const;
-
-	/**
-	 * Emulates ResetLoaders for the package's Linker objects, hence deleting it. 
-	 */
-	void ResetLoader();
-
-	/**
-	* Disassociates linker from this package
-	*/
-	void DetachLinker();
-
-	/**
-	* Flushes linker cache for all objects loaded with this package
-	*/
-	void FlushObjectLinkerCache();
-
-	/**
-	 * Returns the name of the package to load.
-	 */
-	FORCEINLINE const FName& GetPackageName() const
-	{
-		return Desc.Name;
-	}
-
-	/**
-	 * Returns the PackagePath from which to load the package.
-	 */
-	FORCEINLINE const FPackagePath& GetPackagePath() const
-	{
-		return Desc.PackagePath;
-	}
-
-	void AddCompletionCallback(TUniquePtr<FLoadPackageAsyncDelegate>&& Callback, bool bInternal);
-
-	/** Gets the number of references to this package from other packages in the dependency tree. */
-	FORCEINLINE int32 GetDependencyRefCount() const
-	{
-		return DependencyRefCount.GetValue();
-	}
-
-	FORCEINLINE UPackage* GetLinkerRoot() const
-	{
-		return LinkerRoot;
-	}
-
-	/** Returns true if the package has finished loading. */
-	FORCEINLINE bool HasFinishedLoading() const
-	{
-		return bLoadHasFinished;
-	}
-
-	/** Returns package loading priority. */
-	FORCEINLINE int32 GetRequestID() const
-	{
-		return Desc.RequestID;
-	}
-
-	/** Returns package loading priority. */
-	FORCEINLINE TAsyncLoadPriority GetPriority() const
-	{
-		return Desc.Priority;
-	}
-
-	/** Returns package loading priority. */
-	FORCEINLINE void SetPriority(TAsyncLoadPriority InPriority)
-	{
-		Desc.Priority = InPriority;
-	}
-
-	/** Returns true if loading has failed */
-	FORCEINLINE bool HasLoadFailed() const
-	{
-		return bLoadHasFailed;
-	}
-
-	/** Marks the threaded loading phase as complete for this package */
-	void ThreadedLoadingHasFinished()
-	{
-		bThreadedLoadingFinished = true;
-	}
-
-	/** Returns true if the threaded loaded phas has completed for this package */
-	bool HasThreadedLoadingFinished() const
-	{
-		return bThreadedLoadingFinished;
-	}
-
-	/** Adds new request ID to the existing package */
-	void AddRequestID(int32 Id);
-
-	/**
-	* Cancel loading this package.
-	*/
-	void Cancel();
-
-	/**
-	* Set the package that spawned this package as a dependency.
-	*/
-	void SetDependencyRootPackage(FAsyncPackage* InDependencyRootPackage)
-	{
-		DependencyRootPackage = InDependencyRootPackage;
-	}
-
-	/** Returns true if this package is already being loaded in the current callstack */
-	bool IsBeingProcessedRecursively() const
-	{
-		return ReentryCount > 0;
-	}
-
-	/** FGCObject Interface */
-	virtual void AddReferencedObjects(FReferenceCollector& Collector) override;
-	virtual FString GetReferencerName() const override
-	{ 
-		return FString::Printf(TEXT("FAsyncPackage %s"), *GetPackageName().ToString()); 
-	}
-
-	/** Adds a new object referenced by this package */
-	void AddObjectReference(UObject* InObject);
-
-	/** Removes all objects from the list and clears async loading flags */
-	void EmptyReferencedObjects();
-
-	/** Returns the UPackage wrapped by this, if it is valid */
-	UPackage* GetLoadedPackage();
-
-#if WITH_EDITOR
-	/** Gets all assets and packages loaded by this async package, used in the editor */
-	void GetLoadedAssetsAndPackages(TSet<FWeakObjectPtr>& AssetList, TSet<UPackage*>& PackageList);
-#endif
-
-	/** Checks if all dependencies (imported packages) of this package have been fully loaded */
-	bool AreAllDependenciesFullyLoaded(TSet<UPackage*>& VisitedPackages);
-
-	/** Returs true if this package loaded objects that can create GC clusters */
-	bool HasClusterObjects() const
-	{
-		return DeferredClusterObjects.Num() > 0;
-	}
-
-	/** Creates GC clusters from loaded objects */
-	EAsyncPackageState::Type CreateClusters(double InTickStartTime, bool bInUseTimeLimit, double InOutTimeLimit);
-
-	/** Returns true if this package is a dependency to the explicit request id RequestID. */
-	bool IsDependencyOf(int32 RequestID);
-
-private:
-
-	friend class FLinkerLoad;
-
-	/** Checks if all dependencies (imported packages) of this package have been fully loaded */
-	static bool AreAllDependenciesFullyLoadedInternal(FAsyncPackage* Package, TSet<UPackage*>& VisitedPackages, FString& OutError);
-
-	struct FCompletionCallback
-	{
-		bool bIsInternal;
-		bool bCalled;
-		TUniquePtr<FLoadPackageAsyncDelegate> Callback;
-
-		FCompletionCallback()
-			: bIsInternal(false)
-			, bCalled(false)
-		{
-		}
-		FCompletionCallback(bool bInInternal, TUniquePtr<FLoadPackageAsyncDelegate>&& InCallback)
-			: bIsInternal(bInInternal)
-			, bCalled(false)
-			, Callback(MoveTemp(InCallback))
-		{
-		}
-	};
-
-	/** Propagate dependent requests with a package dependency. */
-	void AddDependentRequestsTo(FAsyncPackage* InDependency);
-
-	/** Add the list of passed in Requests to the list of requests that are dependent on this package. */
-	void AddDependentRequests(TArrayView<int32 const> Requests);
-
-	bool IsFullyLoadedObj(UObject* Obj);
-
-	/** Basic information associated with this package */
-	FAsyncPackageDesc Desc;
-	/** Linker which is going to have its exports and imports loaded									*/
-	FLinkerLoad*				Linker;
-	/** Package which is going to have its exports and imports loaded									*/
-	UPackage*				LinkerRoot;
-	/** Call backs called when we finished loading this package											*/
-	TArray<FCompletionCallback>	CompletionCallbacks;
-	/** Pending Import packages - we wait until all of them have been pre loaded. */
-	TMap<FName, FAsyncPackage*> PendingImportedPackages;
-	/** Referenced imports - list of packages we need until we finish loading this package. */
-	TArray<FAsyncPackage*> ReferencedImports;
-	/** Root package if this package was loaded as a dependency of another. NULL otherwise */
-	FAsyncPackage* DependencyRootPackage;
-	/** Number of references to this package from other packages in the dependency tree. */
-	FThreadSafeCounter	DependencyRefCount;
-	/** Current index into linkers import table used to spread creation over several frames				*/
-	int32							LoadImportIndex;
-	/** Current index into linkers import table used to spread creation over several frames				*/
-	int32							ImportIndex;
-	/** Current index into linkers export table used to spread creation over several frames				*/
-	int32							ExportIndex;
-	/** Current index into ObjLoaded array used to spread routing PreLoad over several frames			*/
-	int32							PreLoadIndex;
-	/** Current index into ObjLoaded array used to spread routing PreLoad over several frames			*/
-	int32							PreLoadSortIndex;
-	/** Current index into ObjLoaded array used to spread routing PreLoad over several frames			*/
-	int32							FinishExternalReadDependenciesIndex;
-	/** Current index into ObjLoaded array used to spread routing PostLoad over several frames			*/
-	int32							PostLoadIndex;
-	/** Current index into DeferredPostLoadObjects array used to spread routing PostLoad over several frames			*/
-	int32						DeferredPostLoadIndex;
-	/** Current index into DeferredPostLoadObjects array used to spread routing PostLoad over several frames			*/
-	int32						PostLoadInstanceIndex;
-	/** Current index into DeferredFinalizeObjects array used to spread routing PostLoad over several frames			*/
-	int32						DeferredFinalizeIndex;
-	/** Current index into DeferredClusterObjects array used to spread routing CreateClusters over several frames			*/
-	int32						DeferredClusterIndex;
-	/** Currently used time limit for this tick.														*/
-	float						TimeLimit;
-	/** Whether we are using a time limit for this tick.												*/
-	bool						bUseTimeLimit;
-	/** Whether we should use the entire time limit, even if we're blocked on I/O						*/
-	bool						bUseFullTimeLimit;
-	/** Whether we already exceed the time limit this tick.												*/
-	bool						bTimeLimitExceeded;
-	/** True if our load has failed */
-	bool						bLoadHasFailed;
-	/** True if our load on the async thread has finished */
-	bool						bLoadHasFinished;
-	/** True if threaded loading has finished for this package */
-	bool						bThreadedLoadingFinished; /** @note FH: seems to serve a purpose almost identical to bLoadHasFinished. */
-	/** True if this package was created by this async package */
-	bool						bCreatedLinkerRoot;
-	/** The time taken when we started the tick.														*/
-	double						TickStartTime;
-	/** Last object work was performed on. Used for debugging/ logging purposes.						*/
-	UObject*					LastObjectWorkWasPerformedOn;
-	/** Last type of work performed on object.															*/
-	const TCHAR*				LastTypeOfWorkPerformed;
-	/** Time load begun. This is NOT the time the load was requested in the case of pending requests.	*/
-	double						LoadStartTime;
-	/** Estimated load percentage.																		*/
-	float						LoadPercentage;
-	/** Objects to be post loaded on the game thread */
-	TArray<UObject*> DeferredPostLoadObjects;
-	/** Objects to be finalized on the game thread */
-	TArray<TObjectPtr<UObject>> DeferredFinalizeObjects;
-	/** Objects loaded while loading this package */
-	TArray<TObjectPtr<UObject>> PackageObjLoaded;
-	/** Packages that were loaded synchronously while async loading this package or packages added by verify import */
-	TArray<FLinkerLoad*> DelayedLinkerClosePackages;
-	/** Objects to create GC clusters from */
-	TArray<UObject*> DeferredClusterObjects;
-
-	/** List of all request handles. Those handles refer to explicitly made request. */
-	TArray<int32> RequestIDs;
-
-	FRWLock DependentRequestsLock;
-	/** Requests that currently depends on this package. */
-	TSet<int32> DependentRequests;
-
-#if WITH_METADATA
-	/** Index of the meta-data object within the linkers export table (unset if not yet processed, although may still be INDEX_NONE if there is no meta-data) */
-	TOptional<int32> MetaDataIndex;
-#endif // WITH_METADATA
-	/** Number of times we recursed to load or post load this package. could be split into 2 counters not to use atomics. */
-	std::atomic<int32> ReentryCount;
-	/** List of objects referenced by this package */
-	TSet<TObjectPtr<UObject>> ReferencedObjects;
-	/** Critical section for referenced objects list */
-	FCriticalSection ReferencedObjectsCritical;
-	/** Cached async loading thread object this package was created by */
-	FAsyncLoadingThread& AsyncLoadingThread;
-	FEDLBootNotificationManager& EDLBootNotificationManager;
-	/** Packages that have been imported by this async package */
-	TSet<TObjectPtr<UPackage>> ImportedPackages;
-
-public:
-	
-
-	/** [EDL] Begin Event driven loader specific stuff */
-
-	EAsyncPackageLoadingState AsyncPackageLoadingState;
-	int32 SerialNumber;
-
-	FImportOrImportIndexArray ImportsThatAreNowCreated;
-	FImportOrImportIndexArray ImportsThatAreNowSerialized;
-	FImportOrImportIndexArray ExportsThatCanBeCreated;
-	FImportOrImportIndexArray ExportsThatCanHaveIOStarted;
-	FImportOrImportIndexArray ExportsThatCanBeSerialized;
-	TArray<IAsyncReadRequest*> ReadyPrecacheRequests;
-
-	struct FExportIORequest
-	{
-		int64 Offset;
-		int64 BytesToRead;
-		int32 FirstExportCovered;
-		int32 LastExportCovered;
-		TArray<int32> ExportsToRead;
-
-		FExportIORequest()
-			: Offset(-1)
-			, BytesToRead(-1)
-			, FirstExportCovered(-1)
-			, LastExportCovered(-1)
-		{
-		}
-	};
-
-	TMap<IAsyncReadRequest*, FAsyncPackage::FExportIORequest> PrecacheRequests;
-	TMap<int32, IAsyncReadRequest*> ExportIndexToPrecacheRequest;
-	int64 CurrentBlockOffset;
-	int64 CurrentBlockBytes;
-	TSet<int32> ExportsInThisBlock;
-
-	TMap<TPair<FName, FPackageIndex>, FPackageIndex> ObjectNameWithOuterToExport;
-
-	TSet<FWeakAsyncPackagePtr> PackagesIMayBeWaitingForBeforePostload; // these need to be reexamined and perhaps deleted or collapsed
-
-	TSet<FWeakAsyncPackagePtr> PackagesIAmWaitingForBeforePostload;    // These are linked with PackagesIAmWaitingForBeforePostload, so we can be sure the other package will let us know
-	TSet<FWeakAsyncPackagePtr> OtherPackagesWaitingForMeBeforePostload;
-
-	TArray<FCheckedWeakAsyncPackagePtr> PackagesWaitingToLinkImports;
-
-	int32 ImportAddNodeIndex;
-	int32 ExportAddNodeIndex;
-
-	bool bProcessImportsAndExportsInFlight;
-	bool bProcessPostloadWaitInFlight;
-	bool bAllExportsSerialized;
-
-	void Event_CreateLinker();
-	void Event_FinishLinker();
-	void Event_StartImportPackages();
-	void Event_SetupImports();
-	void Event_SetupExports();
-	void Event_ProcessImportsAndExports();
-	void Event_ExportsDone();
-	void Event_ProcessPostloadWait();
-	void Event_StartPostload();
-
-	void MarkNewObjectForLoadIfItIsAnExport(UObject *Object);
-	bool AnyImportsAndExportWorkOutstanding();
-	void ConditionalQueueProcessImportsAndExports(bool bRequeueForTimeout = false);
-	void ConditionalQueueProcessPostloadWait();
-
-
-	EAsyncPackageState::Type LoadImports_Event();
-	EAsyncPackageState::Type SetupImports_Event();
-	EAsyncPackageState::Type SetupExports_Event();
-	EAsyncPackageState::Type ProcessImportsAndExports_Event();
-
-	FObjectImport* FindExistingImport(int32 LocalImportIndex);
-	void LinkImport(int32 LocalImportIndex);
-	void EventDrivenCreateExport(int32 LocalExportIndex);
-	void StartPrecacheRequest();
-	void EventDrivenSerializeExport(int32 LocalExportIndex);
-	int64 PrecacheRequestReady(IAsyncReadRequest * Req);
-	void MakeNextPrecacheRequestCurrent();
-	void FlushPrecacheBuffer();
-	void EventDrivenLoadingComplete();
-
-	void DumpDependencies(const TCHAR* Label, UObject* Obj);
-	void DumpDependencies(const TCHAR* Label, FLinkerLoad* DumpLinker, FPackageIndex Index);
-
-	UObject* EventDrivenIndexToObject(FPackageIndex Index, bool bCheckSerialized, FPackageIndex DumpIndex = FPackageIndex());
-	template<class T>
-	T* CastEventDrivenIndexToObject(FPackageIndex Index, bool bCheckSerialized, FPackageIndex DumpIndex = FPackageIndex())
-	{
-		UObject* Result = EventDrivenIndexToObject(Index, bCheckSerialized, DumpIndex);
-		if (!Result)
-		{
-			return nullptr;
-		}
-		return CastChecked<T>(Result);
-	}
-
-	FEventLoadNodeArray EventNodeArray;	
-
-	FEventLoadNodePtr AddNode(EEventLoadNode Phase, FPackageIndex ImportOrExportIndex = FPackageIndex(), bool bHoldForLater = false, int32 NumImplicitPrereqs = 0);
-	void DoneAddingPrerequistesFireIfNone(EEventLoadNode Phase, FPackageIndex ImportOrExportIndex = FPackageIndex(), bool bWasHeldForLater = false);
-	void RemoveNode(EEventLoadNode Phase, FPackageIndex ImportOrExportIndex = FPackageIndex());
-	void NodeWillBeFiredExternally(EEventLoadNode Phase, FPackageIndex ImportOrExportIndex = FPackageIndex());
-	void AddArc(FEventLoadNodePtr& PrereqisiteNode, FEventLoadNodePtr& DependentNode);
-	void RemoveAllNodes();
-	void FireNode(FEventLoadNodePtr& NodeToFire);
-
-	FString GetDebuggingPath(FPackageIndex Idx);
-
-	void SetTimeLimit(FAsyncLoadEventArgs& Args, const TCHAR* WorkType)
-	{
-		Args.OutLastTypeOfWorkPerformed = WorkType;
-		Args.OutLastObjectWorkWasPerformedOn = LinkerRoot;
-		TickStartTime = Args.TickStartTime;
-		LastTypeOfWorkPerformed = WorkType;
-		LastObjectWorkWasPerformedOn = LinkerRoot;
-		TimeLimit = Args.TimeLimit;
-		bUseTimeLimit = Args.bUseTimeLimit;
-		bUseFullTimeLimit = Args.bUseFullTimeLimit;
-	}
-
-	FAsyncLoadingThread& GetOwnerThread()
-	{
-		return AsyncLoadingThread;
-	}
-
-	/** [EDL] End Event driven loader specific stuff */
-
-
-#if PERF_TRACK_DETAILED_ASYNC_STATS
-	/** Number of times Tick function has been called.													*/
-	int32							TickCount;
-	/** Number of iterations in loop inside Tick.														*/
-	int32							TickLoopCount;
-
-	/** Number of iterations for CreateLinker.															*/
-	int32							CreateLinkerCount;
-	/** Number of iterations for FinishLinker.															*/
-	int32							FinishLinkerCount;
-	/** Number of iterations for CreateImports.															*/
-	int32							CreateImportsCount;
-	/** Number of iterations for CreateExports.															*/
-	int32							CreateExportsCount;
-	/** Number of iterations for PreLoadObjects.														*/
-	int32							PreLoadObjectsCount;
-	/** Number of iterations for PostLoadObjects.														*/
-	int32							PostLoadObjectsCount;
-	/** Number of iterations for FinishObjects.															*/
-	int32							FinishObjectsCount;
-
-	/** Total time spent in Tick.																		*/
-	double						TickTime;
-	/** Total time spent in	CreateLinker.																*/
-	double						CreateLinkerTime;
-	/** Total time spent in FinishLinker.																*/
-	double						FinishLinkerTime;
-	/** Total time spent in CreateImports.																*/
-	double						CreateImportsTime;
-	/** Total time spent in CreateExports.																*/
-	double						CreateExportsTime;
-	/** Total time spent in PreLoadObjects.																*/
-	double						PreLoadObjectsTime;
-	/** Total time spent in PostLoadObjects.															*/
-	double						PostLoadObjectsTime;
-	/** Total time spent in FinishObjects.																*/
-	double						FinishObjectsTime;
-
-#endif
-
-	void CallCompletionCallbacks(bool bInternalOnly, EAsyncLoadingResult::Type LoadingResult);
-
-	/**
-	* Route PostLoad to deferred objects.
-	*
-	* @return true if we finished calling PostLoad on all loaded objects and no new ones were created, false otherwise
-	*/
-	EAsyncPackageState::Type PostLoadDeferredObjects(double InTickStartTime, bool bInUseTimeLimit, double InOutTimeLimit);
-
-	/** Close any linkers that have been open as a result of synchronous load during async loading */
-	void CloseDelayedLinkers();
-
-private:
-	/**
-	 * Gives up time slice if time limit is enabled.
-	 *
-	 * @return true if time slice can be given up, false otherwise
-	 */
-	bool GiveUpTimeSlice();
-	/**
-	 * Returns whether time limit has been exceeded.
-	 *
-	 * @return true if time limit has been exceeded (and is used), false otherwise
-	 */
-	bool IsTimeLimitExceeded();
-
-	/**
-	 * Begin async loading process. Simulates parts of BeginLoad.
-	 *
-	 * Objects created during BeginAsyncLoad and EndAsyncLoad will have EInternalObjectFlags::AsyncLoading set
-	 */
-	void BeginAsyncLoad();
-	/**
-	 * End async loading process. Simulates parts of EndLoad(). FinishObjects 
-	 * simulates some further parts once we're fully done loading the package.
-	 */
-	void EndAsyncLoad();
-	/**
-	 * Create linker async. Linker is not finalized at this point.
-	 *
-	 * @return true
-	 */
-	EAsyncPackageState::Type CreateLinker();
-	/**
-	 * Finalizes linker creation till time limit is exceeded.
-	 *
-	 * @return true if linker is finished being created, false otherwise
-	 */
-	EAsyncPackageState::Type FinishLinker();
-	/**
-	 * Loads imported packages..
-	 *
-	 * @return true if we finished loading all imports, false otherwise
-	 */
-	EAsyncPackageState::Type LoadImports();
-	/**
-	 * Create imports till time limit is exceeded.
-	 *
-	 * @return true if we finished creating all imports, false otherwise
-	 */
-	EAsyncPackageState::Type CreateImports();
-
-#if WITH_METADATA
-	/**
-	 * Creates and loads meta-data for the package.
-	 *
-	 * @return true if we finished creating meta-data, false otherwise.
-	 */
-	EAsyncPackageState::Type CreateMetaData();
-#endif // WITH_METADATA
-	/**
-	 * Create exports till time limit is exceeded.
-	 *
-	 * @return true if we finished creating and preloading all exports, false otherwise.
-	 */
-	EAsyncPackageState::Type CreateExports();
-	/**
-	 * Preloads aka serializes all loaded objects.
-	 *
-	 * @return true if we finished serializing all loaded objects, false otherwise.
-	 */
-	EAsyncPackageState::Type PreLoadObjects();
-	/**
-	 * Route PostLoad to all loaded objects. This might load further objects!
-	 *
-	 * @return true if we finished calling PostLoad on all loaded objects and no new ones were created, false otherwise
-	 */
-	EAsyncPackageState::Type PostLoadObjects();
-	/**
-	 * Finish up objects and state, which means clearing the EInternalObjectFlags::AsyncLoading flag on newly created ones
-	 *
-	 * @return true
-	 */
-	EAsyncPackageState::Type FinishObjects();
-
-	/**
-	 * Finalizes external dependencies till time limit is exceeded
-	 *
-	 * @return Complete if all dependencies are finished, TimeOut otherwise
-	 */
-	EAsyncPackageState::Type FinishExternalReadDependencies();
-
-	/**
-	 * Function called when pending import package has been preloaded. called from the async loading thread
-	 */
-	void ImportPreLoadedAsyncCallback(const FName& PackageName, UPackage* LoadedPackage, EAsyncLoadingResult::Type Result);
-
-	/**
-	 * Adds dependency tree to the list if packages to wait for until their linkers have been created.
-	 *
-	 * @param ImportedPackage Package imported either directly or by one of the imported packages
-	 */
-	void AddDependencyTree(FAsyncPackage& ImportedPackage, TSet<FAsyncPackage*>& SearchedPackages);
-
-	/**
-	 * Adds a unique package to the list of packages to wait for until their linkers have been created.
-	 *
-	 * @param PendingImport Package imported either directly or by one of the imported packages
-	 */
-	bool AddUniqueLinkerDependencyPackage(FAsyncPackage& PendingImport);
-
-	/**
-	 * Adds a package to the list of pending import packages.
-	 *
-	 * @param PendingImport Name of the package imported either directly or by one of the imported packages
-	 * @param PackageToLoad Name of the package to load into PendingImport
-	 */
-	void AddImportDependency(const FName& PendingImport, const FName& PackageToLoad, FLinkerInstancingContext InstancingContext);
-
-	/**
-	 * Removes references to any imported packages.
-	 */
-	void FreeReferencedImports();
-
-	/**
-	* Updates load percentage stat
-	*/
-	void UpdateLoadPercentage();
-
-#if PERF_TRACK_DETAILED_ASYNC_STATS
-	/** Add this time taken for object of class Class to have CreateExport called, to the stats we track. */
-	void TrackCreateExportTimeForClass(const UClass* Class, double Time);
-
-	/** Add this time taken for object of class Class to have PostLoad called, to the stats we track. */
-	void TrackPostLoadTimeForClass(const UClass* Class, double Time);
-#endif // PERF_TRACK_DETAILED_ASYNC_STATS
-
-public:
-
-	/** Serialization context for this package */
-	FUObjectSerializeContext* GetSerializeContext();
-};
-
-struct FScopedAsyncPackageEvent
-{
-	/** Current scope package */
-	FAsyncPackage* Package;
-	/** Outer scope package */
-	FAsyncPackage* PreviousPackage;
-	IAsyncPackageLoader* PreviousAsyncPackageLoader;
-	
-	FScopedAsyncPackageEvent(FAsyncPackage* InPackage);
-	~FScopedAsyncPackageEvent();
-};
-
diff --git a/Engine/Source/Runtime/CoreUObject/Private/Serialization/AsyncLoadingPrivate.h b/Engine/Source/Runtime/CoreUObject/Private/Serialization/AsyncLoadingPrivate.h
--- a/Engine/Source/Runtime/CoreUObject/Private/Serialization/AsyncLoadingPrivate.h
+++ b/Engine/Source/Runtime/CoreUObject/Private/Serialization/AsyncLoadingPrivate.h
@@ -150,24 +150,3 @@ private:
 
 	FLinkerLoad* OwnerLinker;
 };
-
-/**
- * The Linker export archive converts Export relative offsets
- * to file relative offsets when exports are cooked to separate
- * archives.
- */
-class FLinkerExportArchive final
-	: public FArchiveProxy
-{
-public:
-	FLinkerExportArchive(class FLinkerLoad& InLinker, int64 InExportSerialOffset, int64 InExportSerialSize);
-
-	COREUOBJECT_API virtual int64 Tell() override;
-	COREUOBJECT_API virtual void Seek(int64 Position) override;
-
-private:
-	int64 ExportSerialOffset;
-	int64 ExportSerialSize;
-	bool bExportsCookedToSeparateArchive;
-};
-
diff --git a/Engine/Source/Runtime/CoreUObject/Private/Serialization/AsyncLoadingThread.h b/Engine/Source/Runtime/CoreUObject/Private/Serialization/AsyncLoadingThread.h
--- a/Engine/Source/Runtime/CoreUObject/Private/Serialization/AsyncLoadingThread.h
+++ b/Engine/Source/Runtime/CoreUObject/Private/Serialization/AsyncLoadingThread.h
@@ -23,202 +23,6 @@
 
 #include <atomic>
 
-struct FPrecacheCallbackHandler;
-
-/** [EDL] Event Driven loader event */
-struct FAsyncLoadEvent
-{
-	enum
-	{
-		EventSystemPriority_MAX = MAX_int32
-	};
-
-	int32 UserPriority;
-	int32 PackageSerialNumber;
-	int32 EventSystemPriority;
-	int32 SerialNumber;
-	TFunction<void(FAsyncLoadEventArgs& Args)> Payload;
-
-	FAsyncLoadEvent()
-		: UserPriority(0)
-		, PackageSerialNumber(0)
-		, EventSystemPriority(0)
-		, SerialNumber(0)
-	{
-	}
-	FAsyncLoadEvent(int32 InUserPriority, int32 InPackageSerialNumber, int32 InEventSystemPriority, int32 InSerialNumber, TFunction<void(FAsyncLoadEventArgs& Args)>&& InPayload)
-		: UserPriority(InUserPriority)
-		, PackageSerialNumber(InPackageSerialNumber)
-		, EventSystemPriority(InEventSystemPriority)
-		, SerialNumber(InSerialNumber)
-		, Payload(Forward<TFunction<void(FAsyncLoadEventArgs& Args)>>(InPayload))
-	{
-	}
-
-	FORCEINLINE bool operator<(const FAsyncLoadEvent& Other) const
-	{
-		if (UserPriority != Other.UserPriority)
-		{
-			return UserPriority > Other.UserPriority;
-		}
-		if (EventSystemPriority != Other.EventSystemPriority)
-		{
-			return EventSystemPriority > Other.EventSystemPriority;
-		}
-		if (PackageSerialNumber != Other.PackageSerialNumber)
-		{
-			return PackageSerialNumber > Other.PackageSerialNumber; // roughly DFS
-		}
-		return SerialNumber < Other.SerialNumber;
-	}
-};
-
-/** [EDL] Event queue for the event driven loader */
-struct FAsyncLoadEventQueue
-{
-	//FCriticalSection AsyncLoadEventQueueCritical; //@todoio maybe this doesn't need to be protected by a critical section. 
-	int32 RunningSerialNumber;
-	TArray<FAsyncLoadEvent> EventQueue;
-
-	FAsyncLoadEventQueue()
-		: RunningSerialNumber(0)
-	{
-	}
-
-	FORCEINLINE void AddAsyncEvent(int32 UserPriority, int32 PackageSerialNumber, int32 EventSystemPriority, TFunction<void(FAsyncLoadEventArgs& Args)>&& Payload)
-	{
-		//UE::TScopeLock Lock(AsyncLoadEventQueueCritical);
-		//@todoio check(FAsyncLoadingThread::IsInAsyncLoadThread());
-		EventQueue.HeapPush(FAsyncLoadEvent(UserPriority, PackageSerialNumber, EventSystemPriority, ++RunningSerialNumber, Forward<TFunction<void(FAsyncLoadEventArgs& Args)>>(Payload)));
-	}
-	bool PopAndExecute(FAsyncLoadEventArgs& Args)
-	{
-		FAsyncLoadEvent Event;
-		bool bResult = false;
-		{
-			//UE::TScopeLock Lock(AsyncLoadEventQueueCritical);
-			//@todoio check(FAsyncLoadingThread::IsInAsyncLoadThread());
-			if (EventQueue.Num())
-			{
-				EventQueue.HeapPop(Event, EAllowShrinking::No);
-				bResult = true;
-			}
-		}
-		if (bResult)
-		{
-			Event.Payload(Args);
-		}
-		return bResult;
-	}
-};
-
-// Shared data for FFlushRequest
-struct FFlushRequestData
-{
-	FFlushRequestData(TConstArrayView<int32> InRequestIDs)
-		: RequestIDs(MakeUnique<std::atomic<int32>[]>(InRequestIDs.Num()))
-		, NumRequests(InRequestIDs.Num())
-	{
-		for (int32 i=0; i < InRequestIDs.Num(); ++i)
-		{
-			RequestIDs.Get()[i] = InRequestIDs[i];
-		}
-	}
-	
-	TConstArrayView<std::atomic<int32>> GetRequestIDs() const
-	{
-		return MakeArrayView(RequestIDs.Get(), NumRequests);
-	}
-
-	TArrayView<std::atomic<int32>> GetRequestIDs() 
-	{
-		return MakeArrayView(RequestIDs.Get(), NumRequests);
-	}
-	
-	void OnRequestComplete(int32 InCompletedRequest)
-	{
-		for (std::atomic<int32>& RequestID : GetRequestIDs())
-		{
-			if (RequestID.load(std::memory_order_relaxed) == InCompletedRequest)
-			{
-				++NumCompletedRequests;
-			}
-		}
-	}
-	
-	bool IsComplete()
-	{
-		return NumRequests == NumCompletedRequests;
-	}
-
-private:
-	/** 
-	 * Shared request id state between main thread and async loading thread. 
-	 * @note because request id are assigned to package request after being generated
-	 * multiple request ids can identify a specific package request however only the original request id
-	 * will properly identify unloaded dependencies. The shared nature of the request allows the async loading thread to
-	 * fix up the main thread flush request id to identify the original request id instead of duplicate when trying to identify
-	 * which main thread packages should be processed when flushing.
-	 */
-	TUniquePtr<std::atomic<int32>[]> RequestIDs;	
-	int32 NumRequests;
-	int32 NumCompletedRequests = 0;
-};
-
-/** 
- * Request to flush a specific package using the request id 
- */
-struct FFlushRequest
-{
-	FFlushRequest() = default;
-
-	FFlushRequest(TConstArrayView<int32> InRequestIDs)
-		: Data(MakeShared<FFlushRequestData>(InRequestIDs))
-	{
-	}
-	
-	TConstArrayView<std::atomic<int32>> GetRequestIDs() const 
-	{
-		return Data->GetRequestIDs();
-	}
-
-	bool IsValid() const
-	{
-		return Data.IsValid();
-	}
-	
-	bool IsComplete() const
-	{
-		return Data->IsComplete();
-	}
-
-	operator bool() const
-	{
-		return IsValid();
-	}
-
-private:
-	void OnRequestComplete(int32 RequestID)
-	{
-		Data->OnRequestComplete(RequestID);	
-	}
-
-	void AdjustRequestIDs(TMap<int32, int32>& DuplicateRequestMap)
-	{
-		for (std::atomic<int32>& RequestID : Data->GetRequestIDs())
-		{	
-			if (int32* MainRequest = DuplicateRequestMap.Find(RequestID.load(std::memory_order_relaxed)))
-			{
-				RequestID.store(*MainRequest, std::memory_order_relaxed);
-			}
-		}
-	}
-
-	friend class FAsyncLoadingThread;
-
-	TSharedPtr<FFlushRequestData, ESPMode::ThreadSafe> Data;
-};
-
 /** Holds the maximum package summary size that can be set via ini files
   * This is used for the initial precache and should be large enough to hold the actual Sum.TotalHeaderSize 
   */
@@ -227,620 +31,3 @@ struct FMaxPackageSummarySize
 	static int32 Value;
 	static void Init();
 };
-
-/**
- * Async loading thread. Preloads/serializes packages on async loading thread. Postloads objects on the game thread.
- */
-class FAsyncLoadingThread final : public FRunnable, public IAsyncPackageLoader
-{
-	friend struct FAsyncPackage;
-
-	struct FEDLBootNotificationManager& EDLBootNotificationManager;
-
-	/** Thread to run the worker FRunnable on */
-	FRunnableThread* Thread;
-	/** Stops this thread */
-	FThreadSafeCounter StopTaskCounter;
-
-	/** [ASYNC/GAME THREAD] true if the async thread is actually started. We don't start it until after we boot because the boot process on the game thread can create objects that are also being created by the loader */
-	static bool bThreadStarted;
-
-	/** [ASYNC/GAME THREAD] Event used to signal there's queued packages to stream */
-	FEvent* QueuedRequestsEvent;
-	/** [ASYNC/GAME THREAD] Event used to signal loading should be cancelled */
-	FEvent* CancelLoadingEvent;
-	/** [ASYNC/GAME THREAD] Event used to signal that the async loading thread should be suspended */
-	FEvent* ThreadSuspendedEvent;
-	/** [ASYNC/GAME THREAD] Event used to signal that the async loading thread has resumed */
-	FEvent* ThreadResumedEvent;
-	/** [ASYNC/GAME THREAD] List of queued packages to stream */
-	TArray<FAsyncPackageDesc*> QueuedPackages;
-#if THREADSAFE_UOBJECTS
-	/** [ASYNC/GAME THREAD] Package queue critical section */
-	FCriticalSection QueueCritical;
-#endif
-	/** [ASYNC/GAME THREAD] True if the async loading thread received a request to cancel async loading **/
-	FThreadSafeBool bShouldCancelLoading;
-	/** [ASYNC/GAME THREAD] True if the async loading thread received a request to suspend **/
-	FThreadSafeCounter IsLoadingSuspended;
-	/** [ASYNC/GAME THREAD] Event used to signal there's queued packages to stream */
-	TArray<FAsyncPackage*> LoadedPackages;	
-	TMap<FName, FAsyncPackage*> LoadedPackagesNameLookup;
-#if THREADSAFE_UOBJECTS
-	/** [ASYNC/GAME THREAD] Critical section for LoadedPackages list */
-	FCriticalSection LoadedPackagesCritical;
-#endif
-	/** [GAME THREAD] Event used to signal there's queued packages to stream */
-	TArray<FAsyncPackage*> LoadedPackagesToProcess;
-	TArray<FAsyncPackage*> PackagesToDelete;
-	TMap<FName, FAsyncPackage*> LoadedPackagesToProcessNameLookup;
-#if WITH_EDITOR
-	TSet<FWeakObjectPtr> LoadedAssets;
-#endif
-#if THREADSAFE_UOBJECTS
-	/** [ASYNC/GAME THREAD] Critical section for LoadedPackagesToProcess list. 
-	 * Note this is only required for looking up existing packages on the async loading thread 
-	 */
-	FCriticalSection LoadedPackagesToProcessCritical;
-#endif
-
-	/** [ASYNC/GAME THREAD] FIFO of current flush requests from the game thread, managed by the async thread.
-	 * producer can only be the game thread.
-	 */
-	TArray<FFlushRequest, TInlineAllocator<4>> FlushRequests;
-#if THREADSAFE_UOBJECTS
-	/** [ASYNC/GAME THREAD] Critical section for the FlushRequests stack. needed if running with the async loading thread. */
-	FCriticalSection FlushRequestCritical;
-#endif
-	/** [ASYNC THREAD] Map containing a mapping of duplicate request to their main request. Use to operate flush request accurately. */
-	TMap<int32, int32> DuplicateRequestMap;
-
-	/** [ASYNC THREAD] Array of packages that are being preloaded */
-	TArray<FAsyncPackage*> AsyncPackages;
-	TMap<FName, FAsyncPackage*> AsyncPackageNameLookup;
-public:
-	/** [EDL] Async Packages that are ready for tick */
-	TArray<FAsyncPackage*> AsyncPackagesReadyForTick;
-private:
-
-#if THREADSAFE_UOBJECTS
-	/** We only lock AsyncPackages array to make GetAsyncLoadPercentage thread safe, so we only care about locking Add/Remove operations on the async thread */
-	FCriticalSection AsyncPackagesCritical;
-#endif
-
-	/** List of all pending package requests */
-	TSet<int32> PendingRequests;
-#if THREADSAFE_UOBJECTS
-	/** Synchronization object for PendingRequests list */
-	FCriticalSection PendingRequestsCritical;
-#endif
-
-	/** [ASYNC/GAME THREAD] Number of package load requests in the async loading queue */
-	FThreadSafeCounter QueuedPackagesCounter;
-	/** [ASYNC/GAME THREAD] Number of packages being loaded on the async thread and post loaded on the game thread */
-	FThreadSafeCounter ExistingAsyncPackagesCounter;
-
-	FThreadSafeCounter AsyncThreadReady;
-
-	/** When cancelling async loading: list of package requests to cancel */
-	TArray<FAsyncPackageDesc*> QueuedPackagesToCancel;
-	/** When cancelling async loading: list of packages to cancel */
-	TSet<FAsyncPackage*> PackagesToCancel;
-
-	/** Async loading thread ID */
-	static uint32 AsyncLoadingThreadID;
-
-	/** Enum describing async package request insertion mode */
-	enum class EAsyncPackageInsertMode
-	{
-		InsertBeforeMatchingPriorities,	// Insert this package before all other packages of the same priority
-		InsertAfterMatchingPriorities	// Insert this package after all other packages of the same priority
-	};
-
-#if LOOKING_FOR_PERF_ISSUES
-	/** Thread safe counter used to accumulate cycles spent on blocking. Using stats may generate to many stats messages. */
-	static FThreadSafeCounter BlockingCycles;
-#endif
-
-	IAsyncPackageLoader* IoStorePackageLoader = nullptr;
-
-public:
-
-	FAsyncLoadingThread(int32 InThreadIndex);
-	virtual ~FAsyncLoadingThread();
-
-	virtual ELoaderType GetLoaderType() const override
-	{
-		return ELoaderType::LegacyLoader;
-	}
-
-	IAsyncPackageLoader* GetIoStorePackageLoader() const
-	{
-		return IoStorePackageLoader;
-	}
-
-	void SetIoStorePackageLoader(IAsyncPackageLoader* InIoStorePackageLoader)
-	{
-		IoStorePackageLoader = InIoStorePackageLoader;
-	}
-
-	//~ Begin FRunnable Interface.
-	virtual bool Init();
-	virtual uint32 Run();
-	virtual void Stop();
-	//~ End FRunnable Interface
-
-	/** Returns the async loading thread singleton */
-	static FAsyncLoadingThread& Get()
-	{
-		check(Instance);
-		return *Instance;
-	}
-
-	/** Start the async loading thread */
-	void StartThread() override;
-
-	bool ShouldAlwaysLoadPackageAsync(const FPackagePath& InPackagePath) override;
-
-	int32 LoadPackage(const FPackagePath& PackagePath, FLoadPackageAsyncOptionalParams OptionalParams) override;
-
-	int32 LoadPackage(
-			const FPackagePath& InPackagePath,
-			FName InCustomName,
-			FLoadPackageAsyncDelegate InCompletionDelegate,
-			EPackageFlags InPackageFlags,
-			int32 InPIEInstanceID,
-			int32 InPackagePriority,
-			const FLinkerInstancingContext* InInstancingContext,
-			uint32 InLoadFlags) override;
-
-	EAsyncPackageState::Type ProcessLoading(bool bUseTimeLimit, bool bUseFullTimeLimit, double TimeLimit) override;
-
-	EAsyncPackageState::Type ProcessLoadingUntilComplete(TFunctionRef<bool()> CompletionPredicate, double TimeLimit) override;
-
-	void FlushLoading(TConstArrayView<int32> RequestIDs) override;
-
-	void NotifyConstructedDuringAsyncLoading(UObject* Object, bool bSubObject) override;
-
-	void NotifyUnreachableObjects(const TArrayView<FUObjectItem*>& UnreachableObjects) override {};
-
-	void FireCompletedCompiledInImport(void* AsyncPacakge, FPackageIndex Import);
-
-	void NotifyRegistrationEvent(FName PackageName, FName Name, ENotifyRegistrationType NotifyRegistrationType, ENotifyRegistrationPhase NotifyRegistrationPhase, UObject* (*InRegister)(), bool InbDynamic, UObject* FinishedObject) override;
-
-	void NotifyRegistrationComplete() override;
-
-	/** [EDL] Event queue */
-	FAsyncLoadEventQueue EventQueue;
-
-	/** [EDL] Gets a package from a weak pointer */
-	FORCEINLINE FAsyncPackage* GetPackage(FWeakAsyncPackagePtr Ptr)
-	{
-		if (Ptr.PackageName != NAME_None && Ptr.SerialNumber)
-		{
-			FAsyncPackage* Package = FindAsyncPackage(Ptr.PackageName);
-			if (Package && Package->SerialNumber == Ptr.SerialNumber)
-			{
-				return Package;
-			}
-		}
-		return nullptr;
-	}
-
-	/** [EDL] Queues CreateLinker event */
-	void QueueEvent_CreateLinker(FAsyncPackage* Pkg, int32 EventSystemPriority = 0);
-	/** [EDL] Queues FinishLinker event */
-	void QueueEvent_FinishLinker(FWeakAsyncPackagePtr WeakPtr, int32 EventSystemPriority = 0);
-	/** [EDL] Queues StartImportPackages event */
-	void QueueEvent_StartImportPackages(FAsyncPackage* Pkg, int32 EventSystemPriority = 0);
-	/** [EDL] Queues SetupImports event */
-	void QueueEvent_SetupImports(FAsyncPackage* Pkg, int32 EventSystemPriority = 0);
-	/** [EDL] Queues SetupExports event */
-	void QueueEvent_SetupExports(FAsyncPackage* Pkg, int32 EventSystemPriority = 0);
-	/** [EDL] Queues ProcessImportsAndExports event */
-	void QueueEvent_ProcessImportsAndExports(FAsyncPackage* Pkg, int32 EventSystemPriority = 0);
-	/** [EDL] Queues ExportsDone event */
-	void QueueEvent_ExportsDone(FAsyncPackage* Pkg, int32 EventSystemPriority = 0);
-	/** [EDL] Queues ProcessPostload event */
-	void QueueEvent_ProcessPostloadWait(FAsyncPackage* Pkg, int32 EventSystemPriority = 0);
-	/** [EDL] Queues StartPostLoad event */
-	void QueueEvent_StartPostLoad(FAsyncPackage* Pkg, int32 EventSystemPriority = 0);
-
-	/** True if multithreaded async loading is currently being used. */
-	FORCEINLINE bool IsMultithreaded() override
-	{
-		return bThreadStarted;
-	}
-
-	/** Sets the current state of async loading */
-	static void EnterAsyncLoadingTick(int32 ThreadIndex)
-	{
-		AsyncLoadingTickCounter.Increment();
-		check(CurrentAsyncLoadingTickThreadIndex == INDEX_NONE || CurrentAsyncLoadingTickThreadIndex == ThreadIndex);
-		CurrentAsyncLoadingTickThreadIndex = ThreadIndex;
-	}
-
-	static void LeaveAsyncLoadingTick(int32 ThreadIndex)
-	{
-		int32 AsyncLoadingTickCounterValue = AsyncLoadingTickCounter.Decrement();
-		check(AsyncLoadingTickCounterValue >= 0);
-		check(CurrentAsyncLoadingTickThreadIndex == ThreadIndex);
-		if (AsyncLoadingTickCounterValue == 0)
-		{
-			CurrentAsyncLoadingTickThreadIndex = INDEX_NONE;
-		}
-	}
-
-	/** Gets the current state of async loading */
-	UE_AUTORTFM_ALWAYS_OPEN static bool GetIsInAsyncLoadingTick()
-	{
-		return !!AsyncLoadingTickCounter.GetValue();
-	}
-
-	UE_AUTORTFM_ALWAYS_OPEN bool IsAsyncLoadingPackages() override
-	{
-		return QueuedPackagesCounter.GetValue() != 0 || ExistingAsyncPackagesCounter.GetValue() != 0;
-	}
-
-	FORCEINLINE bool IsInAsyncLoadThread() override
-	{
-		bool bResult = false;
-		if (IsMultithreaded())
-		{
-			// We still need to report we're in async loading thread even if 
-			// we're on game thread but inside of async loading code (PostLoad mostly)
-			// to make it behave exactly like the non-threaded version
-			bResult = FPlatformTLS::GetCurrentThreadId() == AsyncLoadingThreadID || 
-				(IsInGameThread() && GetIsInAsyncLoadingTick());
-		}
-		else
-		{
-			bResult = IsInGameThread() && GetIsInAsyncLoadingTick();
-		}
-		return bResult;
-	}
-
-	/** Returns true if async loading is suspended */
-	FORCEINLINE bool IsAsyncLoadingSuspendedInternal() const
-	{
-		return !!IsLoadingSuspended.GetValue();
-	}
-
-	virtual bool IsAsyncLoadingSuspended() override
-	{
-		return IsAsyncLoadingSuspendedInternal();
-	}
-
-	FORCEINLINE int32 GetAsyncLoadingSuspendedCount()
-	{
-		FPlatformMisc::MemoryBarrier();
-		return IsLoadingSuspended.GetValue();
-	}
-
-	/** Returns the number of async packages that are currently queued but not yet processed */
-	FORCEINLINE int32 GetNumQueuedPackages() override
-	{
-		return QueuedPackagesCounter.GetValue();
-	}
-
-	/** Returns the number of async packages that are currently being processed */
-	FORCEINLINE int32 GetNumAsyncPackages() override
-	{
-		return ExistingAsyncPackagesCounter.GetValue();
-	}
-
-	/**
-	* [ASYNC THREAD] Finds an existing async package in the AsyncPackages by its name.
-	*
-	* @param PackageName async package name.
-	* @return Pointer to the package or nullptr if not found
-	*/
-	FORCEINLINE FAsyncPackage* FindAsyncPackage(const FName& PackageName)
-	{
-		checkSlow(IsInAsyncLoadThread());
-		return AsyncPackageNameLookup.FindRef(PackageName);
-	}
-
-	/**
-	* [ASYNC THREAD] Inserts package to queue according to priority.
-	*
-	* @param PackageName - async package name.
-	* @param InsertMode - Insert mode, describing how we insert this package into the request list
-	*/
-	void InsertPackage(FAsyncPackage* Package, bool bReinsert = false, EAsyncPackageInsertMode InsertMode = EAsyncPackageInsertMode::InsertBeforeMatchingPriorities);
-
-	/**
-	* [ASYNC THREAD] Finds an existing async package in the LoadedPackages by its name.
-	*
-	* @param PackageName async package name.
-	* @return Index of the async package in LoadedPackages array or INDEX_NONE if not found.
-	*/
-	FORCEINLINE FAsyncPackage* FindLoadedPackage(const FName& PackageName)
-	{
-		checkSlow(IsInAsyncLoadThread());
-		return LoadedPackagesNameLookup.FindRef(PackageName);
-	}
-
-	/**
-	* [ASYNC/GAME THREAD] Queues a package for streaming.
-	*
-	* @param Package package descriptor.
-	*/
-	void QueuePackage(FAsyncPackageDesc& Package);
-
-	/**
-	* [GAME THREAD] Cancels streaming
-	*/
-	void CancelLoading() override;
-
-	/**
-	* [GAME THREAD] Stops the async loading thread and blocks until the thread has exited.
-	*/
-	void ShutdownLoading() override;
-
-	/**
-	* [GAME THREAD] Suspends async loading thread
-	*/
-	void SuspendLoading() override;
-
-	/**
-	* [GAME THREAD] Resumes async loading thread
-	*/
-	void ResumeLoading() override;
-
-	/**
-	* [ASYNC/GAME THREAD] Queues a package for streaming.
-	*
-	* @param Package package descriptor.
-	*/
-	FORCEINLINE FAsyncPackage* GetPackage(int32 PackageIndex)
-	{
-		checkSlow(IsInAsyncLoadThread());
-		return AsyncPackages[PackageIndex];
-	}
-
-	/**
-	* [ASYNC* THREAD] Loads all packages
-	*
-	* @param OutPackagesProcessed Number of packages processed in this call.
-	* @param bUseTimeLimit True if time limit should be used [time-slicing].
-	* @param bUseFullTimeLimit True if full time limit should be used [time-slicing].
-	* @param TimeLimit Maximum amount of time that can be spent in this call [time-slicing].
-	* @param FlushRequest The request to flush a package and its depedencies. May be invalid.
-	* @return The current state of async loading
-	*/
-	EAsyncPackageState::Type ProcessAsyncLoading(int32& OutPackagesProcessed, bool bUseTimeLimit, bool bUseFullTimeLimit, float TimeLimit, FFlushRequest& FlushRequest);
-
-	/**
-	* [EDL] [ASYNC* THREAD] Checks fopr cycles in the event driven loader and does fatal errors in that case
-	*/
-	void CheckForCycles();
-
-	/**
-	* [GAME THREAD] Ticks game thread side of async loading.
-	*
-	* @param bUseTimeLimit True if time limit should be used [time-slicing].
-	* @param bUseFullTimeLimit True if full time limit should be used [time-slicing].
-	* @param TimeLimit Maximum amount of time that can be spent in this call [time-slicing].
-	* @param FlushRequest The request to flush a package and its depedencies. May be invalid.
-	* @return The current state of async loading
-	*/
-	EAsyncPackageState::Type TickAsyncLoading(bool bUseTimeLimit, bool bUseFullTimeLimit, double TimeLimit, FFlushRequest FlushRequest = FFlushRequest());
-
-	/**
-	* [ASYNC THREAD] Main thread loop
-	*
-	* @param bUseTimeLimit True if time limit should be used [time-slicing].
-	* @param bUseFullTimeLimit True if full time limit should be used [time-slicing].
-	* @param TimeLimit Maximum amount of time that can be spent in this call [time-slicing].
-	* @param FlushRequest The request to flush a package and its depedencies. May be invalid.
-	*/
-	EAsyncPackageState::Type TickAsyncThread(bool bUseTimeLimit, bool bUseFullTimeLimit, double TimeLimit, bool& bDidSomething, FFlushRequest& FlushRequest);
-
-	/** Initializes async loading thread */
-	void InitializeLoading() override;
-
-	/** 
-	 * [GAME THREAD] Gets the load percentage of the specified package
-	 * @param PackageName Name of the package to return async load percentage for
-	 * @return Percentage (0-100) of the async package load or -1 of package has not been found
-	 */
-	float GetAsyncLoadPercentage(const FName& PackageName) override;
-
-
-private:
-	/** 
-	 * [ASYNC/GAME THREAD] Adds a request ID to the list of pending requests
-	 */
-	void AddPendingRequest(int32 RequestID)
-	{
-#if THREADSAFE_UOBJECTS
-		UE::TScopeLock Lock(PendingRequestsCritical);
-#endif
-		if (!PendingRequests.Contains(RequestID))
-		{
-			PendingRequests.Add(RequestID);
-		}
-	}
-
-	/** 
-	 * [ASYNC/GAME THREAD] Removes a request ID from the list of pending requests
-	 */
-	void RemovePendingRequests(TArray<int32>& RequestIDs)
-	{
-#if THREADSAFE_UOBJECTS
-		UE::TScopeLock Lock(PendingRequestsCritical);
-#endif
-		for (int32 ID : RequestIDs)
-		{
-			PendingRequests.Remove(ID);
-			TRACE_LOADTIME_END_REQUEST(ID);
-		}		
-	}
-
-	/**
-	 * [ASYNC/GAME THREAD] Pop flush request from the flush request stacks.
-	 * @note: one of the request id passed should normally match the top of the request stack if any.
-	 */
-	void CompleteFlushRequests(TArray<int32>& RequestIDs);
-
-	/**
-	 * [ASYNC THREAD] Add a mapping being duplicate request and the main/first explicit request made for a package
-	 * @param DuplicateRequestId the duplicate request id
-	 * @param MainRequestId the main request id that will be used for processing flush request
-	 */
-	void MapDuplicateRequestID(int32 DuplicateRequestId, int32 MainRequestId);
-
-
-	/** [ASYNC/GAME THREAD] Number of package load requests in the async loading queue */
-	int32 GetQueuedPackagesCount() const
-	{
-		return QueuedPackagesCounter.GetValue();
-	}
-	/** [ASYNC/GAME THREAD] Number of packages being loaded on the async thread and post loaded on the game thread */
-	int32 GetExistingAsyncPackagesCount() const
-	{
-		return ExistingAsyncPackagesCounter.GetValue();
-	}
-
-	/**
-	 * [ASYNC/GAME THREAD] Checks if a request ID already is added to the loading queue
-	 */
-	bool ContainsRequestInternal(int32 RequestID)
-	{
-#if THREADSAFE_UOBJECTS
-		UE::TScopeLock Lock(PendingRequestsCritical);
-#endif
-		return PendingRequests.Contains(RequestID);
-	}
-
-	/**
-	 * [ASYNC/GAME THREAD] Checks if a request ID already is added to the loading queue
-	 */
-	bool ContainsAnyRequestInternal(TConstArrayView<int32> RequestIDs)
-	{
-#if THREADSAFE_UOBJECTS
-		UE::TScopeLock Lock(PendingRequestsCritical);
-#endif
-		return Algo::AnyOf(RequestIDs, [this](int32 RequestID){ return PendingRequests.Contains(RequestID); });
-	}
-
-	/**
-	 * [ASYNC/GAME THREAD] Checks if a request ID already is added to the loading queue
-	 */
-	bool ContainsAnyRequestInternal(TConstArrayView<std::atomic<int32>> RequestIDs)
-	{
-#if THREADSAFE_UOBJECTS
-		UE::TScopeLock Lock(PendingRequestsCritical);
-#endif
-		return Algo::AnyOf(RequestIDs, [this](const std::atomic<int32>& RequestID){ return PendingRequests.Contains(RequestID.load(std::memory_order_relaxed)); });
-	}
-
-	/** 
-	 * [GAME THREAD] Add a Flush request to the flush stack
-	 * @param RequestIDs the request id of the packages to flush
-	 * @returns the flushtree for the flush request to use on the game thread
-	 * @note the FlushRequest is a shared resource between the async loading thread and the main thread
-	 */
-	FFlushRequest AddFlushRequest(TConstArrayView<int32> RequestIDs);
-
-	/** 
-	 * [ASYNC THREAD] Peek the top of the flush request stack
-	 * @returns the current flush tree to consider 
-	 */
-	FFlushRequest PeekFlushRequest();
-
-	/** 
-	 * [ASYNC THREAD] Adjust the flush request if needed, in case of duplicate request
-	 * @param  InFlushRequest the flush request to adjust
-	 */
-	void AdjustFlushRequest(FFlushRequest& FlushRequest);
-
-	/**
-	 * [GAME THREAD/ASYNC THREAD] Use by both the game thread and async thread to determine if the current package should be processed in relation to the current flush tree
-	 * @param InAsyncPackage The package to potentially process
-	 * @param FlushRequest the current flush request
-	 * @returns true if we should process the package
-	 */
-	bool ShouldProcessPackage(FAsyncPackage* InAsyncPackage, const FFlushRequest& FlushRequest);
-	
-	/**
-	* [GAME THREAD] Performs game-thread specific operations on loaded packages (not-thread-safe PostLoad, callbacks)
-	*
-	* @param bUseTimeLimit True if time limit should be used [time-slicing].
-	* @param bUseFullTimeLimit True if full time limit should be used [time-slicing].
-	* @param TimeLimit Maximum amount of time that can be spent in this call [time-slicing].
-	* @param FlushRequest Package dependency tree to be flushed
-	* @return The current state of async loading
-	*/
-	EAsyncPackageState::Type ProcessLoadedPackages(bool bUseTimeLimit, bool bUseFullTimeLimit, double TimeLimit, bool& bDidSomething, const FFlushRequest& FlushRequest);
-
-	/**
-	* [ASYNC THREAD] Creates async packages from the queued requests
-	*/
-	int32 CreateAsyncPackagesFromQueue(bool bUseTimeLimit, bool bUseFullTimeLimit, double TimeLimit);
-
-	/**
-	* [ASYNC THREAD] Internal helper function for processing a package load request. If dependency preloading is enabled, 
-	* it will call itself recursively for all the package dependencies
-	*/
-	void ProcessAsyncPackageRequest(FAsyncPackageDesc* InRequest, FAsyncPackage* InRootPackage);
-
-	/**
-	* [ASYNC THREAD] Internal helper function for updating the priorities of an existing package and all its dependencies
-	*/
-	void UpdateExistingPackagePriorities(FAsyncPackage* InPackage, TAsyncLoadPriority InNewPriority);
-
-	/**
-	* [ASYNC THREAD] Finds existing async package and adds the new request's completion callback to it.
-	* @param PackageRequest The package description for the request
-	* @param PackageList The package list to look into
-	*/
-	FAsyncPackage* FindExistingPackageAndAddCompletionCallback(FAsyncPackageDesc* PackageRequest, TMap<FName, FAsyncPackage*>& PackageList);
-
-	/**
-	* [ASYNC THREAD] Adds a package to a list of packages that have finished loading on the async thread
-	*/
-	void AddToLoadedPackages(FAsyncPackage* Package);
-
-	/** Cancels async loading internally */
-	void CancelAsyncLoadingInternal();
-	void FinalizeCancelAsyncLoadingInternal();
-
-	/** Event graph for EDL */
-	FEventLoadGraph EventGraph;	
-	/** ELD precache handler */
-	FPrecacheCallbackHandler* PrecacheHandler;
-	/** This Async Loading Thread Index (future use) */
-	int32 AsyncLoadingThreadIndex;
-
-	/** Number of times we re-entered the async loading tick, mostly used by singlethreaded ticking. Debug purposes only. */
-	static FThreadSafeCounter AsyncLoadingTickCounter;
-	static int32 CurrentAsyncLoadingTickThreadIndex;
-
-public:
-
-	/** Gets the EDL event graph */
-	FEventLoadGraph& GetEventGraph()
-	{
-		return EventGraph;
-	}
-
-	/** Gets the EDL precache handler */
-	FPrecacheCallbackHandler& GetPrecacheHandler()
-	{
-		return *PrecacheHandler;
-	}
-
-	/** Gets this ALT index (future use) */
-	int32 GetThreadIndex() const
-	{
-		return AsyncLoadingThreadIndex;
-	}
-
-#if !UE_BUILD_SHIPPING && !UE_BUILD_TEST
-	FThreadSafeCounter RecursionNotAllowed;
-#endif
-
-	static FAsyncLoadingThread* Instance;
-};
diff --git a/Engine/Source/Runtime/CoreUObject/Private/Serialization/AsyncPackageLoader.cpp b/Engine/Source/Runtime/CoreUObject/Private/Serialization/AsyncPackageLoader.cpp
--- a/Engine/Source/Runtime/CoreUObject/Private/Serialization/AsyncPackageLoader.cpp
+++ b/Engine/Source/Runtime/CoreUObject/Private/Serialization/AsyncPackageLoader.cpp
@@ -217,58 +217,21 @@ void InitAsyncThread()
 {
 	LLM_SCOPE(ELLMTag::AsyncLoading);
 	
-	bool bCommandLineDisabled = FParse::Param(FCommandLine::Get(), TEXT("NoZenLoader"));
-
-	if (FIoDispatcher::IsInitialized())
+	if (!FIoDispatcher::IsInitialized())
 	{
-		bool bSettingsEnabled = false;
-		bool bCommandLineEnabled = false;
-		bool bHasUseIoStoreParamInEditor = false;
-#if WITH_EDITOR
-		// enable zenloader by default in editor case
-		bSettingsEnabled = true;
-
-		bCommandLineEnabled = FParse::Param(FCommandLine::Get(), TEXT("ZenLoader"));
-		check(GConfig);
-		GConfig->GetBool(TEXT("/Script/Engine.EditorStreamingSettings"), TEXT("s.ZenLoaderEnabled"), bSettingsEnabled, GEngineIni);
-		bHasUseIoStoreParamInEditor = UE_FORCE_USE_IOSTORE || FParse::Param(FCommandLine::Get(), TEXT("UseIoStore"));
-#endif
-		FIoDispatcher& IoDispatcher = FIoDispatcher::Get();
-		const bool bHasScriptObjectsChunk = FIoDispatcherInternal::HasPackageData();
-
-		// If there is no IODispatcher data, enable the loose file loading path for zenloader so 
-		// the user can attempt to find content on disk, otherwise they will have no loading capabilities
-		if (!bHasScriptObjectsChunk)
-		{
-			FAsyncLoadingThreadSettings::Get().bLooseFileLoadingEnabled = true;
-		}
-		if (!bCommandLineDisabled && (bSettingsEnabled || bCommandLineEnabled))
-		{
-			GPackageLoader.Reset(MakeAsyncPackageLoader2(IoDispatcher));
-		}
-		else if (bHasScriptObjectsChunk)
-		{
-			GPackageLoader.Reset(MakeAsyncPackageLoader2(IoDispatcher));
-		}
+		FIoDispatcher::Initialize();
+		GIoDispatcherInitializedByLoader = true;
 	}
-	if (!GPackageLoader.IsValid())
-	{
-		if (GUseOldLoaderAsFallback || bCommandLineDisabled)
-		{
-			GPackageLoader = MakeUnique<FAsyncLoadingThread>(/** ThreadIndex = */ 0);
-		}
-		else
-		{
-			if (!FIoDispatcher::IsInitialized())
-			{
-				FIoDispatcher::Initialize();
-				GIoDispatcherInitializedByLoader = true;
-			}
 
-			GPackageLoader.Reset(MakeAsyncPackageLoader2(FIoDispatcher::Get()));
-		}
+	// If there is no IODispatcher data, enable the loose file loading path for zenloader so 
+	// the user can attempt to find content on disk, otherwise they will have no loading capabilities
+	if (!FIoDispatcherInternal::HasPackageData())
+	{
+		FAsyncLoadingThreadSettings::Get().bLooseFileLoadingEnabled = true;
 	}
 
+	GPackageLoader.Reset(MakeAsyncPackageLoader2(FIoDispatcher::Get()));
+
 #if UE_AUTORTFM_ENABLED
 	// If we are using AutoRTFM we need to layer a transactionally safe package loader over
 	// whatever loader was to be used.
@@ -641,15 +604,6 @@ static FPackagePath GetLoadPackageAsyncPackagePath(FStringView InPackageNameOrFi
 	return PackagePath;
 }
 
-bool ShouldAlwaysLoadPackageAsync(const FPackagePath& InPackagePath)
-{
-	if (!GPackageLoader)
-	{
-		return false;
-	}
-	return GPackageLoader->ShouldAlwaysLoadPackageAsync(InPackagePath);
-}
-
 int32 LoadPackageAsync(const FPackagePath& InPackagePath, FLoadPackageAsyncOptionalParams InOptionalParams)
 {
 	LLM_SCOPE(ELLMTag::AsyncLoading);
diff --git a/Engine/Source/Runtime/CoreUObject/Private/Serialization/TransactionallySafeAsyncLoading.cpp b/Engine/Source/Runtime/CoreUObject/Private/Serialization/TransactionallySafeAsyncLoading.cpp
--- a/Engine/Source/Runtime/CoreUObject/Private/Serialization/TransactionallySafeAsyncLoading.cpp
+++ b/Engine/Source/Runtime/CoreUObject/Private/Serialization/TransactionallySafeAsyncLoading.cpp
@@ -106,11 +106,6 @@ public:
 		WrappedPackageLoader->StartThread();
 	}
 
-	bool ShouldAlwaysLoadPackageAsync(const FPackagePath& PackagePath) override
-	{
-		return AutoRTFM::IsClosed() || WrappedPackageLoader->ShouldAlwaysLoadPackageAsync(PackagePath);
-	}
-
 	int32 LoadPackage(
 			const FPackagePath& PackagePath,
 			FName CustomPackageName,
diff --git a/Engine/Source/Runtime/CoreUObject/Private/UObject/UObjectGlobals.cpp b/Engine/Source/Runtime/CoreUObject/Private/UObject/UObjectGlobals.cpp
--- a/Engine/Source/Runtime/CoreUObject/Private/UObject/UObjectGlobals.cpp
+++ b/Engine/Source/Runtime/CoreUObject/Private/UObject/UObjectGlobals.cpp
@@ -1615,26 +1615,6 @@ public:
 	}
 };
 
-// this class is a hack to work around calling private functions in the linker 
-// I just want to replace the Linkers loader with a custom one
-class FUnsafeLinkerLoad : public FLinkerLoad
-{
-public:
-	FUnsafeLinkerLoad(UPackage *Package, const FPackagePath& PackagePath, const FPackagePath& DiffPackagePath, uint32 LoadFlags)
-		: FLinkerLoad(Package, PackagePath, LoadFlags)
-	{
-		Package->SetLinker(this);
-		while ( Tick(0.0, false, false, nullptr) == FLinkerLoad::LINKER_TimedOut ) 
-		{ 
-		}
-
-		FOpenPackageResult OtherFile = IPackageResourceManager::Get().OpenReadPackage(DiffPackagePath);
-		checkf(!OtherFile.Archive.IsValid() || OtherFile.Format == EPackageFormat::Binary, TEXT("Text format is not yet supported with DiffPackage"));
-		FDiffFileArchive* DiffArchive = new FDiffFileArchive(GetLoader(), OtherFile.Archive.Release());
-		SetLoader(DiffArchive, true /* bInLoaderNeedsEngineVersionChecks */);
-	}
-};
-
 #endif
 
 // Temporary load counter for the game thread, used mostly for checking if we're still loading
@@ -1642,36 +1622,10 @@ public:
 static int32 GGameThreadLoadCounter = 0;
 static int32 GGameThreadEndLoadCounter = -1;
 
-/** Notify delegate listeners of all the packages that loaded; called only once per explicit call to LoadPackage. */
-void BroadcastEndLoad(TArray<UPackage*>&& LoadedPackages)
-{
-#if WITH_EDITOR
-	// check(IsInGameThread()) was called by the caller, but we still need to test !IsInAsyncLoadingThread to exclude that callsite when the engine is single-threaded
-	if (!IsInAsyncLoadingThread() && GGameThreadLoadCounter == 0)
-	{
-		LoadedPackages.RemoveAllSwap([](UPackage* Package)
-			{
-				return Package->HasAnyFlags(RF_Transient) || Package->HasAnyPackageFlags(PKG_InMemoryOnly);
-			});
-		for (UPackage* LoadedPackage : LoadedPackages)
-		{
-			LoadedPackage->SetHasBeenEndLoaded(true);
-		}
-		++GGameThreadEndLoadCounter; // Starts at -1, so the first increment takes it to 0
-		FCoreUObjectDelegates::OnEndLoadPackage.Broadcast(
-			FEndLoadPackageContext{ LoadedPackages, GGameThreadEndLoadCounter, true/* bSynchronous */ });
-		--GGameThreadEndLoadCounter;
-		ensure(GGameThreadEndLoadCounter >= -1);
-	}
-#endif
-}
-
 UE_TRACE_EVENT_BEGIN(CUSTOM_LOADTIMER_LOG, LoadPackageInternal, NoSync)
 	UE_TRACE_EVENT_FIELD(UE::Trace::WideString, PackageName)
 UE_TRACE_EVENT_END()
 
-bool ShouldAlwaysLoadPackageAsync(const FPackagePath& InPackagePath);
-
 UPackage* LoadPackageInternal(UPackage* InOuter, const FPackagePath& PackagePath, uint32 LoadFlags, FLinkerLoad* ImportLinker, FArchive* InReaderOverride,
 	const FLinkerInstancingContext* InstancingContext, const FPackagePath* DiffPackagePath)
 {
@@ -1694,328 +1648,68 @@ UPackage* LoadPackageInternal(UPackage* InOuter, const FPackagePath& PackagePath
 	}
 
 	FUObjectThreadContext& ThreadContext = FUObjectThreadContext::Get();
-	if (ShouldAlwaysLoadPackageAsync(PackagePath))
-	{
-		checkf(!InOuter || !InOuter->GetOuter(), TEXT("Loading into subpackages is not implemented.")); // Subpackages are no longer supported in UE
-		FName PackageName(InOuter ? InOuter->GetFName() : PackagePath.GetPackageFName());
-		if (PackageName.IsNone())
-		{
-			UE_LOG(LogUObjectGlobals, Warning, TEXT("Attempted to LoadPackage from non-mounted path %s. This is not supported."), *PackagePath.GetDebugName());
-			return nullptr;
-		}
-
-		// This delegate is not thread-safe and the subscribers are mostly interested by sync loads
-		// that might stall the game thread anyway. So for now, do not broadcast when sync loading
-		// from the loading thread.
-		if (IsInGameThread() && FCoreDelegates::OnSyncLoadPackage.IsBound())
-		{
-			FCoreDelegates::OnSyncLoadPackage.Broadcast(PackageName.ToString());
-		}
-
-		ThreadContext.SyncLoadUsingAsyncLoaderCount++;
-		EPackageFlags PackageFlags = PKG_None;
-#if WITH_EDITOR
-		// If we are loading a package for diffing, set the package flag
-		if (LoadFlags & LOAD_ForDiff)
-		{
-			PackageFlags |= PKG_ForDiffing;
-		}
-		if ((!FApp::IsGame() || GIsEditor) && (LoadFlags & LOAD_PackageForPIE) != 0)
-		{
-			PackageFlags |= PKG_PlayInEditor;
-		}
-#endif
-		FLoadPackageAsyncOptionalParams OptionalParams
-		{
-			.CustomPackageName = PackageName,
-			.PackageFlags = PackageFlags,
-			.PackagePriority = INT32_MAX,
-			.InstancingContext = InstancingContext,
-			.LoadFlags = LoadFlags
-		};
-		int32 RequestID = LoadPackageAsync(PackagePath, MoveTemp(OptionalParams));
 
-		if (RequestID != INDEX_NONE)
-		{
-			UE_SCOPED_IO_ACTIVITY(*WriteToString<512>(TEXT("Sync "), PackagePath.GetDebugName()));
-			FlushAsyncLoading(RequestID);
-		}
-		ThreadContext.SyncLoadUsingAsyncLoaderCount--;
-
-		if (InOuter)
-		{
-			return InOuter;
-		}
-		else
-		{
-			UPackage* Result = FindObjectFast<UPackage>(nullptr, PackageName);
-			if (!Result)
-			{
-				// Might have been redirected
-				const FCoreRedirectObjectName NewPackageName = FCoreRedirects::GetRedirectedName(ECoreRedirectFlags::Type_Package, FCoreRedirectObjectName(NAME_None, NAME_None, PackageName));
-				if (NewPackageName.PackageName != PackageName)
-				{
-					Result = FindObjectFast<UPackage>(nullptr, NewPackageName.PackageName);
-				}
-			}
-			return Result;
-		}
+	checkf(!InOuter || !InOuter->GetOuter(), TEXT("Loading into subpackages is not implemented.")); // Subpackages are no longer supported in UE
+	FName PackageName(InOuter ? InOuter->GetFName() : PackagePath.GetPackageFName());
+	if (PackageName.IsNone())
+	{
+		UE_LOG(LogUObjectGlobals, Warning, TEXT("Attempted to LoadPackage from non-mounted path %s. This is not supported."), *PackagePath.GetDebugName());
+		return nullptr;
 	}
 
-	checkf(IsInGameThread(), TEXT("Unable to load %s. Objects and Packages can only be loaded from the game thread with the currently active loader '%s'."), *PackagePath.GetDebugName(), LexToString(GetLoaderType()));
-	UPackage* Result = nullptr;
-
-#if WITH_EDITOR
-	// In the editor loading cannot be part of a transaction as it cannot be undone, and may result in recording half-loaded objects. So we suppress any active transaction while in this stack, and set the editor loading flag
-	TGuardValue<ITransaction*> SuppressTransaction(GUndo, nullptr);
-	TGuardValueAccessors<bool> IsEditorLoadingPackage(UE::GetIsEditorLoadingPackage, UE::SetIsEditorLoadingPackage, GIsEditor || UE::GetIsEditorLoadingPackage());
-#endif
-
-	TOptional<FScopedSlowTask> SlowTask;
-	if (ShouldCreateThrottledSlowTask())
+	// This delegate is not thread-safe and the subscribers are mostly interested by sync loads
+	// that might stall the game thread anyway. So for now, do not broadcast when sync loading
+	// from the loading thread.
+	if (IsInGameThread() && FCoreDelegates::OnSyncLoadPackage.IsBound())
 	{
-		static const FTextFormat LoadingPackageTextFormat = NSLOCTEXT("Core", "LoadingPackage_Scope", "Loading Package '{0}'");
-		SlowTask.Emplace(100.0f, FText::Format(LoadingPackageTextFormat, PackagePath.GetDebugNameText()));
-		SlowTask->Visibility = ESlowTaskVisibility::Invisible;
-		SlowTask->EnterProgressFrame(10);
+		FCoreDelegates::OnSyncLoadPackage.Broadcast(PackageName.ToString());
 	}
 
-	if (FCoreDelegates::OnSyncLoadPackage.IsBound())
+	ThreadContext.SyncLoadUsingAsyncLoaderCount++;
+	EPackageFlags PackageFlags = PKG_None;
+#if WITH_EDITOR
+	// If we are loading a package for diffing, set the package flag
+	if (LoadFlags & LOAD_ForDiff)
 	{
-		FCoreDelegates::OnSyncLoadPackage.Broadcast(PackagePath.GetPackageNameOrFallback());
+		PackageFlags |= PKG_ForDiffing;
 	}
-	
-	TRACE_LOADTIME_POSTLOAD_SCOPE;
-
-	// Set up a load context
-	TRefCountPtr<FUObjectSerializeContext> LoadContext = ThreadContext.GetSerializeContext();
-
-	UE_SCOPED_IO_ACTIVITY(*WriteToString<512>(TEXT("Sync "), PackagePath.GetDebugName()));
-
-	// Try to load.
-	BeginLoad(LoadContext, *PackagePath.GetDebugName());
-
-	if (!ImportLinker)
+	if ((!FApp::IsGame() || GIsEditor) && (LoadFlags & LOAD_PackageForPIE) != 0)
 	{
-		TRACE_LOADTIME_BEGIN_REQUEST(0);
+		PackageFlags |= PKG_PlayInEditor;
 	}
-	ON_SCOPE_EXIT
+#endif
+	FLoadPackageAsyncOptionalParams OptionalParams
 	{
-		if (!ImportLinker)
-		{
-			TRACE_LOADTIME_END_REQUEST(0);
-		}
+		.CustomPackageName = PackageName,
+		.PackageFlags = PackageFlags,
+		.PackagePriority = INT32_MAX,
+		.InstancingContext = InstancingContext,
+		.LoadFlags = LoadFlags
 	};
+	int32 RequestID = LoadPackageAsync(PackagePath, MoveTemp(OptionalParams));
 
-	bool bFullyLoadSkipped = false;
-
-	if (SlowTask)
+	if (RequestID != INDEX_NONE)
 	{
-		SlowTask->EnterProgressFrame(30);
+		UE_SCOPED_IO_ACTIVITY(*WriteToString<512>(TEXT("Sync "), PackagePath.GetDebugName()));
+		FlushAsyncLoading(RequestID);
 	}
+	ThreadContext.SyncLoadUsingAsyncLoaderCount--;
 
-	// Declare here so that the linker does not get destroyed before ResetLoaders is called
-	FLinkerLoad* Linker = nullptr;
-	TArray<UPackage*> LoadedPackages;
+	if (InOuter)
 	{
-		// Keep track of start time.
-		const double StartTime = FPlatformTime::Seconds();
-
-		// Create a new linker object which goes off and tries load the file.
-#if WITH_EDITOR
-		if (DiffPackagePath)
-		{
-			// Create the package with the provided long package name.
-			if (!InOuter)
-			{
-				InOuter = CreatePackage(*PackagePath.GetPackageName());
-			}
-			
-			new FUnsafeLinkerLoad(InOuter, PackagePath, *DiffPackagePath, LOAD_ForDiff);
-		}
-#endif
-
-		{
-			FUObjectSerializeContext* InOutLoadContext = LoadContext;
-			Linker = GetPackageLinker(InOuter, PackagePath, LoadFlags, nullptr, InReaderOverride, &InOutLoadContext, ImportLinker, InstancingContext);
-			if (ImportLinker)
-			{
-				TRACE_LOADTIME_ASYNC_PACKAGE_IMPORT_DEPENDENCY(ImportLinker, Linker);
-			}
-			else
-			{
-				TRACE_LOADTIME_ASYNC_PACKAGE_REQUEST_ASSOCIATION(Linker, 0);
-			}
-			if (InOutLoadContext != LoadContext && InOutLoadContext)
-			{
-				// The linker already existed and was associated with another context
-				LoadContext->DecrementBeginLoadCount();
-				LoadContext = InOutLoadContext;
-				LoadContext->IncrementBeginLoadCount();
-			}			
-		}
-
-		if (!Linker)
-		{
-			EndLoad(LoadContext, &LoadedPackages);
-			BroadcastEndLoad(MoveTemp(LoadedPackages));
-			return nullptr;
-		}
-
-		Result = Linker->LinkerRoot;
-		checkf(Result, TEXT("LinkerRoot is null"));
-		UE_TRACK_REFERENCING_PACKAGE_SCOPED(Result, PackageAccessTrackingOps::NAME_Load);
-
-		auto EndLoadAndCopyLocalizationGatherFlag = [&]
-		{
-			EndLoad(LoadContext, &LoadedPackages);
-			// Set package-requires-localization flags from archive after loading. This reinforces flagging of packages that haven't yet been resaved.
-			Result->ThisRequiresLocalizationGather(Linker->RequiresLocalizationGather());
-		};
-
-		if (Result->HasAnyFlags(RF_WasLoaded))
-		{
-			// The linker is associated with a package that has already been loaded.
-			// Loading packages that have already been loaded is unsupported.
-			EndLoadAndCopyLocalizationGatherFlag();	
-			BroadcastEndLoad(MoveTemp(LoadedPackages));
-			return Result;
-		}
-
-		// The time tracker keeps track of time spent in LoadPackage.
-		FExclusiveLoadPackageTimeTracker::FScopedPackageTracker Tracker(Result);
-
-		// If we are loading a package for diffing, set the package flag
-		if(LoadFlags & LOAD_ForDiff)
-		{
-			Result->SetPackageFlags(PKG_ForDiffing);
-		}
-
-		// Save the PackagePath we loaded from
-		Result->SetLoadedPath(PackagePath);
-		
-		PRAGMA_DISABLE_DEPRECATION_WARNINGS
-		// ScriptSHAHash is no longer supported; delete this block when deleting it.
-		uint8 SavedScriptSHA[20];
-		bool bHasScriptSHAHash = FSHA1::GetFileSHAHash(*Linker->LinkerRoot->GetName(), SavedScriptSHA, false);
-		if (bHasScriptSHAHash)
-		{
-			Linker->StartScriptSHAGeneration();
-		}
-		PRAGMA_ENABLE_DEPRECATION_WARNINGS
-
-		if (SlowTask)
-		{
-			SlowTask->EnterProgressFrame(30);
-		}
-
-		uint32 DoNotLoadExportsFlags = LOAD_Verify;
-#if USE_CIRCULAR_DEPENDENCY_LOAD_DEFERRING
-		// if this linker already has the DeferDependencyLoads flag, then we're
-		// already loading it earlier up the load chain (don't let it invoke any
-		// deeper loads that may introduce a circular dependency)
-		DoNotLoadExportsFlags |= LOAD_DeferDependencyLoads;
-#endif // USE_CIRCULAR_DEPENDENCY_LOAD_DEFERRING
-
-		if ((LoadFlags & DoNotLoadExportsFlags) == 0)
-		{
-			// Make sure we pass the property that's currently being serialized by the linker that owns the import 
-			// that triggered this LoadPackage call
-			FSerializedPropertyScope SerializedProperty(*Linker, ImportLinker ? ImportLinker->GetSerializedProperty() : Linker->GetSerializedProperty());
-			Linker->LoadAllObjects(GEventDrivenLoaderEnabled);
-
-			// @todo: remove me when loading can be self-contained (and EndLoad doesn't check for IsInAsyncLoadingThread) or there's just one loading path
-			// If we start a non-async loading during async loading and the serialization context is not associated with any other package and
-			// doesn't come from an async package, queue this package to be async loaded, otherwise we'll end up not loading its exports
-			if (!Linker->AsyncRoot && LoadContext->GetBeginLoadCount() == 1 && IsInAsyncLoadingThread())
-			{
-				LoadPackageAsync(Linker->LinkerRoot->GetName());
-			}
-		}
-		else
-		{
-			bFullyLoadSkipped = true;
-		}
-
-		if (SlowTask)
-		{
-			SlowTask->EnterProgressFrame(30);
-		}
-
-		Linker->FinishExternalReadDependencies(0.0);
-
-		EndLoadAndCopyLocalizationGatherFlag();
-
-#if WITH_EDITOR
-		UE::SetIsEditorLoadingPackage(IsEditorLoadingPackage.GetOriginalValue());
-#endif
-
-		PRAGMA_DISABLE_DEPRECATION_WARNINGS
-		// ScriptSHAHash is no longer supported; delete this if block when deleting it.
-		if (bHasScriptSHAHash)
-		{
-			uint8 LoadedScriptSHA[20];
-			Linker->GetScriptSHAKey(LoadedScriptSHA);
-			if (FMemory::Memcmp(SavedScriptSHA, LoadedScriptSHA, 20) != 0)
-			{
-				appOnFailSHAVerification(*Linker->GetPackagePath().GetLocalFullPath(), false);
-			}
-		}
-		PRAGMA_ENABLE_DEPRECATION_WARNINGS
-
-		Linker->Flush();
-
-		if (!FPlatformProperties::RequiresCookedData())
-		{
-			// Flush cache on uncooked platforms to free precache memory
-			Linker->FlushCache();
-		}
-
-		// With UE and single asset per package, we load so many packages that some platforms will run out
-		// of file handles. So, this will close the package, but just things like bulk data loading will
-		// fail, so we only currently do this when loading on consoles.
-		// The only exception here is when we're in the middle of async loading where we can't reset loaders yet. This should only happen when
-		// doing synchronous load in the middle of streaming.
-		if (FPlatformProperties::RequiresCookedData())
-		{
-			if (!IsInAsyncLoadingThread())
-			{				
-				if (GGameThreadLoadCounter == 0)
-				{
-					// Sanity check to make sure that Linker is the linker that loaded our Result package or the linker has already been detached
-					check(!Result || Result->GetLinker() == Linker || Result->GetLinker() == nullptr);
-					if (Result && Linker->HasLoader())
-					{
-						ResetLoaders(Result);
-					}
-					// Reset loaders could have already deleted Linker so guard against deleting stale pointers
-					if (Result && Result->GetLinker())
-					{
-						Linker->DestroyLoader();
-					}
-					// And make sure no one can use it after it's been deleted
-					Linker = nullptr;
-				}
-				// Async loading removes delayed linkers on the game thread after streaming has finished
-				else
-				{
-					LoadContext->AddDelayedLinkerClosePackage(Linker);
-				}
-			}
-			else
-			{
-				LoadContext->AddDelayedLinkerClosePackage(Linker);
-			}
-		}
+		return InOuter;
 	}
 
-	if (!bFullyLoadSkipped)
+	UPackage* Result = FindObjectFast<UPackage>(nullptr, PackageName);
+	if (!Result)
 	{
-		// Mark package as loaded.
-		Result->SetFlags(RF_WasLoaded);
+		// Might have been redirected
+		const FCoreRedirectObjectName NewPackageName = FCoreRedirects::GetRedirectedName(ECoreRedirectFlags::Type_Package, FCoreRedirectObjectName(NAME_None, NAME_None, PackageName));
+		if (NewPackageName.PackageName != PackageName)
+		{
+			Result = FindObjectFast<UPackage>(nullptr, NewPackageName.PackageName);
+		}
 	}
-
-	BroadcastEndLoad(MoveTemp(LoadedPackages));
 	return Result;
 }
 
diff --git a/Engine/Source/Runtime/CoreUObject/Public/Serialization/AsyncLoading2.h b/Engine/Source/Runtime/CoreUObject/Public/Serialization/AsyncLoading2.h
--- a/Engine/Source/Runtime/CoreUObject/Public/Serialization/AsyncLoading2.h
+++ b/Engine/Source/Runtime/CoreUObject/Public/Serialization/AsyncLoading2.h
@@ -29,7 +29,6 @@
 class FArchive;
 class FIoDispatcher;
 class IAsyncPackageLoader;
-class IEDLBootNotificationManager;
 class UPackage;
 
 class FPackageImportReference
diff --git a/Engine/Source/Runtime/CoreUObject/Public/Serialization/AsyncPackageLoader.h b/Engine/Source/Runtime/CoreUObject/Public/Serialization/AsyncPackageLoader.h
--- a/Engine/Source/Runtime/CoreUObject/Public/Serialization/AsyncPackageLoader.h
+++ b/Engine/Source/Runtime/CoreUObject/Public/Serialization/AsyncPackageLoader.h
@@ -42,12 +42,14 @@ extern const FName PrestreamPackageClassNameLoad;
 /** Returns true if we're inside a FGCScopeGuard */
 extern bool IsGarbageCollectionLocked();
 
+UE_DEPRECATED(5.8, "This method will be removed.")
 bool IsNativeCodePackage(UPackage* Package);
 
 /** Checks if the object can have PostLoad called on the Async Loading Thread */
 bool CanPostLoadOnAsyncLoadingThread(UObject* Object);
 
 template <typename T>
+UE_DEPRECATED(5.8, "This method will be removed.")
 void ClearFlagsAndDissolveClustersFromLoadedObjects(T& LoadedObjects)
 {
 	const EObjectFlags ObjectLoadFlags = RF_NeedLoad | RF_NeedPostLoad | RF_NeedPostLoadSubobjects | RF_WasLoaded;
@@ -101,7 +103,8 @@ public:
 
 	virtual void StartThread() = 0;
 
-	virtual bool ShouldAlwaysLoadPackageAsync(const FPackagePath& PackagePath) = 0;
+	UE_DEPRECATED(5.8, "This method will be removed.")
+	virtual bool ShouldAlwaysLoadPackageAsync(const FPackagePath& PackagePath) { return true; }
 
 	/**
 	 * Asynchronously load a package.
diff --git a/Engine/Source/Runtime/CoreUObject/Public/UObject/Package.h b/Engine/Source/Runtime/CoreUObject/Public/UObject/Package.h
--- a/Engine/Source/Runtime/CoreUObject/Public/UObject/Package.h
+++ b/Engine/Source/Runtime/CoreUObject/Public/UObject/Package.h
@@ -483,7 +483,6 @@ public:
 
 private:
 	friend class FLinkerLoad;
-	friend class FUnsafeLinkerLoad;
 	friend class FSaveContext;
 	friend struct FAsyncPackage2;
 
