Fix UE_WITH_REMOTE_OBJECT_HANDLE=1 no longer compiling
- Regressed in CL 48350573

#Synced-CL 48355214

diff --git a/Engine/Source/Runtime/CoreUObject/Private/UObject/UObjectHash.cpp b/Engine/Source/Runtime/CoreUObject/Private/UObject/UObjectHash.cpp
--- a/Engine/Source/Runtime/CoreUObject/Private/UObject/UObjectHash.cpp
+++ b/Engine/Source/Runtime/CoreUObject/Private/UObject/UObjectHash.cpp
@@ -1278,7 +1278,7 @@ UObject* StaticFindFirstObjectFastInternal(const UClass* ObjectClass, FName Obje
 UObject* StaticFindObjectFastInternalThreadSafe(FUObjectHashTables& ThreadHash, FRemoteObjectId RemoteId, EObjectFlags InExclusiveFlags, EInternalObjectFlags InExlusiveInternalFlags)
 {
 #if UE_WITH_REMOTE_OBJECT_HANDLE
-	const EInternalObjectFlags ExclusiveInternalFlags = DefaultInternalExclusionFlags | InExlusiveInternalFlags;
+	const EInternalObjectFlags ExclusiveInternalFlags = GetMandatoryInternalExclusionFlags() | InExlusiveInternalFlags;
 	RemoteId = UE::RemoteObject::Private::FRemoteIdLocalizationHelper::GetLocalized(RemoteId);
 	uint32 Hash = GetTypeHash(RemoteId);
 
