[ZenLoader]
- Add optional support for package redirector at runtime since some tests rely on that

#Synced-CL Unknown

diff --git a/Engine/Source/Runtime/CoreUObject/Private/Serialization/AsyncLoading2.cpp b/Engine/Source/Runtime/CoreUObject/Private/Serialization/AsyncLoading2.cpp
--- a/Engine/Source/Runtime/CoreUObject/Private/Serialization/AsyncLoading2.cpp
+++ b/Engine/Source/Runtime/CoreUObject/Private/Serialization/AsyncLoading2.cpp
@@ -555,6 +555,13 @@ static FAutoConsoleVariableRef CVarGFailLoadOnMissingImport(
 	ECVF_Default
 );
 
+static bool GAllowPackageRedirectorSupport = WITH_EDITOR;
+static FAutoConsoleVariableRef CVarGAllowPackageRedirectorSupport(
+	TEXT("s.AllowPackageRedirectorSupport"),
+	GAllowPackageRedirectorSupport,
+	TEXT("Whether to use package redirector support."),
+	ECVF_Default
+);
 
 #if WITH_EDITOR
 // This is important for the editor because the linker can end up recreating missing exports from its CreateImport function.
@@ -663,7 +670,6 @@ static bool AsyncLoading2_ShouldSkipLoadingExport(const EExportFilterFlags Filte
 #endif
 }
 
-#if WITH_EDITOR
 /** Returns the CoreRedirected package name. Otherwise returns the name unchanged. */
 FName ApplyPackageNameRedirections(FName PackageName)
 {
@@ -675,7 +681,6 @@ FName ApplyPackageNameRedirections(FName PackageName)
 
 	return RedirectedName;
 }
-#endif
 
 /** Returns localized package name to load when running as -game or using loose cooked files in a runtime process. Otherwise returns the name unchanged. */
 FName ApplyLooseFileLocalizationPackageNameRedirects(FName PackageName)
@@ -6082,11 +6087,14 @@ bool FAsyncLoadingThread2::CreateAsyncPackagesFromQueue(FAsyncLoadingThreadState
 			}
 
 			FName UPackageName = PackageNameToLoad;
-#if WITH_EDITOR
-			// Redirects should redirect both the UPackage name as well as the package name to load
-			PackageNameToLoad = ApplyPackageNameRedirections(PackageNameToLoad);
-			UPackageName = PackageNameToLoad;
-#endif
+
+			if (GAllowPackageRedirectorSupport)
+			{
+				// Redirects should redirect both the UPackage name as well as the package name to load
+				PackageNameToLoad = ApplyPackageNameRedirections(PackageNameToLoad);
+				UPackageName = PackageNameToLoad;
+			}
+
 			// Localization should only update the package name to load
 			PackageNameToLoad = ApplyLooseFileLocalizationPackageNameRedirects(PackageNameToLoad);
 
@@ -6992,23 +7000,24 @@ void FAsyncPackage2::ImportPackagesRecursiveInner(FAsyncLoadingThreadState2& Thr
 		FName ImportedPackageNameToLoad = ImportedPackageUPackageName;
 		FPackageId ImportedPackageIdToLoad = ImportedPackageId;
 
-#if WITH_EDITOR
-		if (!ImportedPackageNameToLoad.IsNone())
+		if (GAllowPackageRedirectorSupport)
 		{
-			// Redirects should redirect both the UPackage name as well as the package name to load
-			const FName NewPackageNameToLoad = ApplyPackageNameRedirections(ImportedPackageNameToLoad);
-			if (ImportedPackageNameToLoad != NewPackageNameToLoad)
+			if (!ImportedPackageNameToLoad.IsNone())
 			{
-				ImportedPackageNameToLoad = NewPackageNameToLoad;
-				ImportedPackageIdToLoad = FPackageId::FromName(ImportedPackageNameToLoad);
-				ImportedPackageUPackageName = ImportedPackageNameToLoad;
-				ImportedPackageId = ImportedPackageIdToLoad;
-				// Rewrite the import table
-				ImportedPackageIds[LocalImportedPackageIndex] = ImportedPackageId;
-				ImportedPackageNames[LocalImportedPackageIndex] = ImportedPackageUPackageName;
+				// Redirects should redirect both the UPackage name as well as the package name to load
+				const FName NewPackageNameToLoad = ApplyPackageNameRedirections(ImportedPackageNameToLoad);
+				if (ImportedPackageNameToLoad != NewPackageNameToLoad)
+				{
+					ImportedPackageNameToLoad = NewPackageNameToLoad;
+					ImportedPackageIdToLoad = FPackageId::FromName(ImportedPackageNameToLoad);
+					ImportedPackageUPackageName = ImportedPackageNameToLoad;
+					ImportedPackageId = ImportedPackageIdToLoad;
+					// Rewrite the import table
+					ImportedPackageIds[LocalImportedPackageIndex] = ImportedPackageId;
+					ImportedPackageNames[LocalImportedPackageIndex] = ImportedPackageUPackageName;
+				}
 			}
 		}
-#endif
 
 		bool bIsInstanced = false;
 #if WITH_EDITORONLY_DATA
