Restore IsSerializingScriptProperties in FLinkerLoad::Preload

#rnx
#Synced-CL 48299257

diff --git a/Engine/Source/Runtime/CoreUObject/Private/UObject/Class.cpp b/Engine/Source/Runtime/CoreUObject/Private/UObject/Class.cpp
--- a/Engine/Source/Runtime/CoreUObject/Private/UObject/Class.cpp
+++ b/Engine/Source/Runtime/CoreUObject/Private/UObject/Class.cpp
@@ -1937,12 +1937,6 @@ void UStruct::SerializeVersionedTaggedProperties(FStructuredArchive::FSlot Slot,
 								// Serialize the value if it contains a struct or enum because only those may contain an unknown property.
 								if (bSerializeValue)
 								{
-#if UE_INTERNAL_UNIVERSAL_PLACEHOLDERS
-									// The Link call above may have preloaded an object, which will clear
-									// bIsSerializingScriptProperties and cause the linker to stop reading
-									// tagged properties, force setting it here:
-									UnderlyingArchive.MarkScriptSerializationStart(nullptr);
-#endif // UE_INTERNAL_UNIVERSAL_PLACEHOLDERS
 									UnderlyingArchive.Seek(StartOfProperty);
 									FStructuredArchive::FSlot ValueSlotCopy = PropertyRecord.EnterField(TEXT("Value"));
 									void* TempData = TempProperty->AllocateAndInitializeValue();
diff --git a/Engine/Source/Runtime/CoreUObject/Private/UObject/LinkerLoad.cpp b/Engine/Source/Runtime/CoreUObject/Private/UObject/LinkerLoad.cpp
--- a/Engine/Source/Runtime/CoreUObject/Private/UObject/LinkerLoad.cpp
+++ b/Engine/Source/Runtime/CoreUObject/Private/UObject/LinkerLoad.cpp
@@ -5256,6 +5256,9 @@ void FLinkerLoad::Preload( UObject* Object )
 #endif
 
 					UE::FScopedObjectSerializeContext ObjectSerializeScope(Object, *this);
+#if WITH_EDITOR
+					FGuardValue_Bitfield(bIsSerializingScriptProperties, false);
+#endif // WITH_EDITOR
 
 					if (Object->HasAnyFlags(RF_ClassDefaultObject))
 					{
