[Feature] Made IDTs respect override serialization when copying structs
[Feature] Made IDTs create or forward unknown properties when copying unlike struct types
[Feature] Made IDTs copy FInitializedPropertyValueState when redirecting properties
#Synced-CL 48303336

diff --git a/Engine/Source/Runtime/CoreUObject/Internal/UObject/UObjectSerializeContext.h b/Engine/Source/Runtime/CoreUObject/Internal/UObject/UObjectSerializeContext.h
--- a/Engine/Source/Runtime/CoreUObject/Internal/UObject/UObjectSerializeContext.h
+++ b/Engine/Source/Runtime/CoreUObject/Internal/UObject/UObjectSerializeContext.h
@@ -3,6 +3,7 @@
 #pragma once
 
 #include "CoreTypes.h"
+#include "UObject/PropertyPathName.h"
 
 #define UE_API COREUOBJECT_API
 
@@ -23,6 +24,9 @@ struct FScopedObjectSerializeContext
 	UE_API FScopedObjectSerializeContext(UObject* Object, FArchive& Archive);
 	UE_API ~FScopedObjectSerializeContext();
 
+	// specialized context specifically for serializing individual properties during IDTs
+	FScopedObjectSerializeContext(UObject* Object, FArchive& Archive, const FPropertyPathName& PathBeingSerialized);
+
 	FScopedObjectSerializeContext(const FScopedObjectSerializeContext&) = delete;
 	FScopedObjectSerializeContext& operator=(const FScopedObjectSerializeContext&) = delete;
 
@@ -42,6 +46,7 @@ private:
 	bool bSavedTrackUnknownEnumNames;
 	bool bSavedImpersonateProperties;
 	bool bCreateInstanceDataObject;
+	FPropertyPathName SavedSerializedPropertyPath;
 #endif
 };
 
diff --git a/Engine/Source/Runtime/CoreUObject/Private/UObject/InstanceDataTransforms.cpp b/Engine/Source/Runtime/CoreUObject/Private/UObject/InstanceDataTransforms.cpp
--- a/Engine/Source/Runtime/CoreUObject/Private/UObject/InstanceDataTransforms.cpp
+++ b/Engine/Source/Runtime/CoreUObject/Private/UObject/InstanceDataTransforms.cpp
@@ -18,11 +18,13 @@
 #include "UObject/OverridableManager.h"
 #include "UObject/PropertyHelper.h"
 #include "UObject/Package.h"
+#include "UObject/UObjectSerializeContext.h"
 #include "Dom/JsonValue.h"
 #include "Dom/JsonObject.h"
 #include "Serialization/JsonReader.h"
 #include "Serialization/JsonWriter.h"
 #include "Serialization/JsonSerializer.h"
+#include "Serialization/ArchiveSerializedPropertyChain.h"
 #include "HAL/FileManager.h"
 #include "Misc/FileHelper.h"
 #include "Misc/Paths.h"
@@ -71,40 +73,80 @@ namespace UE::Private
 		}
 	};
 
-
-	// TODO: this method allows dangerous and narrowing conversions. That's fine for now because the user explicitly asked to do so but in the future it could be good
-	// to include a flag in the RedirectProperty params that specifies weather the user is ok with narrowing conversions
-	bool TryCopyOrConvert(const FProperty* SourceProperty, const void* SourceContainer, int32 SourceIndex, FProperty* DestinationProperty, void* DestinationContainer, int32 DestinationIndex)
+	TArray<uint8> WriteToBuffer(UObject* Owner, const FPropertyPathNameResolver& Src)
 	{
+		const FProperty* SourceProperty = Src.Value.Property;
+		const void* SourceContainer = Src.Value.Container;
+		int32 SourceIndex = Src.Value.ArrayIndex;
 		const void* SourceData = SourceProperty->ContainerPtrToValuePtr<void>(SourceContainer, SourceIndex == INDEX_NONE ? 0 : SourceIndex);
-		void* DestinationData = DestinationProperty->ContainerPtrToValuePtr<void>(DestinationContainer, DestinationIndex == INDEX_NONE ? 0 : DestinationIndex);
 
-		FUObjectSerializeContext* SerializeContext = FUObjectThreadContext::Get().GetSerializeContext();
-		TGuardValue<bool> ScopedImpersonateProperties(SerializeContext->bImpersonateProperties, true);
 		TArray<uint8> Buffer;
+
 		FObjectWriter ObjectWriter(Buffer); // using FObjectWriter so that object properties will serialize correctly
 		FStructuredArchiveFromArchive StructuredWriter(ObjectWriter);
+
+		// set up all the flags to act as if we're serializing all of 'Owner' and the TPS is about to write the Src data.
+		FScopedObjectSerializeContext SerializeContext(Owner, ObjectWriter, Src.Path);
+
+		// enable overridable serialization
+		if (FOverriddenPropertySet* OverriddenProperties = FOverridableManager::Get().GetOverriddenProperties(Src.Object))
+		{
+			FOverridableSerializationLogic::Enable(OverriddenProperties);
+			for (const FEditPropertyChain::TDoubleLinkedListNode* Itr = Src.EventChain.GetHead(); Itr != Src.EventChain.GetTail(); Itr = Itr->GetNextNode())
+			{
+				ObjectWriter.PushSerializedProperty(Itr->GetValue(), Itr->GetValue()->IsEditorOnlyProperty());
+			}
+		}
+
 		// todo: handle static arrays
 		FPropertyTag WriteTag(const_cast<FProperty*>(SourceProperty), SourceIndex, (uint8*)SourceData);
 		StructuredWriter.GetSlot() << WriteTag;
-		WriteTag.SerializeTaggedProperty(StructuredWriter.GetSlot(), const_cast<FProperty*>(SourceProperty), (uint8*)SourceData, nullptr);
+		WriteTag.SerializeTaggedProperty(StructuredWriter.GetSlot(), const_cast<FProperty*>(SourceProperty), (uint8*)SourceData, (uint8*)SourceData);
+		if (FOverriddenPropertySet* OverriddenProperties = FOverridableManager::Get().GetOverriddenProperties(Src.Object))
+		{
+			FOverridableSerializationLogic::Disable();
+		}
+
+		return Buffer;
+	}
+
+	bool ReadFromBuffer(UObject* Owner, const FPropertyPathNameResolver& Dst, TArray<uint8> Buffer, bool bAllowNoConvert)
+	{
+		bool bResult = false;
 
+		FProperty* DestinationProperty = const_cast<FProperty*>(Dst.Value.Property);
+		void* DestinationContainer = Dst.Value.Container;
+		int32 DestinationIndex = Dst.Value.ArrayIndex;
+		void* DestinationData = DestinationProperty->ContainerPtrToValuePtr<void>(DestinationContainer, DestinationIndex == INDEX_NONE ? 0 : DestinationIndex);
+		
 		FObjectReader ObjectReader(Buffer);
 		FStructuredArchiveFromArchive StructuredReader(ObjectReader);
 
+		// set up all the flags to act as if we're serializing all of 'Owner' and the TPS is about to read the Dst data.
+		FScopedObjectSerializeContext ScopedSerializeContext(Owner, ObjectReader, Dst.Path);
+
+		// enable overridable serialization
+		if (FOverriddenPropertySet* OverriddenProperties = FOverridableManager::Get().GetOverriddenProperties(Dst.Object))
+		{
+			FOverridableSerializationLogic::Enable(OverriddenProperties);
+			for (const FEditPropertyChain::TDoubleLinkedListNode* Itr = Dst.EventChain.GetHead(); Itr != Dst.EventChain.GetTail(); Itr = Itr->GetNextNode())
+			{
+				ObjectReader.PushSerializedProperty(Itr->GetValue(), Itr->GetValue()->IsEditorOnlyProperty());
+			}
+		}
+
 		FPropertyTag ReadTag;
 		StructuredReader.GetSlot() << ReadTag;
 		ReadTag.SetProperty(DestinationProperty);
 		ReadTag.ArrayIndex = DestinationIndex;
 		ReadTag.Name = DestinationProperty->GetFName();
 
-		bool bResult = false;
-		switch (DestinationProperty->ConvertFromType(ReadTag, StructuredReader.GetSlot(), (uint8*)DestinationContainer, SourceProperty->GetOwnerStruct(), nullptr))
+		switch (DestinationProperty->ConvertFromType(ReadTag, StructuredReader.GetSlot(), (uint8*)DestinationContainer, Dst.Value.Struct, (uint8*)DestinationContainer))
 		{
 		case EConvertFromTypeResult::UseSerializeItem:
-			if (SourceProperty->GetID() == DestinationProperty->GetID())
+			if (bAllowNoConvert)
 			{
-				ReadTag.SerializeTaggedProperty(StructuredReader.GetSlot(), DestinationProperty, (uint8*)DestinationData, nullptr);
+				ReadTag.SerializeTaggedProperty(StructuredReader.GetSlot(), DestinationProperty, (uint8*)DestinationData, (uint8*)DestinationData);
 				bResult = true;
 			}
 			break;
@@ -118,6 +160,38 @@ namespace UE::Private
 			break;
 		}
 
+		if (FOverriddenPropertySet* OverriddenProperties = FOverridableManager::Get().GetOverriddenProperties(Dst.Object))
+		{
+			FOverridableSerializationLogic::Disable();
+		}
+
+		return bResult;
+	}
+
+	// TODO: this method allows dangerous and narrowing conversions. That's fine for now because the user explicitly asked to do so but in the future it could be good
+	// to include a flag in the RedirectProperty params that specifies weather the user is ok with narrowing conversions
+	bool TryCopyOrConvert(const FPropertyPathNameResolver& Src, const FPropertyPathNameResolver& Dst)
+	{
+		const FProperty* SourceProperty = Src.Value.Property;
+		const void* SourceContainer = Src.Value.Container;
+		int32 SourceIndex = Src.Value.ArrayIndex;
+		FProperty* DestinationProperty = const_cast<FProperty*>(Dst.Value.Property);
+		void* DestinationContainer = Dst.Value.Container;
+		int32 DestinationIndex = Dst.Value.ArrayIndex;
+
+		const void* SourceData = SourceProperty->ContainerPtrToValuePtr<void>(SourceContainer, SourceIndex == INDEX_NONE ? 0 : SourceIndex);
+		void* DestinationData = DestinationProperty->ContainerPtrToValuePtr<void>(DestinationContainer, DestinationIndex == INDEX_NONE ? 0 : DestinationIndex);
+
+		UObject* Owner = Dst.Object;
+		if (const UObject* Instance = UE::FPropertyBagRepository::Get().FindInstanceForDataObject(Dst.Object))
+		{
+			Owner = const_cast<UObject*>(Instance);
+		}
+
+		// attempt regular TPS conversion
+		TArray<uint8> Buffer = WriteToBuffer(Owner, Src);
+		bool bResult = ReadFromBuffer(Owner, Dst, Buffer, SourceProperty->GetID() == DestinationProperty->GetID());
+
 		if (!bResult)
 		{
 			bool bTryTextSerialize = false;
@@ -156,14 +230,21 @@ namespace UE::Private
 			}
 		}
 
-		return bResult;
-	}
+		if (bResult)
+		{
+			UE::FInitializedPropertyValueState(Dst.Value.Struct, DestinationContainer).Set(DestinationProperty, DestinationIndex);
+			for (const FPropertyValueInContainer& Value : Dst.ContainerValues)
+			{
+				if (Value.IsDynamicContainerElement())
+				{
+					// elements of dynamic containers (like TArray) don't store Initialized property value state
+					continue;
+				}
+				UE::FInitializedPropertyValueState(Value.Struct, Value.Container).Set(Value.Property, Value.ArrayIndex);
+			}
+		}
 
-	bool TryCopyOrConvert(const FPropertyValueInContainer& SourceValue, const FPropertyValueInContainer& DestValue)
-	{
-		return TryCopyOrConvert(
-			SourceValue.Property, SourceValue.Container, SourceValue.ArrayIndex,
-			const_cast<FProperty*>(DestValue.Property), DestValue.Container, DestValue.ArrayIndex);
+		return bResult;
 	}
 
 	bool InferMisingTypeInfo(FPropertyTypeName& Type, UStruct*& CurrentStruct)
@@ -296,6 +377,119 @@ namespace UE::Private
 		}
 		return false;
 	}
+
+	// When performing a copy from a struct to another struct we may need to create unknown properties in the new location. This occurs when either the struct types
+	// differ, or when the old struct has unknown properties that the new one doesn't. This method is called before ReinstanceWithCachedUnknownPropertyTree to grow
+	// the property tree to fit the desired unknown properties.
+	// 
+	// To determine which unknown properties need to be added to the tree, we'll set up a serialization pass that copies the struct into a temp location but turn on
+	// property tree collection flags in the serialization context. This will automatically grow the tree to fit the mismatched data.
+	void AnticipateUnknownPropertiesForTypeConversion(const TSharedPtr<FPropertyPathNameTree>& UnknownPropertyTree, const FPropertyPathNameResolver& SrcValue, FPropertyPathName DstPath)
+	{
+		const FPropertyTypeName& SrcType = SrcValue.Path.GetSegment(SrcValue.Path.GetSegmentCount() - 1).Type;
+		const FPropertyTypeName& DstType = DstPath.GetSegment(DstPath.GetSegmentCount() - 1).Type;
+
+		// for now this method only handles converting structs
+		if (SrcType.GetName() != NAME_StructProperty || DstType.GetName() != NAME_StructProperty)
+		{
+			return;
+		}
+
+		// if types don't match, we'll need to attempt serializing and see if any unknown properties were discovered
+		FProperty* SrcProperty = const_cast<FProperty*>(SrcValue.Value.Property);
+		const void* SourceData = SrcProperty->ContainerPtrToValuePtr<void>(SrcValue.Value.Container, SrcValue.Value.ArrayIndex == INDEX_NONE ? 0 : SrcValue.Value.ArrayIndex);
+		
+		UObject* Owner = SrcValue.Object;
+		if (const UObject* Instance = UE::FPropertyBagRepository::Get().FindInstanceForDataObject(SrcValue.Object))
+		{
+			Owner = const_cast<UObject*>(Instance);
+		}
+
+		// WRITE TO BUFFER
+		TArray<uint8> Buffer = WriteToBuffer(Owner, SrcValue);
+
+		// READ FROM BUFFER into a temp location
+		{
+			FObjectReader ObjectReader(Buffer);
+			FStructuredArchiveFromArchive StructuredReader(ObjectReader);
+			FScopedObjectSerializeContext ScopedSerializeContext(Owner, ObjectReader, DstPath);
+
+			FPropertyTag ReadTag;
+			StructuredReader.GetSlot() << ReadTag;
+			ReadTag.ArrayIndex = DstPath.GetSegment(DstPath.GetSegmentCount() - 1).Index;
+			ReadTag.Name = DstType.GetName();
+
+			TUniquePtr<FField> DstField(FField::TryConstruct(DstType.GetName(), nullptr, ReadTag.Name, RF_NoFlags));
+			FProperty* DstProperty = CastField<FProperty>(DstField.Get());
+			if (!(DstProperty && DstProperty->LoadTypeName(DstType, &ReadTag)))
+			{
+				return;
+			}
+			DstProperty->Link(ObjectReader);
+			ReadTag.SetProperty(DstProperty);
+			void* DestinationData = DstProperty->AllocateAndInitializeValue();
+
+			// because the ScopedSerializeContext has enabled unknown property tracking, the unknown property tree has now grown to include any needed unknown properties
+			ReadTag.SerializeTaggedProperty(StructuredReader.GetSlot(), DstProperty, (uint8*)DestinationData, nullptr);
+			DstProperty->DestroyAndFreeValue(DestinationData);
+		}
+	}
+
+	static const FName NAME_IsLooseMetadata(ANSITEXTVIEW("IsLoose"));
+
+	// Compare unknown property tree with the override set. If there's a branch that's not overridden but still in the tree, it should be pruned because
+	// it's no longer needed given that the non-overridden properties won't get saved.
+	void PruneUnknownPropertyTree(FPropertyPathNameTree* UnknownPropertyTree, const UStruct* Struct, const FOverriddenPropertyNode* Overrides)
+	{
+		TArray<FPropertyPathName> PathsToRemove;
+		auto PushPathToRemove = [&PathsToRemove](const FPropertyPathNameSegment& Segment) {
+			FPropertyPathName SubPath;
+			SubPath.Push(Segment);
+			PathsToRemove.Add(SubPath);
+			};
+
+		for (UE::FPropertyPathNameTree::FConstIterator Itr = UnknownPropertyTree->CreateConstIterator(); Itr; ++Itr)
+		{
+			const FProperty* Property = UE::FindPropertyByNameAndTypeName(Struct, Itr.GetName(), Itr.GetType());
+			if (!Property)
+			{
+				PushPathToRemove({ Itr.GetName(), Itr.GetType() });
+				continue;
+			}
+
+			// const cast needed because there's no non-const iterator. I'm being very careful to not iterate and modify at the same time so this is fine.
+			FPropertyPathNameTree* SubTree = const_cast<FPropertyPathNameTree*>(Itr.GetNode().GetSubTree());
+			const FStructProperty* AsStructProperty = CastField<FStructProperty>(Property);
+			const FOverriddenPropertyNode* FoundOverride = Overrides->GetSubPropertyNodes().FindByKey(FOverriddenPropertyNodeID(Property));
+			if (!FoundOverride)
+			{
+				PushPathToRemove({ Itr.GetName(), Itr.GetType() });
+				continue;
+			}
+
+			// recurse into struct properties with children
+			if (SubTree && !SubTree->IsEmpty() && AsStructProperty)
+			{
+				PruneUnknownPropertyTree(SubTree, AsStructProperty->Struct, FoundOverride);
+			}
+
+			// TODO: should we handle other containers here?
+
+			// prune overridden leaves if they aren't marked as loose
+			if (!SubTree || SubTree->IsEmpty())
+			{
+				if (!Property->GetBoolMetaData(NAME_IsLooseMetadata))
+				{
+					PushPathToRemove({ Itr.GetName(), Itr.GetType() });
+				}
+				continue;
+			}
+		}
+		for (FPropertyPathName& Path : PathsToRemove)
+		{
+			UnknownPropertyTree->Remove(Path);
+		}
+	}
 }
 
 #if WITH_EDITOR
@@ -338,6 +532,15 @@ namespace UE::Private::IDT
 			TObjectPtr<UObject> OldIDO = IDO;
 			IDO = UE::FPropertyBagRepository::Get().ReinstanceWithCachedUnknownPropertyTree(Owner);
 
+			// Cull nodes in the UnknownPropertyTree if they're not overridden
+			if (const FOverriddenPropertySet* OverriddenProperties = FOverridableManager::Get().GetOverriddenProperties(IDO))
+			{
+				if (const FOverriddenPropertyNode* OverrideTree = OverriddenProperties->GetOverriddenPropertyNode(nullptr))
+				{
+					PruneUnknownPropertyTree(UnknownPropertyTree.Get(), IDO->GetClass(), OverrideTree);
+				}
+			}
+
 #if WITH_EDITOR
 			FCoreUObjectDelegates::OnObjectsReinstanced.Broadcast({ {OldIDO, IDO} });
 #endif
@@ -393,17 +596,7 @@ namespace UE::Private::IDT
 				return ETransformResult::Continue;
 			}
 
-			// clear override(s) on souce value(s)
-			while (SrcValue)
-			{
-				TSharedPtr<FPropertyChangedChainEvent> ClearOverrideChangeEvent;
-				SrcValue.BuildChangeEvent(ClearOverrideChangeEvent, EPropertyChangeType::ResetToDefault);
-				UE::SendClearOverriddenPropertyEvent(IDO, *ClearOverrideChangeEvent, ClearOverrideChangeEvent->PropertyChain);
-				SrcValue.Next();
-			}
-			// reset source value to the beginning of the iteration list (needed in the case of wildcards)
-			SrcValue.Reset();
-			check(SrcValue);
+			AnticipateUnknownPropertiesForTypeConversion(UnknownPropertyTree, SrcValue, DstPath);
 
 			// create a new IDO with the modified property tree
 			TObjectPtr<UObject> OldIDO = IDO;
@@ -420,9 +613,13 @@ namespace UE::Private::IDT
 
 				// copy property from SrcValue to DstIDOValue (note that because we're using change events, the value will be changed in the instance as well)
 				IDO->PreEditChange(DstIdoChangeEvent->PropertyChain);
-				TryCopyOrConvert(SrcValue.Value, DstIDOValue.Value);
+				TryCopyOrConvert(SrcValue, DstIDOValue);
 				IDO->PostEditChangeChainProperty(*DstIdoChangeEvent);
 
+				TSharedPtr<FPropertyChangedChainEvent> ClearOverrideChangeEvent;
+				SrcValue.BuildChangeEvent(ClearOverrideChangeEvent, EPropertyChangeType::ResetToDefault);
+				UE::SendClearOverriddenPropertyEvent(IDO, *ClearOverrideChangeEvent, ClearOverrideChangeEvent->PropertyChain);
+
 				/*
 					because containers may contain more than one instance of a struct being redirected, we redirect them all by using wildcards in the path. 
 					FPropertyPathName wildcards are resolved by treating INDEX_NONE as "all indices of this container". FPropertyPathNameResolver will resove 
@@ -436,6 +633,15 @@ namespace UE::Private::IDT
 			ensureMsgf(!SrcValue && !DstIDOValue, TEXT("An IDT ran that operates on every element of a container but the source and destination had a " 
 				"different number of elements! This may have caused corruption in your data. Save at your own risk."));
 
+			// Cull nodes in the UnknownPropertyTree if they're not overridden
+			if (const FOverriddenPropertySet* OverriddenProperties = FOverridableManager::Get().GetOverriddenProperties(IDO))
+			{
+				if (const FOverriddenPropertyNode* OverrideTree = OverriddenProperties->GetOverriddenPropertyNode(nullptr))
+				{
+					PruneUnknownPropertyTree(UnknownPropertyTree.Get(), IDO->GetClass(), OverrideTree);
+				}
+			}
+
 #if WITH_EDITOR
 			FCoreUObjectDelegates::OnObjectsReinstanced.Broadcast({ {OldIDO, IDO} });
 #endif
diff --git a/Engine/Source/Runtime/CoreUObject/Private/UObject/PropertyBagRepository.cpp b/Engine/Source/Runtime/CoreUObject/Private/UObject/PropertyBagRepository.cpp
--- a/Engine/Source/Runtime/CoreUObject/Private/UObject/PropertyBagRepository.cpp
+++ b/Engine/Source/Runtime/CoreUObject/Private/UObject/PropertyBagRepository.cpp
@@ -827,60 +827,6 @@ const UObject* FPropertyBagRepository::FindInstanceForDataObject(const UObject*
 	return Owner ? *Owner : nullptr;
 }
 
-// Compare unknown property tree with the override set. If there's a branch that's not overridden but still in the tree, it should be pruned because
-// it's no longer needed given that the non-overridden properties won't get saved.
-static void PruneUnknownPropertyTree(FPropertyPathNameTree* UnknownPropertyTree, const UStruct* Struct, const FOverriddenPropertyNode* Overrides)
-{
-	TArray<FPropertyPathName> PathsToRemove;
-	auto PushPathToRemove = [&PathsToRemove](const FPropertyPathNameSegment& Segment) {
-		FPropertyPathName SubPath;
-		SubPath.Push(Segment);
-		PathsToRemove.Add(SubPath);
-	};
-
-	for (UE::FPropertyPathNameTree::FConstIterator Itr = UnknownPropertyTree->CreateConstIterator(); Itr; ++Itr)
-	{
-		const FProperty* Property = UE::FindPropertyByNameAndTypeName(Struct, Itr.GetName(), Itr.GetType());
-		if (!Property)
-		{
-			PushPathToRemove({ Itr.GetName(), Itr.GetType() });
-			continue;
-		}
-
-		// const cast needed because there's no non-const iterator. I'm being very careful to not iterate and modify at the same time so this is fine.
-		FPropertyPathNameTree* SubTree = const_cast<FPropertyPathNameTree*>(Itr.GetNode().GetSubTree());
-		const FStructProperty* AsStructProperty = CastField<FStructProperty>(Property);
-		const FOverriddenPropertyNode* FoundOverride = Overrides->GetSubPropertyNodes().FindByKey(FOverriddenPropertyNodeID(Property));
-		if (!FoundOverride)
-		{
-			PushPathToRemove({ Itr.GetName(), Itr.GetType() });
-			continue;
-		}
-
-		// recurse into struct properties with children
-		if (SubTree && !SubTree->IsEmpty() && AsStructProperty)
-		{
-			PruneUnknownPropertyTree(SubTree, AsStructProperty->Struct, FoundOverride);
-		}
-		
-		// TODO: should we handle other containers here?
-
-		// prune overridden leaves if they aren't marked as loose
-		if (!SubTree || SubTree->IsEmpty())
-		{
-			if (!Property->GetBoolMetaData(NAME_IsLooseMetadata))
-			{
-				PushPathToRemove({ Itr.GetName(), Itr.GetType() });
-			}
-			continue;
-		}
-	}
-	for (FPropertyPathName& Path : PathsToRemove)
-	{
-		UnknownPropertyTree->Remove(Path);
-	}
-}
-
 UObject* FPropertyBagRepository::ReinstanceWithCachedUnknownPropertyTree(UObject* Owner)
 {
 	if (FPropertyBagRepositoryLock LockRepo(this); FPropertyBagAssociationData* BagData = AssociatedData.Find(Owner))
@@ -919,15 +865,6 @@ UObject* FPropertyBagRepository::ReinstanceWithCachedUnknownPropertyTree(UObject
 
 		CopyTaggedProperties(IDO, NewIDO);
 
-		// Cull nodes in the UnknownPropertyTree if they're not overridden
-		if (const FOverriddenPropertySet* OverriddenProperties = FOverridableManager::Get().GetOverriddenProperties(NewIDO))
-		{
-			if (const FOverriddenPropertyNode* OverrideTree = OverriddenProperties->GetOverriddenPropertyNode(nullptr))
-			{
-				PruneUnknownPropertyTree(UnknownPropertyTree.Get(), NewIdoClass, OverrideTree);
-			}
-		}
-
 		// replace the old IDO with the new one in the repository
 		InstanceDataObjectToOwner.Remove(BagData->InstanceDataObject);
 		BagData->InstanceDataObject = NewIDO;
diff --git a/Engine/Source/Runtime/CoreUObject/Private/UObject/PropertyPathFunctions.cpp b/Engine/Source/Runtime/CoreUObject/Private/UObject/PropertyPathFunctions.cpp
--- a/Engine/Source/Runtime/CoreUObject/Private/UObject/PropertyPathFunctions.cpp
+++ b/Engine/Source/Runtime/CoreUObject/Private/UObject/PropertyPathFunctions.cpp
@@ -126,7 +126,7 @@ bool FPropertyPathNameResolver::BuildChangeEvent(TSharedPtr<FPropertyChangedChai
 
 void FPropertyPathNameResolver::Reset()
 {
-	ActiveContainerIterators.Empty();
+	ActiveContainerIterators.Reset();
 	Resolve();
 }
 
@@ -141,6 +141,7 @@ FPropertyPathNameResolver& FPropertyPathNameResolver::Resolve()
 		NextContainer = Object;
 		PathIndex = 0;
 		EventChain.Empty();
+		ContainerValues.Reset();
 	}
 	else
 	{
@@ -152,6 +153,12 @@ FPropertyPathNameResolver& FPropertyPathNameResolver::Resolve()
 		{
 			EventChain.RemoveNode(EventChain.GetTail());
 		}
+		int32 NewContainerValuesNum = ActiveContainerIterators.Last().ContainerValuesNum;
+		if (ContainerValues.Num() > NewContainerValuesNum)
+		{
+			ContainerValues.RemoveAt(NewContainerValuesNum, ContainerValues.Num() - NewContainerValuesNum, EAllowShrinking::No);
+		}
+		
 	}
 
 	for (const int32 SegmentCount = Path.GetSegmentCount(); PathIndex < SegmentCount; ++PathIndex)
@@ -185,11 +192,15 @@ FPropertyPathNameResolver& FPropertyPathNameResolver::Resolve()
 				else if (Property->ArrayDim > 1 || Property->IsA<FArrayProperty>() || Property->IsA<FSetProperty>() || Property->IsA<FMapProperty>())
 				{
 					Segment.Index = 0;
-					ActiveContainerIterators.Emplace(PathIndex, Property, NextContainer, NextStruct, EventChain.Num(), Segment.Index);
+					ActiveContainerIterators.Emplace(PathIndex, Property, NextContainer, NextStruct, EventChain.Num(), ContainerValues.Num(), Segment.Index);
 				}
 			}
 		}
 
+		if (Value)
+		{
+			ContainerValues.Push(Value);
+		}
 		Value.Property = Property;
 		Value.Struct = NextStruct;
 		Value.Container = NextContainer;
diff --git a/Engine/Source/Runtime/CoreUObject/Private/UObject/UObjectThreadContext.cpp b/Engine/Source/Runtime/CoreUObject/Private/UObject/UObjectThreadContext.cpp
--- a/Engine/Source/Runtime/CoreUObject/Private/UObject/UObjectThreadContext.cpp
+++ b/Engine/Source/Runtime/CoreUObject/Private/UObject/UObjectThreadContext.cpp
@@ -149,6 +149,7 @@ FScopedObjectSerializeContext::FScopedObjectSerializeContext(UObject* InObject,
 	bSavedTrackUnknownProperties = SerializeContext->bTrackUnknownProperties;
 	bSavedTrackUnknownEnumNames = SerializeContext->bTrackUnknownEnumNames;
 	bSavedImpersonateProperties = SerializeContext->bImpersonateProperties;
+	SavedSerializedPropertyPath = SerializeContext->SerializedPropertyPath;
 
 	SerializeContext->SerializedObjectScriptStartOffset = -1;
 	SerializeContext->SerializedObjectScriptEndOffset = -1;
@@ -203,7 +204,21 @@ FScopedObjectSerializeContext::~FScopedObjectSerializeContext()
 	SerializeContext->bTrackUnknownProperties = bSavedTrackUnknownProperties;
 	SerializeContext->bTrackUnknownEnumNames = bSavedTrackUnknownEnumNames;
 	SerializeContext->bImpersonateProperties = bSavedImpersonateProperties;
+	SerializeContext->SerializedPropertyPath = SavedSerializedPropertyPath;
 #endif
 }
 
+FScopedObjectSerializeContext::FScopedObjectSerializeContext(UObject*const InObject, FArchive& InArchive, const FPropertyPathName& PathBeingSerialized)
+	: FScopedObjectSerializeContext(InObject, InArchive)
+{
+#if WITH_EDITORONLY_DATA
+	const bool bIsLoading = Archive.IsLoading();
+	FUObjectSerializeContext* SerializeContext = FUObjectThreadContext::Get().GetSerializeContext();
+	SerializeContext->bTrackSerializedPropertyPath = bIsLoading;
+	SerializeContext->bTrackUnknownProperties = bIsLoading;
+	SerializeContext->bTrackUnknownEnumNames = bIsLoading;
+	SerializeContext->SerializedPropertyPath = PathBeingSerialized;
+#endif // WITH_EDITORONLY_DATA
+}
+
 } // UE
diff --git a/Engine/Source/Runtime/CoreUObject/Public/UObject/PropertyPathFunctions.h b/Engine/Source/Runtime/CoreUObject/Public/UObject/PropertyPathFunctions.h
--- a/Engine/Source/Runtime/CoreUObject/Public/UObject/PropertyPathFunctions.h
+++ b/Engine/Source/Runtime/CoreUObject/Public/UObject/PropertyPathFunctions.h
@@ -87,6 +87,10 @@ public:
 	UE_API explicit FPropertyPathNameResolver(const FPropertyPathName& Path, UObject* Object);
 
 	FPropertyValueInContainer Value;
+	const FPropertyPathName& Path;
+	UObject* const Object;
+	FEditPropertyChain EventChain;
+	TArray<FPropertyValueInContainer> ContainerValues; // all the values of containers
 
 	// returns true if the path was resolvable. In the case of a wildcarded path, returns false when all indices have been visited
 	inline explicit operator bool() const
@@ -98,11 +102,9 @@ public:
 	UE_API FPropertyPathNameResolver& Next();
 	UE_API bool BuildChangeEvent(TSharedPtr<FPropertyChangedChainEvent>& OutEvent, int32 ChangeType);
 	UE_API void Reset();
+
 private:
 	FPropertyPathNameResolver& Resolve();
-	const FPropertyPathName& Path;
-	UObject* Object;
-	FEditPropertyChain EventChain;
 	TArray<TMap<FString, int32>> EventArrayIndices;
 
 	struct FContainerIterator
@@ -112,6 +114,7 @@ private:
 		void* NextContainer;
 		const UStruct* NextStruct;
 		int32 ChainLength;
+		int32 ContainerValuesNum;
 		int32 CurIndex;
 	};
 
