Fix modular constinit uobjects build. Handling some configurations of FObjectPtr/FObjectHandle where FObjectHandle is not UObject* but UE_WITH_OBJECT_HANDLE_LATE_RESOLVE is 0

#Synced-CL Unknown

diff --git a/Engine/Source/Runtime/CoreUObject/Public/UObject/ObjectHandle.h b/Engine/Source/Runtime/CoreUObject/Public/UObject/ObjectHandle.h
--- a/Engine/Source/Runtime/CoreUObject/Public/UObject/ObjectHandle.h
+++ b/Engine/Source/Runtime/CoreUObject/Public/UObject/ObjectHandle.h
@@ -33,6 +33,7 @@ namespace UE::CoreUObject::Private
 }
 
 using FObjectHandle = UE::CoreUObject::Private::FObjectHandlePrivate;
+#define UE_OBJECT_HANDLE_IS_OBJECT_PTR 0
 
 #elif UE_WITH_REMOTE_OBJECT_HANDLE
 
@@ -42,11 +43,13 @@ namespace UE::CoreUObject::Private
 }
 
 using FObjectHandle = UE::CoreUObject::Private::FRemoteObjectHandlePrivate;
+#define UE_OBJECT_HANDLE_IS_OBJECT_PTR 0
 
 #else
 
 using FObjectHandle = UObject*;
 //NOTE: operator==, operator!=, GetTypeHash fall back to the default on UObject* or void* through coercion.
+#define UE_OBJECT_HANDLE_IS_OBJECT_PTR 1
 
 #endif
 
@@ -391,6 +394,16 @@ inline uint32 GetTypeHash(UE::CoreUObject::Private::FObjectHandlePrivate Handle)
 	}
 	return GetTypeHash(Handle.GetPointerOrRef());
 }
+#elif UE_WITH_CONSTINIT_UOBJECT && !IS_MONOLITHIC
+inline bool UE::CoreUObject::Private::FObjectHandlePrivate::UEOpEquals(FObjectHandlePrivate RHS) const
+{
+	return Pointer == RHS.Pointer;
+}
+
+inline uint32 GetTypeHash(UE::CoreUObject::Private::FObjectHandlePrivate Handle)
+{
+	return GetTypeHash(Handle.Pointer);
+}
 #elif UE_WITH_REMOTE_OBJECT_HANDLE
 inline bool UE::CoreUObject::Private::FRemoteObjectHandlePrivate::UEOpEquals(FRemoteObjectHandlePrivate RHS) const
 {
@@ -466,7 +479,7 @@ inline uint32 GetTypeHash(UE::CoreUObject::Private::FRemoteObjectHandlePrivate H
 
 namespace UE::CoreUObject::Private
 {
-#if UE_WITH_OBJECT_HANDLE_LATE_RESOLVE
+#if UE_WITH_OBJECT_HANDLE_LATE_RESOLVE 
 	inline constexpr FObjectHandle MakeObjectHandle(UObject* Object)
 	{
 		UE_IF_CONSTEVAL
@@ -486,7 +499,7 @@ namespace UE::CoreUObject::Private
 #else
 	inline constexpr FObjectHandle MakeObjectHandle(UObject* Object)
 	{
-		return Object;
+		return FObjectHandle(Object);
 	}
 #endif
 
@@ -668,12 +681,14 @@ namespace UE::CoreUObject::Private
 
 	inline UObject* ReadObjectHandlePointerNoCheck(FObjectHandle Handle)
 	{
-#if UE_WITH_OBJECT_HANDLE_LATE_RESOLVE
+#if UE_OBJECT_HANDLE_IS_OBJECT_PTR
+		return Handle;
+#elif UE_WITH_OBJECT_HANDLE_LATE_RESOLVE
 		return reinterpret_cast<UObject*>(Handle.GetPointerOrRef());
 #elif UE_WITH_REMOTE_OBJECT_HANDLE
 		return reinterpret_cast<UObject*>(Handle.PointerOrHandle);
 #else
-		return Handle;
+		return const_cast<UObject*>(reinterpret_cast<const UObject*>(Handle.Pointer));
 #endif
 	}
 
diff --git a/Engine/Source/Runtime/CoreUObject/Public/UObject/ObjectPtr.h b/Engine/Source/Runtime/CoreUObject/Public/UObject/ObjectPtr.h
--- a/Engine/Source/Runtime/CoreUObject/Public/UObject/ObjectPtr.h
+++ b/Engine/Source/Runtime/CoreUObject/Public/UObject/ObjectPtr.h
@@ -94,7 +94,7 @@ public:
 	}
 #endif // UE_WITH_CONSTINIT_UOBJECT
 
-#if UE_WITH_OBJECT_HANDLE_LATE_RESOLVE || UE_WITH_REMOTE_OBJECT_HANDLE
+#if !UE_OBJECT_HANDLE_IS_OBJECT_PTR
 	[[nodiscard]] explicit FORCEINLINE FObjectPtr(FObjectHandle Handle)
 		: Handle(Handle)
 	{
