Re-enabling import type hierarchy serialization.

#rnx
#Synced-CL 48348679

diff --git a/Engine/Source/Runtime/CoreUObject/Private/Serialization/ImportTypeHierarchy.cpp b/Engine/Source/Runtime/CoreUObject/Private/Serialization/ImportTypeHierarchy.cpp
--- a/Engine/Source/Runtime/CoreUObject/Private/Serialization/ImportTypeHierarchy.cpp
+++ b/Engine/Source/Runtime/CoreUObject/Private/Serialization/ImportTypeHierarchy.cpp
@@ -6,7 +6,7 @@
 namespace UE::Serialization::Private
 {
 
-static int32 GImportTypeHierarchyEnabled = 0;
+static int32 GImportTypeHierarchyEnabled = 1;
 static FAutoConsoleVariableRef CVarImportTypeHierarchyEnabled(
 	TEXT("s.ImportTypeHierarchyEnabled"),
 	GImportTypeHierarchyEnabled,
