Fix SelfReferences not being replaced regression when we added the PrepopulateInstancingGraphForArchetype

#rnx
#Synced-CL 48360443

diff --git a/Engine/Source/Editor/UnrealEd/Private/Kismet2/KismetReinstanceUtilities.cpp b/Engine/Source/Editor/UnrealEd/Private/Kismet2/KismetReinstanceUtilities.cpp
--- a/Engine/Source/Editor/UnrealEd/Private/Kismet2/KismetReinstanceUtilities.cpp
+++ b/Engine/Source/Editor/UnrealEd/Private/Kismet2/KismetReinstanceUtilities.cpp
@@ -3339,7 +3339,7 @@ bool FBlueprintCompileReinstancer::PrePopulateInstancingGraphForArchetype(FObjec
 
 	auto GatherMatchingSubObjects = [&InstancingGraph, OldInstance, OldToNewInstanceMap](UObject* Instance, UObject* Template, auto GatherMatchingSubObjectsRecurse) -> void
 	{
-		InstancingGraph.AddNewObject(Instance, Template);
+		InstancingGraph.PrepopulateWithInstance(Instance, Template);
 
 		ForEachObjectWithOuter(Instance, [OldInstance, Template, GatherMatchingSubObjectsRecurse, OldToNewInstanceMap](UObject* SubObject)
 		{
diff --git a/Engine/Source/Runtime/CoreUObject/Private/UObject/CoreNative.cpp b/Engine/Source/Runtime/CoreUObject/Private/UObject/CoreNative.cpp
--- a/Engine/Source/Runtime/CoreUObject/Private/UObject/CoreNative.cpp
+++ b/Engine/Source/Runtime/CoreUObject/Private/UObject/CoreNative.cpp
@@ -390,7 +390,8 @@ UObject* FObjectInstancingGraph::InstancePropertyValue(UObject* SubObjectTemplat
 		// Dynamic instancing means that we'll analyze the current value to determine how instancing should proceed. At
 		// construction time for example, the current value will have been initialized to the value from the default data
 		// object, so in that case we'll compare that value against the owner's archetype to see if it should be instanced.
-		if (Owner != DestinationRoot)
+		// As an optimization, we only need to do this even if the Owner is the same as destination root during reinstantiation.
+		if (Owner != DestinationRoot || GIsReinstancing)
 		{
 			// In this case, we're initializing an owner as a subgraph within the context of its outer destination root. If
 			// the owner archetype does not exist within the outer object's source hierarchy, we need to adjust the source
@@ -488,8 +489,7 @@ void FObjectInstancingGraph::AddNewObject(UObject* ObjectInstance, UObject* InAr
 
 void FObjectInstancingGraph::AddNewInstance(UObject* ObjectInstance, UObject* InArchetype /*= nullptr*/)
 {
-	check(SourceRoot);
-	check(DestinationRoot);
+	checkf(HasDestinationRoot(), TEXT("Expecting the destination root to be set before calling this method"));
 
 	if (ObjectInstance != nullptr)
 	{
@@ -501,6 +501,14 @@ void FObjectInstancingGraph::AddNewInstance(UObject* ObjectInstance, UObject* In
 	}
 }
 
+void FObjectInstancingGraph::PrepopulateWithInstance(TNotNull<UObject*> ObjectInstance, TNotNull<UObject*> InArchetype)
+{
+	checkf(!HasDestinationRoot(), TEXT("Expecting destination root to NOT be set before calling method"));
+
+	SourceToDestinationMap.Add(InArchetype, ObjectInstance);
+	DestinationToSourceMap.Add(ObjectInstance, InArchetype);
+}
+
 void FObjectInstancingGraph::RetrieveObjectInstances( UObject* SearchOuter, TArray<UObject*>& out_Objects )
 {
 	if ( HasDestinationRoot() && SearchOuter != nullptr && (SearchOuter == DestinationRoot || SearchOuter->IsIn(DestinationRoot)) )
diff --git a/Engine/Source/Runtime/CoreUObject/Public/UObject/Class.h b/Engine/Source/Runtime/CoreUObject/Public/UObject/Class.h
--- a/Engine/Source/Runtime/CoreUObject/Public/UObject/Class.h
+++ b/Engine/Source/Runtime/CoreUObject/Public/UObject/Class.h
@@ -5087,6 +5087,17 @@ public:
 	 */
 	COREUOBJECT_API void AddNewInstance(class UObject* ObjectInstance, class UObject* InArchetype = nullptr);
 
+	/**
+	 * Adds an object instance to the map of source objects to their instances of the outer chain already instanciated as part of prepopulation.
+	 * We use this to pre-create the instantiation graph when subobject are re-instantiated
+	 * If there is already a mapping for this object, it will be replaced
+	 * and the value corresponding to ObjectInstance's archetype will now point to ObjectInstance.
+	 *
+	 * @param	ObjectInstance  the object that should be added as the corresopnding instance for ObjectSource
+	 * @param	InArchetype     Archetype of ObjectInstance
+	 */
+	COREUOBJECT_API void PrepopulateWithInstance(TNotNull<UObject*> ObjectInstance, TNotNull<UObject*> InArchetype);
+
 	/**
 	 * Retrieves a list of objects that have the specified Outer
 	 *
