With UE_WITH_REMOTE_OBJECT_HANDLE enabled, when creating an object handle from an FRemoteObjectId, return null if the object was local, never migrated, and now destroyed. Prevents trying to migrate that object via a valid handle later, which would fail.

#Synced-CL Unknown

diff --git a/Engine/Source/Runtime/CoreUObject/Private/UObject/RemoteObject.cpp b/Engine/Source/Runtime/CoreUObject/Private/UObject/RemoteObject.cpp
--- a/Engine/Source/Runtime/CoreUObject/Private/UObject/RemoteObject.cpp
+++ b/Engine/Source/Runtime/CoreUObject/Private/UObject/RemoteObject.cpp
@@ -975,6 +975,17 @@ namespace UE::CoreUObject::Private
 			{
 				return FRemoteObjectHandlePrivate(Obj);
 			}
+			else if (const FRemoteObjectStub* const FoundStub = ObjectMaps->FindRemoteObjectStub(ObjectId))
+			{
+				return FRemoteObjectHandlePrivate(FoundStub);
+			}
+			else if (!Obj && ObjectId.IsLocal())
+			{
+				// If we don't have an existing stub or object, and the ID is local, we can assume
+				// the object was never migrated or referenced by another server and is now destroyed.
+				// Create a null handle instead of using the RemoteId because there's no remote object to find.
+				return FRemoteObjectHandlePrivate(Obj);
+			}
 			else if (FRemoteObjectStub* Stub = ObjectMaps->FindOrAddRemoteObjectStub(ObjectId))
 			{
 				return FRemoteObjectHandlePrivate(Stub);
