AssetRegistry: Loading from serialized or cached: Fix spurious converstion of UTF8Char to FName when comparing a string value to a cached tag; converting to FName is wasteful and also imposes a spurious limite of 1024 characters on the string being converted.
#rnx
#Synced-CL 45140074

diff --git a/Engine/Source/Runtime/CoreUObject/Private/AssetRegistry/AssetDataTagMap.cpp b/Engine/Source/Runtime/CoreUObject/Private/AssetRegistry/AssetDataTagMap.cpp
--- a/Engine/Source/Runtime/CoreUObject/Private/AssetRegistry/AssetDataTagMap.cpp
+++ b/Engine/Source/Runtime/CoreUObject/Private/AssetRegistry/AssetDataTagMap.cpp
@@ -244,6 +244,12 @@ namespace FixedTagPrivate
 		return FPlatformString::Strnicmp(A.GetData(), B, A.Len()) == 0 && B[A.Len()] == '\0';
 	}
 
+	static bool EqualsInsensitive(FStringView A, const UTF8CHAR* B)
+	{
+		auto BStringCast = StringCast<TCHAR>(B);
+		return A.Equals(FStringView(BStringCast.Get(), BStringCast.Length()), ESearchCase::IgnoreCase);
+	}
+
 	static bool EqualsInsensitive(const ANSICHAR* A, const ANSICHAR* B)
 	{
 		return FPlatformString::Stricmp(A, B) == 0;
@@ -254,6 +260,11 @@ namespace FixedTagPrivate
 		return FPlatformString::Stricmp(A, B) == 0;
 	}
 
+	static bool EqualsInsensitive(const UTF8CHAR* A, const UTF8CHAR* B)
+	{
+		return FPlatformString::Stricmp(A, B) == 0;
+	}
+
 	static bool EqualsInsensitive(const FMarshalledText& A, const FMarshalledText& B)
 	{
 		return A.CompareToCaseIgnored(B) == 0;
