- Moved a missing AssetRegistry module ensure from LinkerLoadPackageStore to HybridPackageStore which is the one that really wants the AssetRegistry

#Synced-CL 48456107

diff --git a/Engine/Source/Runtime/CoreUObject/Private/Serialization/PackageStore.cpp b/Engine/Source/Runtime/CoreUObject/Private/Serialization/PackageStore.cpp
--- a/Engine/Source/Runtime/CoreUObject/Private/Serialization/PackageStore.cpp
+++ b/Engine/Source/Runtime/CoreUObject/Private/Serialization/PackageStore.cpp
@@ -77,8 +77,7 @@ public:
 		// We need to defer registration for the AssetRegistry delegate since the registry might not be ready yet
 		FCoreDelegates::OnFEngineLoopInitComplete.AddLambda([this]()
 			{
-				IAssetRegistryInterface* AssetRegistry = IAssetRegistryInterface::GetPtr();
-				if (ensureMsgf(AssetRegistry, TEXT("It's expected we always have an AssetRegistry after Engine init.")))
+				if (IAssetRegistryInterface* AssetRegistry = IAssetRegistryInterface::GetPtr())
 				{
 					AssetRegistry->RegisterOnAssetsRemovedDelegate([this](TConstArrayView<FAssetData> RemovedAssets) 
 						{
@@ -245,9 +244,13 @@ EPackageStoreEntryStatus FHybridPackageStoreBackend::GetPackageStoreEntry(FPacka
 	{
 		// get the readonly flag (without slowly touching the disk for each package)
 		FAssetPackageData AssetPackageData;
-		UE::AssetRegistry::EExists Result = IAssetRegistryInterface::GetPtr()->TryGetAssetPackageData(PackageName, AssetPackageData);
-		
-		
+		UE::AssetRegistry::EExists Result = UE::AssetRegistry::EExists::Unknown;
+		IAssetRegistryInterface* AssetRegistry = IAssetRegistryInterface::GetPtr();
+		if (ensureMsgf(AssetRegistry, TEXT("It's expected we always have an AssetRegistry after Engine init when using Hybrid PackageStore")))
+		{
+			Result = AssetRegistry->TryGetAssetPackageData(PackageName, AssetPackageData);
+		}
+
 		// if we did find it in the AR, and it's writeable, assume user wants local changes, so load loose
 		//
 		// if we don't know the state yet, then instead of just assuming the worst and loading the loose version,
