Add log context to UPackage::Save

#Synced-CL Unknown

diff --git a/Engine/Source/Runtime/CoreUObject/Private/UObject/SavePackage2.cpp b/Engine/Source/Runtime/CoreUObject/Private/UObject/SavePackage2.cpp
--- a/Engine/Source/Runtime/CoreUObject/Private/UObject/SavePackage2.cpp
+++ b/Engine/Source/Runtime/CoreUObject/Private/UObject/SavePackage2.cpp
@@ -15,6 +15,7 @@
 #include "Internationalization/PackageLocalizationManager.h"
 #include "Interfaces/ITargetPlatform.h"
 #include "IO/IoDispatcher.h"
+#include "Logging/StructuredLog.h"
 #include "Misc/AssetRegistryInterface.h"
 #include "Misc/ConfigCacheIni.h"
 #include "Misc/FeedbackContext.h"
@@ -3990,6 +3991,7 @@ FSavePackageResultStruct UPackage::Save(UPackage* InPackage, UObject* InAsset, c
 	AutoRTFM::UnreachableIfClosed("TODO: Mark this function as AUTORTFM_DISABLED and fixup callers");
 
 	TRACE_CPUPROFILER_EVENT_SCOPE(UPackage::Save);
+	UE_LOG_CONTEXT("Package", FSoftObjectPath(InPackage));
 
 	COOK_STAT(FScopedDurationTimer FuncSaveTimer(FSavePackageStats::SavePackageTimeSec));
 	COOK_STAT(FSavePackageStats::NumPackagesSaved++);
