[UObjectGlobals]
- Fix issue where an object owned by the loader could be found for a short period before it was hidden

#Synced-CL Unknown

diff --git a/Engine/Source/Runtime/CoreUObject/Private/Serialization/AsyncLoading2.cpp b/Engine/Source/Runtime/CoreUObject/Private/Serialization/AsyncLoading2.cpp
--- a/Engine/Source/Runtime/CoreUObject/Private/Serialization/AsyncLoading2.cpp
+++ b/Engine/Source/Runtime/CoreUObject/Private/Serialization/AsyncLoading2.cpp
@@ -4081,6 +4081,9 @@ struct FAsyncPackage2 : public FRefCountedObject
 			if (ObjectsNowInPhase2)
 			{
 				Object->SetInternalFlags(EInternalObjectFlags::AsyncLoadingPhase2 | EInternalObjectFlags::Async);
+
+				// This must be cleared after Phase2 is set to avoid any window where the object can be erroneously found.
+				Object->ClearInternalFlags(EInternalObjectFlags::AsyncLoadingPhase1);
 			}
 			else
 			{
@@ -10648,7 +10651,7 @@ EAsyncPackageState::Type FAsyncLoadingThread2::ProcessLoadedPackagesFromGameThre
 				CompletedPackageRequests.Add(FCompletedPackageRequest::FromLoadedPackage(Package));
 			}
 		}
-		LocalCompletedAsyncPackages.Reset();
+			LocalCompletedAsyncPackages.Reset();
 
 		// Move CompletedPackageRequest out of the global collection to prevent it from changing from within the callbacks.
 		// If we're flushing a specific request only call callbacks for that request
@@ -11530,7 +11533,8 @@ void FAsyncLoadingThread2::NotifyConstructedDuringAsyncLoading(UObject* Object,
 	{
 		// Something is creating objects on the async loading thread outside of the actual async loading code
 		// e.g. ShaderCodeLibrary::OnExternalReadCallback doing FTaskGraphInterface::Get().WaitUntilTaskCompletes(Event);
-		// Or the object doesn't have a package, in which case we didn't load it from disk. Created objects not part of the package shoudl be cleaned up by their creator
+		// Or the object doesn't have a package, in which case we didn't load it from disk. Created objects not part of the package should be cleaned up by their creator
+		Object->AtomicallyClearInternalFlags(EInternalObjectFlags_AsyncLoading);
 		UE_LOG(LogStreaming, VeryVerbose, TEXT("Constructed object '%s' during loading will not be post-loaded due to its package '%s' (null and transient package sub-objects are ignored)."), *Object->GetFullName(), (ObjectPackage ? *ObjectPackage->GetFullName() : TEXT("null")));
 		return;
 	}
@@ -11583,12 +11587,12 @@ void FAsyncLoadingThread2::NotifyConstructedDuringAsyncLoading(UObject* Object,
 			LexToString(ContextAsyncPackage->AsyncPackageLoadingState));
 
 		ContextAsyncPackage->AddConstructedObject(Object, bSubObjectThatAlreadyExists);
+		return;
 	}
-	else
-	{
-		UE_LOG(LogStreaming, Error, TEXT("Object '%s' was created while the FUObjectThreadContext's async package '%s(0x%s)' is done post-loading objects. This should not happening (perhaps a FAsyncPackageScope2 is missing?) This object may be left in a partially loaded state and may leak."),
-			*Object->GetFullName(), *ContextAsyncPackage->Desc.UPackageName.ToString(), *LexToString(ContextAsyncPackage->Desc.UPackageId));
-	}
+
+	Object->AtomicallyClearInternalFlags(EInternalObjectFlags_AsyncLoading);
+	UE_LOG(LogStreaming, Error, TEXT("Object '%s' was created while the FUObjectThreadContext's async package '%s(0x%s)' is done post-loading objects. This should not happening (perhaps a FAsyncPackageScope2 is missing?) This object may be left in a partially loaded state and may leak."),
+		*Object->GetFullName(), *ContextAsyncPackage->Desc.UPackageName.ToString(), *LexToString(ContextAsyncPackage->Desc.UPackageId));
 }
 
 void FAsyncLoadingThread2::NotifyRegistrationEvent(
diff --git a/Engine/Source/Runtime/CoreUObject/Private/UObject/UObjectGlobals.cpp b/Engine/Source/Runtime/CoreUObject/Private/UObject/UObjectGlobals.cpp
--- a/Engine/Source/Runtime/CoreUObject/Private/UObject/UObjectGlobals.cpp
+++ b/Engine/Source/Runtime/CoreUObject/Private/UObject/UObjectGlobals.cpp
@@ -3718,8 +3718,17 @@ UObject* StaticAllocateObject
 		InFlags |= RF_Transient;
 	}
 
+	const bool bIsAsyncLoading = (IsInAsyncLoadingThread() || IsInParallelLoadingThread());
+	
 	if (!bSubObject)
 	{
+		// This needs to be set at construction to prevent the object from having a short period where it can be found
+		// between its creation and the call to NotifyConstructedDuringAsyncLoading.
+		if (bIsAsyncLoading)
+		{
+			InternalSetFlags |= EInternalObjectFlags::AsyncLoadingPhase1;
+		}
+
 #if UE_WITH_REMOTE_OBJECT_HANDLE
 		// We can't have two different, non-zero serial numbers
 		checkf(OldSerialNumber == 0 || SerialNumber == 0 || SerialNumber == OldSerialNumber, TEXT("Attempting to change a new object's (%s %s.%s) serial number (%d -> %d)"),
@@ -3773,7 +3782,7 @@ UObject* StaticAllocateObject
 		}
 	}
 
-	if (IsInAsyncLoadingThread() || IsInParallelLoadingThread())
+	if (bIsAsyncLoading)
 	{
 		FUObjectThreadContext& ThreadContext = FUObjectThreadContext::Get();
 		if (ThreadContext.AsyncPackageLoader)
