Disable AutoRTFM validation in `FBlueprintContext::GetThreadSingleton` as it is heavily hit and can't cause validation issues.

#Synced-CL 48501049

diff --git a/Engine/Source/Runtime/CoreUObject/Private/UObject/ScriptCore.cpp b/Engine/Source/Runtime/CoreUObject/Private/UObject/ScriptCore.cpp
--- a/Engine/Source/Runtime/CoreUObject/Private/UObject/ScriptCore.cpp
+++ b/Engine/Source/Runtime/CoreUObject/Private/UObject/ScriptCore.cpp
@@ -209,7 +209,10 @@ FBlueprintContext* FBlueprintContextGetThreadSingletonImpl()
 FBlueprintContext* FBlueprintContext::GetThreadSingleton()
 {
 	FBlueprintContext* Result;
-	UE_AUTORTFM_OPEN{ Result = FBlueprintContextGetThreadSingletonImpl(); };
+
+	// Running this open without validation because a) it is hit often and b) it is just
+	// initializing a thread-local variable which has no hazard.
+	UE_AUTORTFM_OPEN_NO_VALIDATION { Result = FBlueprintContextGetThreadSingletonImpl(); };
 	return Result;
 }
 
