[ZenLoader]
- Fix an issue where TryRemovePackage would fail a check because there was some deferred cleanups left
- Add missing ALT wakeup after queuing new deferred cleanups from the GT

#Synced-CL Unknown

diff --git a/Engine/Source/Runtime/CoreUObject/Private/Serialization/AsyncLoading2.cpp b/Engine/Source/Runtime/CoreUObject/Private/Serialization/AsyncLoading2.cpp
--- a/Engine/Source/Runtime/CoreUObject/Private/Serialization/AsyncLoading2.cpp
+++ b/Engine/Source/Runtime/CoreUObject/Private/Serialization/AsyncLoading2.cpp
@@ -2054,6 +2054,12 @@ public:
 	{
 		LLM_SCOPE_BYTAG(AsyncLoading_GlobalImportStore);
 
+		// This is currently required by at least 2 different code path marked as such below.
+		// To only flush on demand for these 2 code paths, we'd need to refactor this function so that it's possible to retry a second time in case of failure.
+		// If we do that refactor, we need to take very special care that we don't leave a new Package uninitialized for other threads to see during our retry.
+		// Since that's a non trivial refactor, the flush has been put back here for now to fix asserts caused when packages weren't cleaned fast enough.
+		FlushDeferredCleanupPackagesQueue();
+
 		FWriteLockedPackageRef PackageRef = FindOrAddPackageRef(PackageId);
 
 		// Detect package that have been renamed but are still referenced and attempt to get rid of the references
@@ -2064,13 +2070,7 @@ public:
 			{
 				if (PackageRef->GetOriginalPackageName() != Package->GetFName())
 				{
-					// Avoid deadlock during the cleanup, we'll reacquire after.
-					PackageRef.Reset();
-
-					FlushDeferredCleanupPackagesQueue();
-
-					PackageRef = FindOrAddPackageRef(PackageId);
-
+					// This path requires that FlushDeferredCleanupPackagesQueue has been run to get rid of all the pending references to packages
 					// Flush failed to clear the remaining references, warn since we want to know about this.
 					if (PackageRef->GetPackageRefCount() > 0)
 					{
@@ -2108,6 +2108,7 @@ public:
 					FPackageId OldPackageId = FoundPackage->GetPackageId();
 					if (OldPackageId.IsValid() && OldPackageId != PackageId)
 					{
+						// This path requires that FlushDeferredCleanupPackagesQueue has been run to get rid of all the pending references to packages
 						if (FWriteLockedPackageRef RemovedPackageRef = TryRemovePackage(OldPackageId))
 						{
 							UE_LOG(LogStreaming, Log,
@@ -5215,8 +5216,8 @@ public:
 		bool bDidSomething = false;
 		if (!DeferredCleanupPackages.IsEmpty())
 		{
-			UE::TUniqueLock ScopeLock(DeferredCleanupPackagesLock);
 			TRACE_CPUPROFILER_EVENT_SCOPE(ProcessDeferredCleanupPackagesQueue);
+			UE::TUniqueLock ScopeLock(DeferredCleanupPackagesLock);
 
 			FAsyncPackage2* Package = nullptr;
 			int32 Count = 0;
@@ -10651,7 +10652,14 @@ EAsyncPackageState::Type FAsyncLoadingThread2::ProcessLoadedPackagesFromGameThre
 				CompletedPackageRequests.Add(FCompletedPackageRequest::FromLoadedPackage(Package));
 			}
 		}
+
+		if (LocalCompletedAsyncPackages.Num())
+		{
+			// Notify ALT for processing DeferredCleanupPackages
+			AltZenaphore.NotifyOne();
+
 			LocalCompletedAsyncPackages.Reset();
+		}
 
 		// Move CompletedPackageRequest out of the global collection to prevent it from changing from within the callbacks.
 		// If we're flushing a specific request only call callbacks for that request
