Add core redirectors for enum value to the class schema hash.

#rnx
#Synced-CL 46456044

diff --git a/Engine/Source/Runtime/CoreUObject/Private/UObject/EnumProperty.cpp b/Engine/Source/Runtime/CoreUObject/Private/UObject/EnumProperty.cpp
--- a/Engine/Source/Runtime/CoreUObject/Private/UObject/EnumProperty.cpp
+++ b/Engine/Source/Runtime/CoreUObject/Private/UObject/EnumProperty.cpp
@@ -9,6 +9,7 @@
 #include "UObject/UObjectThreadContext.h"
 
 #if WITH_EDITORONLY_DATA
+#include "UObject/CoreRedirects.h"
 #include "UObject/PropertyStateTracking.h"
 #endif
 
@@ -591,6 +592,25 @@ void FEnumProperty::AppendSchemaHash(FBlake3& Builder, bool bSkipEditorOnly) con
 		{
 			AppendHash(Builder, Enum->GetNameByIndex(Index));
 		}
+
+		// All redirectors that should trigger a recook are handled either by schema hash, package digest, package redirector dependencies or 
+		// global core redirect dependencies. But there is one case not handled : native package redirectors mapping an enum value to another. 
+		// To detect when an enum value is redirected to another value, I add the FCoreRedirect to the schema hash. This will trigger
+		// a recook when an enum property value might be changed because of a redirector.
+		// It might trigger useless recooks as well since the redirection could be about a value not used by the current object but it's
+		// better than not detecting those cases at all.
+		// Adding enum value redirectors to the schema hash will also change the package digest of any package containing an object with this enum. 
+		// This is also a case I want to detect. For example : package A depends on package B and package B has an object with an enum property. 
+		// If the enum property value is redirected, then it could matter for package A. The change of package digest for B will trigger a recook of A.
+		FCoreRedirectObjectName OldName(*NameBuilder);
+		FCoreRedirectObjectName NewName;
+		const FCoreRedirect* Redirect = nullptr;
+		FCoreRedirects::RedirectNameAndValues(ECoreRedirectFlags::Type_Enum, OldName, NewName, &Redirect);
+
+		if (Redirect)
+		{
+			Redirect->AppendHash(Builder);
+		}
 	}
 }
 #endif
