[CoreUObject]
* Fixed non unity compile error

#Synced-CL 48544215

diff --git a/Engine/Source/Runtime/CoreUObject/Private/Cooker/ConsoleVariableData.h b/Engine/Source/Runtime/CoreUObject/Private/Cooker/ConsoleVariableData.h
--- a/Engine/Source/Runtime/CoreUObject/Private/Cooker/ConsoleVariableData.h
+++ b/Engine/Source/Runtime/CoreUObject/Private/Cooker/ConsoleVariableData.h
@@ -4,7 +4,7 @@
 
 #if WITH_EDITOR
 
-#include "Containers/StringFwd.h"
+#include "Containers/UnrealString.h"
 #include "Serialization/CompactBinary.h"
 
 class FCbWriter;
