Fixed TWeakObjectPtr being used as dynamic delegate parameter.  They need to be mapped via a UObject* pointer as this is how they are stored and retrieved from the script VM.

diff --git a/Engine/Source/Runtime/CoreUObject/Public/UObject/ScriptMacros.h b/Engine/Source/Runtime/CoreUObject/Public/UObject/ScriptMacros.h
--- a/Engine/Source/Runtime/CoreUObject/Public/UObject/ScriptMacros.h
+++ b/Engine/Source/Runtime/CoreUObject/Public/UObject/ScriptMacros.h
@@ -40,6 +40,10 @@ enum {MAX_VARIABLE_SIZE = 0x0FFF };
 	UE_REMOVE_OPTIONAL_PARENS(ParamType) ParamName{__VA_ARGS__};																						\
 	Stack.StepCompiledIn<PropertyType>(&ParamName);
 
+#define PARAM_PASSED_BY_REF_PASSTHRU(ParamName, PropertyType, ParamType, PassThruType)																	\
+	UE_REMOVE_OPTIONAL_PARENS(PassThruType) ParamName##Temp;																							\
+	UE_REMOVE_OPTIONAL_PARENS(ParamType)& ParamName = Stack.StepCompiledInRef<PropertyType, UE_REMOVE_OPTIONAL_PARENS(ParamType)>(&ParamName##Temp);
+
 #define PARAM_PASSED_BY_REF(ParamName, PropertyType, ParamType)																							\
 	UE_REMOVE_OPTIONAL_PARENS(ParamType) ParamName##Temp;																								\
 	UE_REMOVE_OPTIONAL_PARENS(ParamType)& ParamName = Stack.StepCompiledInRef<PropertyType, UE_REMOVE_OPTIONAL_PARENS(ParamType)>(&ParamName##Temp);
@@ -86,11 +90,11 @@ enum {MAX_VARIABLE_SIZE = 0x0FFF };
 #define P_GET_TINTERFACE(ObjectType,ParamName)		PARAM_PASSED_BY_VAL(ParamName, FInterfaceProperty, TScriptInterface<ObjectType>)
 #define P_GET_TINTERFACE_REF(ObjectType,ParamName)	PARAM_PASSED_BY_REF(ParamName, FInterfaceProperty, TScriptInterface<ObjectType>)
 
-#define P_GET_WEAKOBJECT(ObjectType,ParamName)		PARAM_PASSED_BY_VAL(ParamName, FWeakObjectProperty, ObjectType)
-#define P_GET_WEAKOBJECT_REF(ObjectType,ParamName)	PARAM_PASSED_BY_REF(ParamName, FWeakObjectProperty, ObjectType)
+#define P_GET_WEAKOBJECT(ObjectType,ParamName)		PARAM_PASSED_BY_VAL(ParamName, FWeakObjectProperty, typename UE_REMOVE_OPTIONAL_PARENS(ObjectType)::ElementType*)
+#define P_GET_WEAKOBJECT_REF(ObjectType,ParamName)	PARAM_PASSED_BY_REF_PASSTHRU(ParamName, FWeakObjectProperty, ObjectType, typename UE_REMOVE_OPTIONAL_PARENS(ObjectType)::ElementType*)
 
-#define P_GET_WEAKOBJECT_NO_PTR(ObjectType,ParamName)		PARAM_PASSED_BY_VAL(ParamName, FWeakObjectProperty, ObjectType)
-#define P_GET_WEAKOBJECT_REF_NO_PTR(ObjectType,ParamName)	PARAM_PASSED_BY_REF(ParamName, FWeakObjectProperty, ObjectType)
+#define P_GET_WEAKOBJECT_NO_PTR(ObjectType,ParamName)		PARAM_PASSED_BY_VAL(ParamName, FWeakObjectProperty, typename UE_REMOVE_OPTIONAL_PARENS(ObjectType)::ElementType*)
+#define P_GET_WEAKOBJECT_REF_NO_PTR(ObjectType,ParamName)	PARAM_PASSED_BY_REF_PASSTHRU(ParamName, FWeakObjectProperty, ObjectType, typename UE_REMOVE_OPTIONAL_PARENS(ObjectType)::ElementType*)
 
 #define P_GET_SOFTOBJECT(ObjectType,ParamName)		PARAM_PASSED_BY_VAL(ParamName, FSoftObjectProperty, ObjectType)
 #define P_GET_SOFTOBJECT_REF(ObjectType,ParamName)	PARAM_PASSED_BY_REF(ParamName, FSoftObjectProperty, ObjectType)
diff --git a/Engine/Source/Runtime/CoreUObject/Tests/WeakObjectPtrTest.cpp b/Engine/Source/Runtime/CoreUObject/Tests/WeakObjectPtrTest.cpp
new file mode 100644
--- /dev/null
+++ b/Engine/Source/Runtime/CoreUObject/Tests/WeakObjectPtrTest.cpp
@@ -0,0 +1,85 @@
+// Copyright Epic Games, Inc. All Rights Reserved.
+
+#include "WeakObjectPtrTest.h"
+
+#if WITH_LOW_LEVEL_TESTS
+#include "TestHarness.h"
+#include "TestMacros/Assertions.h"
+#endif
+
+UWeakObjectPtrTest::UWeakObjectPtrTest()
+{
+	InputRefWeakPtrDelegate.BindDynamic(this, &UWeakObjectPtrTest::InputRefWeakPtr);
+	InputValWeakPtrDelegate.BindDynamic(this, &UWeakObjectPtrTest::InputValWeakPtr);
+	InputRefWeakPtrMulticastDelegate.AddDynamic(this, &UWeakObjectPtrTest::InputRefWeakPtr);
+	InputValWeakPtrMulticastDelegate.AddDynamic(this, &UWeakObjectPtrTest::InputValWeakPtr);
+	InOutWeakPtrDelegate.BindDynamic(this, &UWeakObjectPtrTest::InOutWeakPtr);
+	InOutWeakPtrMulticastDelegate.AddDynamic(this, &UWeakObjectPtrTest::InOutWeakPtr);
+	ReturnWeakPtrDelegate.BindDynamic(this, &UWeakObjectPtrTest::ReturnWeakPtr);
+}
+
+void UWeakObjectPtrTest::InputRefWeakPtr(const TWeakObjectPtr<UObject>& InWeakPtr)
+{
+#if WITH_LOW_LEVEL_TESTS
+	CHECK(InWeakPtr == this);
+#endif
+}
+
+void UWeakObjectPtrTest::InputValWeakPtr(TWeakObjectPtr<UObject> InWeakPtr)
+{
+#if WITH_LOW_LEVEL_TESTS
+	CHECK(InWeakPtr == this);
+#endif
+}
+
+void UWeakObjectPtrTest::InOutWeakPtr(TWeakObjectPtr<UObject>& InOutWeakPtr)
+{
+#if WITH_LOW_LEVEL_TESTS
+	CHECK(InOutWeakPtr == nullptr);
+#endif
+	InOutWeakPtr = this;
+}
+
+TWeakObjectPtr<UObject> UWeakObjectPtrTest::ReturnWeakPtr()
+{
+	return this;
+}
+
+#if WITH_LOW_LEVEL_TESTS
+
+void UWeakObjectPtrTest::RunTests()
+{
+	// Check that a weakobjptr can be passed through a dynamic delegate
+	TWeakObjectPtr<UObject> InRefPtr = this;
+	InputRefWeakPtrDelegate.Execute(InRefPtr);
+	TWeakObjectPtr<UObject> InValPtr = this;
+	InputValWeakPtrDelegate.Execute(InValPtr);
+
+	// Check that a weakobjptr can be passed through a dynamic multicast delegate
+	TWeakObjectPtr<UObject> InRefPtrMulticast = this;
+	InputRefWeakPtrMulticastDelegate.Broadcast(InRefPtrMulticast);
+	TWeakObjectPtr<UObject> InValPtrMulticast = this;
+	InputValWeakPtrMulticastDelegate.Broadcast(InValPtrMulticast);
+
+	// Check that a weakobjptr can be passed in and out of a dynamic delegate
+	TWeakObjectPtr<UObject> InOutPtr;
+	InOutWeakPtrDelegate.Execute(InOutPtr);
+	CHECK(InOutPtr == this);
+
+	// Check that a weakobjptr can be passed in and out of a multicast dynamic delegate
+	TWeakObjectPtr<UObject> InOutPtrMulticast;
+	InOutWeakPtrMulticastDelegate.Broadcast(InOutPtrMulticast);
+	CHECK(InOutPtrMulticast == this);
+
+	// Check that a weakobjptr can be returned by a dynamic delegate
+	TWeakObjectPtr<UObject> OutPtr = ReturnWeakPtrDelegate.Execute();
+	CHECK(OutPtr == this);
+}
+
+TEST_CASE("UE::CoreUObject::WeakObjectPtrTest", "[CoreUObject][WeakObjectPtrTest]")
+{
+	UWeakObjectPtrTest* TestObj = NewObject<UWeakObjectPtrTest>();
+	TestObj->RunTests();
+}
+
+#endif // WITH_LOW_LEVEL_TESTS
diff --git a/Engine/Source/Runtime/CoreUObject/Tests/WeakObjectPtrTest.h b/Engine/Source/Runtime/CoreUObject/Tests/WeakObjectPtrTest.h
new file mode 100644
--- /dev/null
+++ b/Engine/Source/Runtime/CoreUObject/Tests/WeakObjectPtrTest.h
@@ -0,0 +1,61 @@
+// Copyright Epic Games, Inc. All Rights Reserved.
+
+#pragma once
+
+#include "UObject/Object.h"
+#include "UObject/ScriptMacros.h"
+#include "UObject/WeakObjectPtrTemplates.h"
+
+#include "WeakObjectPtrTest.generated.h"
+
+DECLARE_DYNAMIC_DELEGATE_OneParam(FInputRefWeakPtrDelegate, const TWeakObjectPtr<UObject>&, WeakPtr);
+DECLARE_DYNAMIC_DELEGATE_OneParam(FInputValWeakPtrDelegate, TWeakObjectPtr<UObject>, WeakPtr);
+DECLARE_DYNAMIC_MULTICAST_DELEGATE_OneParam(FInputRefWeakPtrMulticastDelegate, const TWeakObjectPtr<UObject>&, WeakPtr);
+DECLARE_DYNAMIC_MULTICAST_DELEGATE_OneParam(FInputValWeakPtrMulticastDelegate, TWeakObjectPtr<UObject>, WeakPtr);
+DECLARE_DYNAMIC_DELEGATE_OneParam(FInOutWeakPtrDelegate, TWeakObjectPtr<UObject>&, WeakPtr);
+DECLARE_DYNAMIC_MULTICAST_DELEGATE_OneParam(FInOutWeakPtrMulticastDelegate, TWeakObjectPtr<UObject>&, WeakPtr);
+DECLARE_DYNAMIC_DELEGATE_RetVal(TWeakObjectPtr<UObject>, FReturnWeakPtrDelegate);
+
+UCLASS()
+class UWeakObjectPtrTest : public UObject
+{
+	GENERATED_BODY()
+public:
+
+	UWeakObjectPtrTest();
+
+	UPROPERTY()
+	FInputRefWeakPtrDelegate InputRefWeakPtrDelegate;
+
+	UPROPERTY()
+	FInputValWeakPtrDelegate InputValWeakPtrDelegate;
+
+	UPROPERTY()
+	FInputRefWeakPtrMulticastDelegate InputRefWeakPtrMulticastDelegate;
+
+	UPROPERTY()
+	FInputValWeakPtrMulticastDelegate InputValWeakPtrMulticastDelegate;
+
+	UPROPERTY()
+	FInOutWeakPtrDelegate InOutWeakPtrDelegate;
+
+	UPROPERTY()
+	FInOutWeakPtrMulticastDelegate InOutWeakPtrMulticastDelegate;
+
+	UPROPERTY()
+	FReturnWeakPtrDelegate ReturnWeakPtrDelegate;
+
+	UFUNCTION()
+	void InputRefWeakPtr(const TWeakObjectPtr<UObject>& InWeakPtr);
+
+	UFUNCTION()
+	void InputValWeakPtr(TWeakObjectPtr<UObject> InWeakPtr);
+
+	UFUNCTION()
+	void InOutWeakPtr(TWeakObjectPtr<UObject>& InOutWeakPtr);
+
+	UFUNCTION()
+	TWeakObjectPtr<UObject> ReturnWeakPtr();
+
+	void RunTests();
+};
