expose FPerPlatformBool to blueprints to match Int and Float
#Synced-CL 48601003

diff --git a/Engine/Source/Runtime/CoreUObject/Public/UObject/PerPlatformProperties.h b/Engine/Source/Runtime/CoreUObject/Public/UObject/PerPlatformProperties.h
--- a/Engine/Source/Runtime/CoreUObject/Public/UObject/PerPlatformProperties.h
+++ b/Engine/Source/Runtime/CoreUObject/Public/UObject/PerPlatformProperties.h
@@ -376,7 +376,7 @@ struct TStructOpsTypeTraits<FPerPlatformFloat>
 };
 
 /** FPerPlatformBool - bool property with per-platform overrides */
-USTRUCT()
+USTRUCT(BlueprintType)
 struct FPerPlatformBool
 #if CPP
 :	public TPerPlatformProperty<FPerPlatformBool, bool, NAME_BoolProperty>
@@ -384,7 +384,7 @@ struct FPerPlatformBool
 {
 	GENERATED_USTRUCT_BODY()
 
-	UPROPERTY(EditAnywhere, Category = PerPlatform)
+	UPROPERTY(BlueprintReadOnly, EditAnywhere, Category = PerPlatform)
 	bool Default;
 
 #if WITH_EDITORONLY_DATA
