[AutoRTFM] Apply AUTORTFM_INFER / UnreachableIfClosed() to fix some debug Verse test builds

#Synced-CL 48603520

diff --git a/Engine/Source/Runtime/CoreUObject/Public/UObject/FastReferenceCollector.h b/Engine/Source/Runtime/CoreUObject/Public/UObject/FastReferenceCollector.h
--- a/Engine/Source/Runtime/CoreUObject/Public/UObject/FastReferenceCollector.h
+++ b/Engine/Source/Runtime/CoreUObject/Public/UObject/FastReferenceCollector.h
@@ -684,7 +684,7 @@ FORCEINLINE EOrigin GetSchemaOrigin(FSchemaView Schema, UObject* Instance)
 #endif // UE_WITH_REMOTE_OBJECT_HANDLE
 
 template<class DispatcherType, typename ObjectType>
-FORCEINLINE_DEBUGGABLE void VisitMembers(DispatcherType& Dispatcher, FSchemaView Schema, ObjectType* Instance)
+AUTORTFM_INFER FORCEINLINE_DEBUGGABLE void VisitMembers(DispatcherType& Dispatcher, FSchemaView Schema, ObjectType* Instance)
 {
 	check(!Schema.IsEmpty());
 
