Add checkSlow in fast reference collector to confirm classes have had their GC schema created

#Synced-CL Unknown

diff --git a/Engine/Source/Runtime/CoreUObject/Public/UObject/FastReferenceCollector.h b/Engine/Source/Runtime/CoreUObject/Public/UObject/FastReferenceCollector.h
--- a/Engine/Source/Runtime/CoreUObject/Public/UObject/FastReferenceCollector.h
+++ b/Engine/Source/Runtime/CoreUObject/Public/UObject/FastReferenceCollector.h
@@ -1011,6 +1011,8 @@ private:
 				Class->AssembleReferenceTokenStream();
 			}
 
+			checkSlow(Class->HasAllClassFlags(CLASS_TokenStreamAssembled));
+
 			FSchemaView Schema = Class->ReferenceSchema.Get();
 			Dispatcher.Context.ReferencingObject = CurrentObject;
 
