Fix using FField::StaticClass accessor during static init when FieldClass is not initialized - take address of ConstInitClass directly

#Synced-CL Unknown

diff --git a/Engine/Source/Runtime/CoreUObject/Public/UObject/Field.h b/Engine/Source/Runtime/CoreUObject/Public/UObject/Field.h
--- a/Engine/Source/Runtime/CoreUObject/Public/UObject/Field.h
+++ b/Engine/Source/Runtime/CoreUObject/Public/UObject/Field.h
@@ -259,7 +259,7 @@ protected:
 		}
 		else
 		{
-			return FieldType::Super::StaticClass();
+			return &FieldType::Super::ConstInitClass;
 		}
 	}
 };
