Fix it so MarkSearchableName does not affect the imports of the Game harvest scope as it is only ever saved/read in editor builds.
#Synced-CL 48530611

diff --git a/Engine/Source/Runtime/CoreUObject/Private/UObject/SavePackage/PackageHarvester.cpp b/Engine/Source/Runtime/CoreUObject/Private/UObject/SavePackage/PackageHarvester.cpp
--- a/Engine/Source/Runtime/CoreUObject/Private/UObject/SavePackage/PackageHarvester.cpp
+++ b/Engine/Source/Runtime/CoreUObject/Private/UObject/SavePackage/PackageHarvester.cpp
@@ -1298,6 +1298,9 @@ void FPackageHarvester::HarvestPackageHeaderName(FName Name)
 
 void FPackageHarvester::HarvestSearchableName(TObjectPtr<UObject> TypeObject, FName Name)
 {
+	// Searchable names are only ever saved into the editor scope, so don't harvest the type object for game saves
+	FHarvestScope EditorOnlyScope = EnterConditionalEditorOnlyScope(true);
+	
 	// Make sure the object is tracked as a dependency
 	bool bAlreadyTrackedInAllHarvestingRealms = true;
 	ForEachExportHarvestingRealm([this, &bAlreadyTrackedInAllHarvestingRealms, TypeObject](ESaveRealm HarvestingRealm)
