[UE] Ensure that transient references to child objects in an instancing subgraph are cloned when the subgraph's root type supports dynamic subobject instancing.

#Synced-CL 48418225

diff --git a/Engine/Source/Runtime/CoreUObject/Private/UObject/CoreNative.cpp b/Engine/Source/Runtime/CoreUObject/Private/UObject/CoreNative.cpp
--- a/Engine/Source/Runtime/CoreUObject/Private/UObject/CoreNative.cpp
+++ b/Engine/Source/Runtime/CoreUObject/Private/UObject/CoreNative.cpp
@@ -430,6 +430,15 @@ UObject* FObjectInstancingGraph::InstancePropertyValue(UObject* SubObjectTemplat
 			// current owner, or return a reference to the instance that was already created for the source archetype.
 			bCausesInstancing = true;
 		}
+		else if (CurrentValue->HasAllFlags(RF_ArchetypeObject)
+				 && SubObjectTemplate && CurrentValue->GetOuter() == SubObjectTemplate->GetOuter()->GetClass()->GetDefaultObject())
+		{
+			// In this case, the source may have been instanced from class default data (e.g. a transient reference).
+			// It is similar to the case above in that if the instancing graph does not yet contain a mapping for this
+			// archetype, we'll allow instancing to proceed and construct a new instance under the current owner. But
+			// in this scenario, the current value is referencing an archetype that's based outside of its source graph.
+			bCausesInstancing = true;
+		}
 	}
 
 	if (!IsSubobjectInstancingEnabled() || // if instancing is off
