Add additional info to the assert checking that bulkdata streaming request buffers match the read request.

#rnx

#Synced-CL 48498703

diff --git a/Engine/Source/Runtime/CoreUObject/Private/Serialization/BulkDataStreaming.cpp b/Engine/Source/Runtime/CoreUObject/Private/Serialization/BulkDataStreaming.cpp
--- a/Engine/Source/Runtime/CoreUObject/Private/Serialization/BulkDataStreaming.cpp
+++ b/Engine/Source/Runtime/CoreUObject/Private/Serialization/BulkDataStreaming.cpp
@@ -863,7 +863,7 @@ FBulkDataBatchRequest::FBatchBuilder& FBulkDataBatchRequest::FBatchBuilder::Read
 	const uint64 ReadOffset = BulkData.GetBulkDataOffsetInFile() + Offset;
 	const uint64 ReadSize	= FMath::Min(uint64(BulkData.GetBulkDataSize()), Size);
 
-	check(Dst.GetSize() == 0 || Dst.GetSize() == ReadSize);
+	checkf(Dst.GetSize() == 0 || Dst.GetSize() == ReadSize, TEXT("Buffer size (%ull) does not match the read size (%ull)"), Dst.GetSize(), Size);
 
 	if (Dst.GetSize() == 0)
 	{
