[ZenLoader]
- Fix optional data that was not postloaded because of logic error

#Synced-CL Unknown

diff --git a/Engine/Source/Runtime/CoreUObject/Private/Serialization/AsyncLoading2.cpp b/Engine/Source/Runtime/CoreUObject/Private/Serialization/AsyncLoading2.cpp
--- a/Engine/Source/Runtime/CoreUObject/Private/Serialization/AsyncLoading2.cpp
+++ b/Engine/Source/Runtime/CoreUObject/Private/Serialization/AsyncLoading2.cpp
@@ -10063,6 +10063,7 @@ EEventLoadNodeExecutionResult FAsyncPackage2::Event_PostLoad(FAsyncLoadingThread
 				break;
 			}
 
+			Package->ExportBundleEntryIndex = 0;
 			Package->ProcessedExportBundlesCount++;
 		}
 
@@ -10196,7 +10197,8 @@ EEventLoadNodeExecutionResult FAsyncPackage2::Event_DeferredPostLoad(FAsyncLoadi
 				break;
 			}
 
-			++Package->ProcessedExportBundlesCount;
+			Package->ExportBundleEntryIndex = 0;
+			Package->ProcessedExportBundlesCount++;
 		}
 	}
 	else
