[UE] Added ground truth/unit tests/regression tests for various transient instanced subobject reference use cases.

Notes:
- This is a followup change to add a regression test for UE-354820, with some additional ground truth testing for a few other common types of transient instancing use cases.

#rnx

#Synced-CL 48418225

diff --git a/Engine/Source/Runtime/CoreUObject/Tests/UObject/SubobjectInstancingTest.cpp b/Engine/Source/Runtime/CoreUObject/Tests/UObject/SubobjectInstancingTest.cpp
--- a/Engine/Source/Runtime/CoreUObject/Tests/UObject/SubobjectInstancingTest.cpp
+++ b/Engine/Source/Runtime/CoreUObject/Tests/UObject/SubobjectInstancingTest.cpp
@@ -39,6 +39,12 @@ USubobjectInstancingTestOuterObject::USubobjectInstancingTestOuterObject(const F
 	// Subobject that may be instanced at edit time rather than at construction time; this is NOT considered a default subobject
 	EditTimeInnerObject = nullptr;
 
+	// Same as above, but marked transient; this is not expected to be serialized on save
+	TransientEditTimeInnerObject = nullptr;
+
+	// Same as above, but marked duplicatetransient; this is not expected to be serialized as part of duplicating the outer object
+	DuplicateTransientEditTimeInnerObject = nullptr;
+
 	// Subobject that does not inherit values from the CDO/template when instanced; 'transient' here does not affect serialization
 	// Note: This only applies to C++ types; otherwise it behaves the same as 'InnerObject' above and 'bTransient' is not supported
 	LocalOnlyInnerObject = CreateDefaultSubobject<USubobjectInstancingTestObject>(TEXT("LocalOnlyInnerObject"), /*bTransient =*/ true);
@@ -51,6 +57,12 @@ USubobjectInstancingTestOuterObject::USubobjectInstancingTestOuterObject(const F
 	// Property flags will exclude RF_DefaultSubObject, and runtime behavior is similar to LocalOnlyInnerObject above (i.e. no inheritance)
 	InnerObjectUsingNew = NewObject<USubobjectInstancingTestObject>(this, TEXT("InnerObjectUsingNew"));
 
+	// Transient reference to an instanced subobject that is constructed via NewObject() instead of CreateDefaultSubobject() (not serialized)
+	TransientInnerObjectUsingNew = NewObject<USubobjectInstancingTestObject>(this, TEXT("TransientInnerObjectUsingNew"));
+
+	// Instanced subobject constructed via NewObject() that should not be serialized during duplication, but is non-transient during save
+	DuplicateTransientInnerObjectUsingNew = NewObject<USubobjectInstancingTestObject>(this, TEXT("DuplicateTransientInnerObjectUsingNew"));
+
 	// Array container of references to instanced subobjects
 	InnerObjectArray.Add(CreateDefaultSubobject<USubobjectInstancingTestObject>(TEXT("InnerObjectArrayElem_0")));
 	InnerObjectArray.Add(CreateDefaultSubobject<USubobjectInstancingTestObject>(TEXT("InnerObjectArrayElem_1")));
@@ -107,6 +119,9 @@ void USubobjectInstancingTestOuterObject::PostInitProperties()
 
 	// Subobject that's deferred from native construction, but still instanced as part of UObject initialization flow
 	InnerObjectPostInit = NewObject<USubobjectInstancingTestObject>(this, TEXT("InnerObjectPostInit"));
+
+	// Same as above, but assigned to a transient instanced property
+	TransientInnerObjectPostInit = NewObject<USubobjectInstancingTestObject>(this, TEXT("TransientInnerObjectPostInit"));
 }
 
 USubobjectInstancingTestDirectlyNestedObject::USubobjectInstancingTestDirectlyNestedObject()
@@ -150,7 +165,9 @@ namespace SubobjectInstancingTest::Private
 	static const FName NonNativeInnerObjectPropertyName("NonNativeInnerObject");
 	static const FName NonNativeOuterObjectPropertyName("NonNativeOwnerObject");
 	static const FName NonNativeEditTimeInnerObjectPropertyName("NonNativeEditTimeInnerObject");
+	static const FName NonNativeTransientInnerObjectPropertyName("NonNativeTransientInnerObject");
 	static const FName NonNullableInnerObjectPropertyName("NonNullableInnerObject");
+	static const FName NonNullableTransientInnerObjectPropertyName("NonNullableTransientInnerObject");
 	static const FName NonNullableOptionalInnerObjectPropertyName("OptionalNonNullableInnerObject");
 
 	template<typename ClassType>
@@ -209,6 +226,12 @@ namespace SubobjectInstancingTest::Private
 		NonNativeEditTimeInnerObjectProperty->SetPropertyClass(USubobjectInstancingTestObject::StaticClass());
 		PropertyListBuilder.AppendNoTerminate(*NonNativeEditTimeInnerObjectProperty);
 
+		// Non-native object reference property to store a transient reference to a non-native instantiation - TObjectPtr<USubobjectInstancingTestObject>
+		FObjectProperty* NonNativeTransientInnerObjectProperty = new FObjectProperty(TestClass, NonNativeTransientInnerObjectPropertyName, RF_Public | RF_Transient);
+		NonNativeTransientInnerObjectProperty->SetPropertyFlags(NonNativeObjectPropertyFlags | CPF_Transient);
+		NonNativeTransientInnerObjectProperty->SetPropertyClass(USubobjectInstancingTestObject::StaticClass());
+		PropertyListBuilder.AppendNoTerminate(*NonNativeTransientInnerObjectProperty);
+
 		// Non-nullable object reference property to store a non-native instantiation by ctor - TNonNullPtr<USubobjectInstancingTestObject>
 		EPropertyFlags NonNullableObjectPropertyFlags = (NonNativeObjectPropertyFlags & ~CPF_TObjectPtrWrapper) | CPF_NonNullable;
 		FObjectProperty* NonNullableInnerObjectProperty = new FObjectProperty(TestClass, NonNullableInnerObjectPropertyName, RF_Public | RF_Transient);
@@ -216,6 +239,12 @@ namespace SubobjectInstancingTest::Private
 		NonNullableInnerObjectProperty->SetPropertyClass(USubobjectInstancingTestObject::StaticClass());
 		PropertyListBuilder.AppendNoTerminate(*NonNullableInnerObjectProperty);
 
+		// Non-nullable transient reference property to store a non-native instantiation by ctor - TNonNullPtr<USubobjectInstancingTestObject>
+		FObjectProperty* NonNullableTransientInnerObjectProperty = new FObjectProperty(TestClass, NonNullableTransientInnerObjectPropertyName, RF_Public | RF_Transient);
+		NonNullableTransientInnerObjectProperty->SetPropertyFlags(NonNullableObjectPropertyFlags | CPF_Transient);
+		NonNullableTransientInnerObjectProperty->SetPropertyClass(USubobjectInstancingTestObject::StaticClass());
+		PropertyListBuilder.AppendNoTerminate(*NonNullableTransientInnerObjectProperty);
+
 		// Non-nullable optional object reference property to store an optional non-null instantiation by ctor - TOptional<TNonNullPtr<USubobjectinstancingTestObject>>
 		// Note: TOptional<TNonNullPtr<T>> is not currently supported as a native property, so this is being used to validate intrusive state for non-nullable references.
 		FOptionalProperty* NonNullableOptionalInnerObjectProperty = new FOptionalProperty(TestClass, NonNullableOptionalInnerObjectPropertyName, RF_Public | RF_Transient);
@@ -243,6 +272,11 @@ namespace SubobjectInstancingTest::Private
 		UObject* NonNativeInnerDefaultSubobject = NewObject<USubobjectInstancingTestObject>(TestCDO, NonNativeInnerObjectName, RF_ArchetypeObject | TestCDO->GetMaskedFlags(RF_PropagateToSubObjects));
 		NonNativeInnerObjectProperty->SetObjectPropertyValue_InContainer(TestCDO, NonNativeInnerDefaultSubobject);
 
+		// Same as above, but assigned to a non-native transient reference property
+		static const FName NonNativeTransientInnerObjectName("NonNativeTransientInnerObject");
+		UObject* NonNativeTransientInnerDefaultSubobject = NewObject<USubobjectInstancingTestObject>(TestCDO, NonNativeTransientInnerObjectName, RF_ArchetypeObject | TestCDO->GetMaskedFlags(RF_PropagateToSubObjects));
+		NonNativeTransientInnerObjectProperty->SetObjectPropertyValue_InContainer(TestCDO, NonNativeTransientInnerDefaultSubobject);
+
 		// Simulate a non-native class constructor initialization of the outer object property's default value
 		NonNativeOuterObjectProperty->SetObjectPropertyValue_InContainer(TestCDO, nullptr);
 
@@ -254,6 +288,11 @@ namespace SubobjectInstancingTest::Private
 		UObject* NonNullableInnerDefaultSubobject = NewObject<USubobjectInstancingTestObject>(TestCDO, NonNullableInnerObjectName, RF_ArchetypeObject | TestCDO->GetMaskedFlags(RF_PropagateToSubObjects));
 		NonNullableInnerObjectProperty->SetObjectPropertyValue_InContainer(TestCDO, NonNullableInnerDefaultSubobject);
 
+		// Same as above, but assigned to a non-nullable transient reference property
+		static const FName NonNullableTransientInnerObjectName("NonNullableTransientInnerObject");
+		UObject* NonNullableTransientInnerDefaultSubobject = NewObject<USubobjectInstancingTestObject>(TestCDO, NonNullableTransientInnerObjectName, RF_ArchetypeObject | TestCDO->GetMaskedFlags(RF_PropagateToSubObjects));
+		NonNullableTransientInnerObjectProperty->SetObjectPropertyValue_InContainer(TestCDO, NonNullableTransientInnerDefaultSubobject);
+
 		return TestClass;
 	}
 
@@ -565,6 +604,7 @@ namespace SubobjectInstancingTest::Private
 			{
 				// BP/non-native types behave differently and will not currently search for an existing subobject at instancing time
 				// before reconstructing the reference using the template object - see FObjectInstancingGraph::GetInstancedSubobject().
+				// @todo - Maybe we can fix this at some point, but this needs to remain backwards-compatible with existing projects.
 				CHECK(OuterObject->InnerObjectUsingNew->TestValue == 200);
 			}
 			else
@@ -584,6 +624,58 @@ namespace SubobjectInstancingTest::Private
 			USubobjectInstancingTestOuterObject* NewOuterObject = NewObject<USubobjectInstancingTestOuterObject>(GetTransientPackage(), InstancingTestClass, NAME_None, RF_NoFlags, Template);
 			CHECK(OuterObject->InnerObjectUsingNew != NewOuterObject->InnerObjectUsingNew);
 		}
+
+		SECTION("Transient native property initialized from CDO using NewObject()")
+		{
+			USubobjectInstancingTestOuterObject* OuterObject = NewObject<USubobjectInstancingTestOuterObject>(GetTransientPackage(), InstancingTestClass);
+
+			// Transient instanced properties are excluded from default value initialization - see FObjectInitializer::InitProperties().
+			// Note that for BP/non-native types, this behavior differs from the non-transient case, where the reference is reinstanced.
+			REQUIRE(OuterObject->TransientInnerObjectUsingNew != nullptr);
+			CHECK(OuterObject->TransientInnerObjectUsingNew->IsInOuter(OuterObject));
+
+			CHECK(OuterObject->TransientInnerObjectUsingNew->TestValue == 100);
+
+			CHECK_FALSE(OuterObject->TransientInnerObjectUsingNew->HasAllFlags(RF_DefaultSubObject));
+			CHECK(OuterObject->TransientInnerObjectUsingNew->HasAllFlags(OuterObject->GetMaskedFlags(RF_PropagateToSubObjects)));
+
+			CHECK_FALSE(OuterObject->TransientInnerObjectUsingNew->IsTemplate());
+			CHECK(OuterObject->TransientInnerObjectUsingNew->IsDefaultSubobject());
+			CHECK_FALSE(OuterObject->TransientInnerObjectUsingNew->IsTemplateForSubobjects());
+
+			USubobjectInstancingTestOuterObject* NewOuterObject = NewObject<USubobjectInstancingTestOuterObject>(GetTransientPackage(), InstancingTestClass);
+			CHECK(OuterObject->TransientInnerObjectUsingNew != NewOuterObject->TransientInnerObjectUsingNew);
+		}
+
+		SECTION("Transient native property initialized from template using NewObject()")
+		{
+			USubobjectInstancingTestOuterObject* Template = NewObject<USubobjectInstancingTestOuterObject>(GetTransientPackage(), InstancingTestClass, NAME_None, RF_ArchetypeObject);
+			Template->TransientInnerObjectUsingNew->TestValue = 200;
+
+			USubobjectInstancingTestOuterObject* OuterObject = NewObject<USubobjectInstancingTestOuterObject>(GetTransientPackage(), InstancingTestClass, NAME_None, RF_NoFlags, Template);
+
+			REQUIRE(OuterObject->TransientInnerObjectUsingNew != nullptr);
+			CHECK(OuterObject->TransientInnerObjectUsingNew->IsInOuter(OuterObject));
+
+			// Note: Because we did not construct using CreateDefaultSubobject(), instance data does not propagate through the template.
+			// As with the non-transient case, using NewObject() is effectively the same result as passing TRUE for the transient argument
+			// to CreateDefaultSubobject().
+			// 
+			// However, unlike the non-transient case, this is true even for BP/non-native types, because transient instanced properties
+			// are explicitly excluded from default/template initialization - see FObjectInitializer::InitProperties(). This means we will
+			// not run through the same path in FObjectInstancingGraph::GetInstancedSubobject() as we otherwise do in the non-transient case.
+			CHECK(OuterObject->TransientInnerObjectUsingNew->TestValue == 100);
+
+			CHECK_FALSE(OuterObject->TransientInnerObjectUsingNew->HasAllFlags(RF_DefaultSubObject));
+			CHECK(OuterObject->TransientInnerObjectUsingNew->HasAllFlags(OuterObject->GetMaskedFlags(RF_PropagateToSubObjects)));
+
+			CHECK_FALSE(OuterObject->TransientInnerObjectUsingNew->IsTemplate());
+			CHECK(OuterObject->TransientInnerObjectUsingNew->IsDefaultSubobject());
+			CHECK_FALSE(OuterObject->TransientInnerObjectUsingNew->IsTemplateForSubobjects());
+
+			USubobjectInstancingTestOuterObject* NewOuterObject = NewObject<USubobjectInstancingTestOuterObject>(GetTransientPackage(), InstancingTestClass, NAME_None, RF_NoFlags, Template);
+			CHECK(OuterObject->TransientInnerObjectUsingNew != NewOuterObject->TransientInnerObjectUsingNew);
+		}
 	}
 
 	void RunSubobjectInstancingTests_Array(UClass* InstancingTestClass)
@@ -1398,7 +1490,7 @@ namespace SubobjectInstancingTest::Private
 		{
 			if (InstancingTestClass->IsNative())
 			{
-				// Skip: This test only applies to non-native outer types.
+				// Skip: This test only applies to non-native outer types. Native types do not currently support dynamic subobject instancing.
 			}
 			else
 			{
@@ -1423,13 +1515,36 @@ namespace SubobjectInstancingTest::Private
 				USubobjectInstancingTestOuterObject* OuterObject = NewObject<USubobjectInstancingTestOuterObject>(GetTransientPackage(), InstancingTestClass, NAME_None, RF_NoFlags, Template);
 				USubobjectInstancingTestDirectlyNestedObject* NonNativeEditTimeInnerObject = CastChecked<USubobjectInstancingTestDirectlyNestedObject>(ObjectProperty->GetObjectPropertyValue_InContainer(OuterObject));
 
+				REQUIRE(NonNativeEditTimeInnerObject != nullptr);
 				CHECK(NonNativeEditTimeInnerObject->TestValue == 200);
+				CHECK(NonNativeEditTimeInnerObject->IsInOuter(OuterObject));
+
+				// Deeply-nested reference to a dynamically-instanced subobject, expected to be duplicated as part of constructing the outermost root object initialized using the template.
+				UObject* NonNativeEditTimeInnerObject_NonNativeInnerObject = CastFieldChecked<FObjectProperty>(TestClass->FindPropertyByName(Private::NonNativeInnerObjectPropertyName))->GetObjectPropertyValue_InContainer(NonNativeEditTimeInnerObject);
+				{
+					REQUIRE(NonNativeEditTimeInnerObject_NonNativeInnerObject != nullptr);
+					CHECK(NonNativeEditTimeInnerObject_NonNativeInnerObject->IsInOuter(NonNativeEditTimeInnerObject));
+				}
 
-				CHECK(CastFieldChecked<FObjectProperty>(TestClass->FindPropertyByName(Private::NonNativeInnerObjectPropertyName))->GetObjectPropertyValue_InContainer(NonNativeEditTimeInnerObject) != nullptr);
-				CHECK(CastFieldChecked<FObjectProperty>(TestClass->FindPropertyByName(Private::NonNativeInnerObjectPropertyName))->GetObjectPropertyValue_InContainer(NonNativeEditTimeInnerObject)->IsInOuter(OuterObject));
+				// Deeply-nested transient reference to a dynamically-instanced subobject, expected to be duplicated as part of constructing the outermost root object initialized using the template.
+				// Note: Unlike the case above, this would not be expected to be serialized as an export on save.
+				UObject* NonNativeEditTimeInnerObject_NonNativeTransientInnerObject = CastFieldChecked<FObjectProperty>(TestClass->FindPropertyByName(Private::NonNativeTransientInnerObjectPropertyName))->GetObjectPropertyValue_InContainer(NonNativeEditTimeInnerObject);
+				{
+					REQUIRE(NonNativeEditTimeInnerObject_NonNativeTransientInnerObject != nullptr);
+					CHECK(NonNativeEditTimeInnerObject_NonNativeTransientInnerObject->IsInOuter(NonNativeEditTimeInnerObject));
+				}
 
-				CHECK(CastFieldChecked<FObjectProperty>(TestClass->FindPropertyByName(Private::NonNativeOuterObjectPropertyName))->GetObjectPropertyValue_InContainer(NonNativeEditTimeInnerObject) == OuterObject);
-				CHECK(CastFieldChecked<FObjectProperty>(TestClass->FindPropertyByName(Private::NonNativeSelfReferencePropertyName))->GetObjectPropertyValue_InContainer(NonNativeEditTimeInnerObject) == NonNativeEditTimeInnerObject);
+				// Nested reference to the outermost instance; ensures that we can self-reference nodes that exist above the immediate parent node in the object graph for nested subobject types that use dynamic instancing.
+				UObject* NonNativeEditTimeInnerObject_NonNativeOuterObject = CastFieldChecked<FObjectProperty>(TestClass->FindPropertyByName(Private::NonNativeOuterObjectPropertyName))->GetObjectPropertyValue_InContainer(NonNativeEditTimeInnerObject);
+				{
+					CHECK(NonNativeEditTimeInnerObject_NonNativeOuterObject == OuterObject);
+				}
+
+				// Nested reference to self; ensures that we can self-reference the immediate parent node in the object graph for nested subobject types that use dynamic instancing.
+				UObject* NonNativeEditTimeInnerObject_NonNativeSelfReference = CastFieldChecked<FObjectProperty>(TestClass->FindPropertyByName(Private::NonNativeSelfReferencePropertyName))->GetObjectPropertyValue_InContainer(NonNativeEditTimeInnerObject);
+				{
+					CHECK(NonNativeEditTimeInnerObject_NonNativeSelfReference == NonNativeEditTimeInnerObject);
+				}
 
 				CHECK_FALSE(NonNativeEditTimeInnerObject->HasAllFlags(RF_DefaultSubObject));
 				CHECK(NonNativeEditTimeInnerObject->HasAllFlags(OuterObject->GetMaskedFlags(RF_PropagateToSubObjects)));
@@ -1608,6 +1723,66 @@ namespace SubobjectInstancingTest::Private
 			}
 		}
 
+		SECTION("Simulate native transient property instanced at edit time on non-native CDO")
+		{
+			if (InstancingTestClass->IsNative())
+			{
+				// Skip: This test only applies to non-native types.
+			}
+			else
+			{
+				USubobjectInstancingTestOuterObject* CDO = Cast<USubobjectInstancingTestOuterObject>(InstancingTestClass->GetDefaultObject());
+
+				// Default value before it's edited (as will be initialized by the native super class ctor).
+				CHECK(CDO->TransientEditTimeInnerObject == nullptr);
+
+				// Simulate a deferred edit-time instancing event along with a value override on the non-native CDO.
+				CDO->TransientEditTimeInnerObject = NewObject<USubobjectInstancingTestObject>(CDO, GET_MEMBER_NAME_CHECKED(USubobjectInstancingTestOuterObject, TransientEditTimeInnerObject), CDO->GetMaskedFlags(RF_PropagateToSubObjects));
+				CDO->TransientEditTimeInnerObject->TestValue = 200;
+
+				CHECK(CDO->TransientEditTimeInnerObject->IsTemplate());
+				CHECK(CDO->TransientEditTimeInnerObject->IsDefaultSubobject());
+				CHECK(CDO->TransientEditTimeInnerObject->IsTemplateForSubobjects());
+
+				USubobjectInstancingTestOuterObject* OuterObject = NewObject<USubobjectInstancingTestOuterObject>(GetTransientPackage(), InstancingTestClass);
+
+				// Since the subobject archetype was not instanced at construction time, it's expected to remain set to NULL. When the CDO is used as the default object for initialization, as an optimization,
+				// the native object initializer will not include properties inherited from the native super class hierarchy, because it expects they have already been initialized by the native ctor. For this
+				// reason, users cannot override a native instanced reference member's default value. This test is validating that any such override will not actually be used by subobject template instancing.
+				CHECK(OuterObject->TransientEditTimeInnerObject == nullptr);
+			}
+		}
+
+		SECTION("Simulate native transient property instanced at edit time on non-native template")
+		{
+			if (InstancingTestClass->IsNative())
+			{
+				// Skip: This test only applies to non-native types.
+			}
+			else
+			{
+				USubobjectInstancingTestOuterObject* Template = NewObject<USubobjectInstancingTestOuterObject>(GetTransientPackage(), InstancingTestClass, NAME_None, RF_ArchetypeObject);
+				CHECK(Template->TransientEditTimeInnerObject == nullptr);
+
+				// Simulate a deferred edit-time instancing event along with a value override on the non-native template.
+				Template->TransientEditTimeInnerObject = NewObject<USubobjectInstancingTestObject>(Template, GET_MEMBER_NAME_CHECKED(USubobjectInstancingTestOuterObject, TransientEditTimeInnerObject), Template->GetMaskedFlags(RF_PropagateToSubObjects));
+				Template->TransientEditTimeInnerObject->TestValue = 200;
+
+				CHECK(Template->TransientEditTimeInnerObject->IsTemplate());
+				CHECK_FALSE(Template->TransientEditTimeInnerObject->IsDefaultSubobject());
+				// Since this is not considered a default subobject nor was it captured by a CDO, it's expected that this call should fail w/ the default flags argument on the internal IsTemplate() call.
+				CHECK_FALSE(Template->TransientEditTimeInnerObject->IsTemplateForSubobjects());
+				// However, we can configure the API to tell us whether the subobject is also an archetype that will be used for instancing at construction time if the outer template is used for initialization.
+				CHECK(Template->TransientEditTimeInnerObject->IsTemplateForSubobjects(RF_ArchetypeObject));
+
+				USubobjectInstancingTestOuterObject* OuterObject = NewObject<USubobjectInstancingTestOuterObject>(GetTransientPackage(), InstancingTestClass, NAME_None, RF_NoFlags, Template);
+
+				// As with the TransientInnerObjectUsingNew case, transient instanced properties are explicitly excluded from default/template initialization - see FObjectInitializer::InitProperties().
+				// This means we will not run through the same path in FObjectInstancingGraph::GetInstancedSubobject() as we otherwise do in the non-transient case where the property is initialized from template.
+				CHECK(OuterObject->TransientEditTimeInnerObject == nullptr);
+			}
+		}
+
 		SECTION("Simulate non-native property instanced at edit time on non-native CDO")
 		{
 			if (InstancingTestClass->IsNative())
@@ -1725,6 +1900,31 @@ namespace SubobjectInstancingTest::Private
 			USubobjectInstancingTestOuterObject* NewOuterObject = NewObject<USubobjectInstancingTestOuterObject>(GetTransientPackage(), InstancingTestClass, NAME_None, RF_NoFlags, Template);
 			CHECK(OuterObject->InnerObjectPostInit != NewOuterObject->InnerObjectPostInit);
 		}
+
+		SECTION("Native transient property initialized from template with deferred construction using NewObject()")
+		{
+			USubobjectInstancingTestOuterObject* Template = NewObject<USubobjectInstancingTestOuterObject>(GetTransientPackage(), InstancingTestClass, NAME_None, RF_ArchetypeObject);
+			Template->TransientInnerObjectPostInit->TestValue = 200;
+
+			USubobjectInstancingTestOuterObject* OuterObject = NewObject<USubobjectInstancingTestOuterObject>(GetTransientPackage(), InstancingTestClass, NAME_None, RF_NoFlags, Template);
+
+			REQUIRE(OuterObject->TransientInnerObjectPostInit != nullptr);
+			CHECK(OuterObject->TransientInnerObjectPostInit->IsInOuter(OuterObject));
+
+			// As with the TransientInnerObjectUsingNew case, transient instanced properties are excluded from default value initialization.
+			// This means it is also effectively the same result as passing TRUE for the transient argument to CreateDefaultSubobject().
+			CHECK(OuterObject->TransientInnerObjectPostInit->TestValue == 100);
+
+			CHECK_FALSE(OuterObject->TransientInnerObjectPostInit->HasAllFlags(RF_DefaultSubObject));
+			CHECK(OuterObject->TransientInnerObjectPostInit->HasAllFlags(OuterObject->GetMaskedFlags(RF_PropagateToSubObjects)));
+
+			CHECK_FALSE(OuterObject->TransientInnerObjectPostInit->IsTemplate());
+			CHECK(OuterObject->TransientInnerObjectPostInit->IsDefaultSubobject());
+			CHECK_FALSE(OuterObject->TransientInnerObjectPostInit->IsTemplateForSubobjects());
+
+			USubobjectInstancingTestOuterObject* NewOuterObject = NewObject<USubobjectInstancingTestOuterObject>(GetTransientPackage(), InstancingTestClass, NAME_None, RF_NoFlags, Template);
+			CHECK(OuterObject->TransientInnerObjectPostInit != NewOuterObject->TransientInnerObjectPostInit);
+		}
 	}
 
 	void RunSubobjectInstancingTests_NonNullable(UClass* InstancingTestClass)
@@ -1736,6 +1936,7 @@ namespace SubobjectInstancingTest::Private
 		}
 
 		const FObjectProperty* ObjectProperty = CastFieldChecked<FObjectProperty>(InstancingTestClass->FindPropertyByName(NonNullableInnerObjectPropertyName));
+		const FObjectProperty* TransientProperty = CastFieldChecked<FObjectProperty>(InstancingTestClass->FindPropertyByName(NonNullableTransientInnerObjectPropertyName));
 		const FOptionalProperty* OptionalProperty = CastFieldChecked<FOptionalProperty>(InstancingTestClass->FindPropertyByName(NonNullableOptionalInnerObjectPropertyName));
 
 		SECTION("Non-nullable subobject reference property semantics")
@@ -1790,6 +1991,79 @@ namespace SubobjectInstancingTest::Private
 			}
 		}
 
+		SECTION("Non-nullable transient subobject reference initialized from CDO")
+		{
+			USubobjectInstancingTestOuterObject* OuterObject = NewObject<USubobjectInstancingTestOuterObject>(GetTransientPackage(), InstancingTestClass);
+			USubobjectInstancingTestObject* ObjectReference = Cast<USubobjectInstancingTestObject>(TransientProperty->GetObjectPropertyValue_InContainer(OuterObject));
+
+			if (InstancingTestClass->ShouldUseDynamicSubobjectInstancing())
+			{
+				REQUIRE(ObjectReference != nullptr);
+				CHECK(ObjectReference->IsIn(OuterObject));
+				CHECK(ObjectReference->TestValue == 100);
+
+				CHECK_FALSE(ObjectReference->IsTemplate());
+				CHECK(ObjectReference->IsDefaultSubobject());
+				CHECK_FALSE(ObjectReference->IsTemplateForSubobjects());
+
+				USubobjectInstancingTestOuterObject* NewOuterObject = NewObject<USubobjectInstancingTestOuterObject>(GetTransientPackage(), InstancingTestClass);
+				CHECK(ObjectReference != TransientProperty->GetObjectPropertyValue_InContainer(NewOuterObject));
+			}
+			else
+			{
+				// Note: Transient non-nullable reference types are not supported unless dynamic instancing is enabled. These are default-initialized to NULL,
+				// and are then skipped as part of default value initialization when CPF_InstancedReference is also set (see FObjectInitializer::InitProperties).
+				// Note: This assumes non-nullable reference types are only declared on non-native class types. This test is not currently valid for native types.
+				CHECK(ObjectReference == nullptr);
+			}
+		}
+
+		SECTION("Non-nullable transient subobject reference initialized from template")
+		{
+			USubobjectInstancingTestOuterObject* Template = NewObject<USubobjectInstancingTestOuterObject>(GetTransientPackage(), InstancingTestClass, NAME_None, RF_ArchetypeObject);
+			{
+				USubobjectInstancingTestObject* ObjectReference = Cast<USubobjectInstancingTestObject>(TransientProperty->GetObjectPropertyValue_InContainer(Template));
+				
+				if (InstancingTestClass->ShouldUseDynamicSubobjectInstancing())
+				{
+					REQUIRE(ObjectReference != nullptr);
+					CHECK(ObjectReference->IsIn(Template));
+					CHECK(ObjectReference->TestValue == 100);
+
+					ObjectReference->TestValue = 200;
+				}
+				else
+				{
+					// Note: As with the CDO case, transient non-nullable reference types are not supported unless dynamic instancing is enabled.
+					CHECK(ObjectReference == nullptr);
+				}
+			}
+
+			USubobjectInstancingTestOuterObject* OuterObject = NewObject<USubobjectInstancingTestOuterObject>(GetTransientPackage(), InstancingTestClass, NAME_None, RF_NoFlags, Template);
+			{
+				USubobjectInstancingTestObject* ObjectReference = Cast<USubobjectInstancingTestObject>(TransientProperty->GetObjectPropertyValue_InContainer(OuterObject));
+
+				if (InstancingTestClass->ShouldUseDynamicSubobjectInstancing())
+				{
+					REQUIRE(ObjectReference != nullptr);
+					CHECK(ObjectReference->IsIn(OuterObject));
+					CHECK(ObjectReference->TestValue == 200);
+
+					CHECK_FALSE(ObjectReference->IsTemplate());
+					CHECK(ObjectReference->IsDefaultSubobject());
+					CHECK_FALSE(ObjectReference->IsTemplateForSubobjects());
+
+					USubobjectInstancingTestOuterObject* NewOuterObject = NewObject<USubobjectInstancingTestOuterObject>(GetTransientPackage(), InstancingTestClass, NAME_None, RF_NoFlags, Template);
+					CHECK(ObjectReference != TransientProperty->GetObjectPropertyValue_InContainer(NewOuterObject));
+				}
+				else
+				{
+					// Note: As with the CDO case, transient non-nullable reference types are not supported unless dynamic instancing is enabled.
+					CHECK(ObjectReference == nullptr);
+				}
+			}
+		}
+
 		SECTION("Non-nullable optional subobject reference property semantics")
 		{
 			CHECK_FALSE(OptionalProperty->HasIntrusiveUnsetOptionalState());
@@ -1834,7 +2108,10 @@ namespace SubobjectInstancingTest::Private
 			SrcObject->LocalOnlyInnerObject->TestValue = 200;
 			SrcObject->InnerObjectFromType->TestValue = 200;
 			SrcObject->InnerObjectUsingNew->TestValue = 200;
+			SrcObject->TransientInnerObjectUsingNew->TestValue = 200;
+			SrcObject->DuplicateTransientInnerObjectUsingNew->TestValue = 200;
 			SrcObject->InnerObjectPostInit->TestValue = 200;
+			SrcObject->TransientInnerObjectPostInit->TestValue = 200;
 			SrcObject->StructWithInnerObjects.InnerObject->TestValue = 200;
 			SrcObject->OptionalInnerObject.GetValue()->TestValue = 200;
 			for (USubobjectInstancingTestObject* InnerObject : SrcObject->InnerObjectArray)
@@ -1849,6 +2126,10 @@ namespace SubobjectInstancingTest::Private
 			SrcObject->InnerObjectWithIndirectlyNestedObject->InnerObject->TestValue = 200;
 			SrcObject->InnerObjectWithIndirectlyNestedObject->InnerObject->InnerObject->TestValue = 200;
 
+			// Simulate a deferred edit-time instancing event along with a value override for editinline tests
+			SrcObject->DuplicateTransientEditTimeInnerObject = NewObject<USubobjectInstancingTestObject>(SrcObject, GET_MEMBER_NAME_CHECKED(USubobjectInstancingTestOuterObject, DuplicateTransientEditTimeInnerObject), SrcObject->GetMaskedFlags(RF_PropagateToSubObjects));
+			SrcObject->DuplicateTransientEditTimeInnerObject->TestValue = 200;
+
 			// Note: Using PIE mode for duplication as low level test runners do not configure UPS at boot time and normal mode requires it.
 			FObjectDuplicationParameters Params = InitStaticDuplicateObjectParams(SrcObject, GetTransientPackage());
 			Params.PortFlags = PPF_DuplicateForPIE;
@@ -1877,6 +2158,22 @@ namespace SubobjectInstancingTest::Private
 			CHECK(DstObject->TransientInnerObject != SrcObject->TransientInnerObject);
 			CHECK(DstObject->TransientInnerObject->TestValue == SrcObject->TransientInnerObject->TestValue);
 
+			CHECK(DstObject->DuplicateTransientInnerObjectUsingNew != nullptr);
+			CHECK(DstObject->DuplicateTransientInnerObjectUsingNew->IsInOuter(DstObject));
+			CHECK(DstObject->DuplicateTransientInnerObjectUsingNew != SrcObject->DuplicateTransientInnerObjectUsingNew);
+			// Note: DuplicateTransient will not prevent a default subobject from being serialized after UE-22874
+			// The expected result here would be that TestValue does not receive the override value from the source
+			// A workaround that's widely in use is to construct non-default only and not rely on subobject instancing
+			// @todo - This is likely a bug that needs fixing, but it will break backcompat - this ensures that for now
+			CHECK(SrcObject->DuplicateTransientInnerObjectUsingNew->IsDefaultSubobject());
+			CHECK(DstObject->DuplicateTransientInnerObjectUsingNew->TestValue == SrcObject->DuplicateTransientInnerObjectUsingNew->TestValue);
+
+			// By contrast, DuplicateTransient correctly affects a property that has no default-constructed archetype
+			// The expected result is the duplicatetransient reference will not be duplicated and thus is not serialized
+			CHECK(SrcObject->DuplicateTransientEditTimeInnerObject != nullptr);
+			CHECK_FALSE(SrcObject->DuplicateTransientEditTimeInnerObject->IsDefaultSubobject());
+			CHECK_FALSE(DstObject->DuplicateTransientEditTimeInnerObject != nullptr);
+
 			CHECK(DstObject->LocalOnlyInnerObject != nullptr);
 			CHECK(DstObject->LocalOnlyInnerObject->IsInOuter(DstObject));
 			CHECK(DstObject->LocalOnlyInnerObject != SrcObject->LocalOnlyInnerObject);
@@ -1892,11 +2189,21 @@ namespace SubobjectInstancingTest::Private
 			CHECK(DstObject->InnerObjectUsingNew != SrcObject->InnerObjectUsingNew);
 			CHECK(DstObject->InnerObjectUsingNew->TestValue == SrcObject->InnerObjectUsingNew->TestValue);
 
+			CHECK(DstObject->TransientInnerObjectUsingNew != nullptr);
+			CHECK(DstObject->TransientInnerObjectUsingNew->IsInOuter(DstObject));
+			CHECK(DstObject->TransientInnerObjectUsingNew != SrcObject->TransientInnerObjectUsingNew);
+			CHECK(DstObject->TransientInnerObjectUsingNew->TestValue == SrcObject->TransientInnerObjectUsingNew->TestValue);
+
 			CHECK(DstObject->InnerObjectPostInit != nullptr);
 			CHECK(DstObject->InnerObjectPostInit->IsInOuter(DstObject));
 			CHECK(DstObject->InnerObjectPostInit != SrcObject->InnerObjectPostInit);
 			CHECK(DstObject->InnerObjectPostInit->TestValue == SrcObject->InnerObjectPostInit->TestValue);
 
+			CHECK(DstObject->TransientInnerObjectPostInit != nullptr);
+			CHECK(DstObject->TransientInnerObjectPostInit->IsInOuter(DstObject));
+			CHECK(DstObject->TransientInnerObjectPostInit != SrcObject->TransientInnerObjectPostInit);
+			CHECK(DstObject->TransientInnerObjectPostInit->TestValue == SrcObject->TransientInnerObjectPostInit->TestValue);
+
 			CHECK(DstObject->StructWithInnerObjects.InnerObject != nullptr);
 			CHECK(DstObject->StructWithInnerObjects.InnerObject->IsInOuter(DstObject));
 			CHECK(DstObject->StructWithInnerObjects.InnerObject != SrcObject->StructWithInnerObjects.InnerObject);
diff --git a/Engine/Source/Runtime/CoreUObject/Tests/UObject/SubobjectInstancingTest.h b/Engine/Source/Runtime/CoreUObject/Tests/UObject/SubobjectInstancingTest.h
--- a/Engine/Source/Runtime/CoreUObject/Tests/UObject/SubobjectInstancingTest.h
+++ b/Engine/Source/Runtime/CoreUObject/Tests/UObject/SubobjectInstancingTest.h
@@ -125,6 +125,12 @@ public:
 	UPROPERTY(Instanced)
 	TObjectPtr<USubobjectInstancingTestObject> EditTimeInnerObject;
 
+	UPROPERTY(Transient, Instanced)
+	TObjectPtr<USubobjectInstancingTestObject> TransientEditTimeInnerObject;
+
+	UPROPERTY(DuplicateTransient, Instanced)
+	TObjectPtr<USubobjectInstancingTestObject> DuplicateTransientEditTimeInnerObject;
+
 	UPROPERTY(Instanced)
 	TObjectPtr<USubobjectInstancingTestObject> LocalOnlyInnerObject;
 
@@ -134,9 +140,18 @@ public:
 	UPROPERTY(Instanced)
 	TObjectPtr<USubobjectInstancingTestObject> InnerObjectUsingNew;
 
+	UPROPERTY(Transient, Instanced)
+	TObjectPtr<USubobjectInstancingTestObject> TransientInnerObjectUsingNew;
+
+	UPROPERTY(DuplicateTransient, Instanced)
+	TObjectPtr<USubobjectInstancingTestObject> DuplicateTransientInnerObjectUsingNew;
+
 	UPROPERTY(Instanced)
 	TObjectPtr<USubobjectInstancingTestObject> InnerObjectPostInit;
 
+	UPROPERTY(Transient, Instanced)
+	TObjectPtr<USubobjectInstancingTestObject> TransientInnerObjectPostInit;
+
 	UPROPERTY(Instanced)
 	TArray<TObjectPtr<USubobjectInstancingTestObject>> InnerObjectArray;
 
