[CoreRedirects]

Adds test for Enum value redirects

#rnx

#Synced-CL 48616042

diff --git a/Engine/Source/Runtime/CoreUObject/Private/UObject/CoreRedirects.cpp b/Engine/Source/Runtime/CoreUObject/Private/UObject/CoreRedirects.cpp
--- a/Engine/Source/Runtime/CoreUObject/Private/UObject/CoreRedirects.cpp
+++ b/Engine/Source/Runtime/CoreUObject/Private/UObject/CoreRedirects.cpp
@@ -2113,6 +2113,18 @@ bool FCoreRedirects::RunTests()
 	NewRedirects.Emplace(ECoreRedirectFlags::Type_Package | ECoreRedirectFlags::Category_Removed, TEXT("/Game/RemovedPackage"), TEXT("/Game/RemovedPackage"));
 	NewRedirects.Emplace(ECoreRedirectFlags::Type_Package | ECoreRedirectFlags::Category_Removed | ECoreRedirectFlags::Option_MissingLoad, TEXT("/Game/MissingLoadPackage"), TEXT("/Game/MissingLoadPackage"));
 
+	FCoreRedirect EnumRedirectWithValues(ECoreRedirectFlags::Type_Enum, TEXT("MaskTarget_PhysMesh"), TEXT("/Script/ClothingSystemRuntimeCommon.EWeightMapTargetCommon"));
+	EnumRedirectWithValues.ValueChanges.Add(TEXT("MaskTarget_PhysMesh::MaxDistance"), TEXT("EWeightMapTargetCommon::MaxDistance"));
+	EnumRedirectWithValues.ValueChanges.Add(TEXT("MaskTarget_PhysMesh::BackstopDistance"), TEXT("EWeightMapTargetCommon::BackstopDistance"));
+	EnumRedirectWithValues.ValueChanges.Add(TEXT("MaskTarget_PhysMesh::BackstopRadius"), TEXT("EWeightMapTargetCommon::BackstopRadius"));
+	EnumRedirectWithValues.ValueChanges.Add(TEXT("MaskTarget_PhysMesh::AnimDriveMultiplier"), TEXT("EWeightMapTargetCommon::AnimDriveStiffness"));
+	NewRedirects.Add(EnumRedirectWithValues);
+
+	// This value redirect extends the values from the EnumRedirectWithValues redirect above
+	FCoreRedirect EnumValueRedirect(ECoreRedirectFlags::Type_Enum, TEXT("/Script/ClothingSystemRuntimeCommon.EWeightMapTargetCommon"), TEXT(""));
+	EnumValueRedirect.ValueChanges.Add(TEXT("EWeightMapTargetCommon::AnimDriveMultiplier"), TEXT("EWeightMapTargetCommon::AnimDriveStiffness"));
+	NewRedirects.Add(EnumValueRedirect);
+
 	AddRedirectList(NewRedirects, TEXT("RunTests"));
 
 	// Run the asset tests first so that their entries are present when we run the other tests
@@ -2121,12 +2133,10 @@ bool FCoreRedirects::RunTests()
 
 	struct FRedirectTest
 	{
-		FString OldName;
-		FString NewName;
-		ECoreRedirectFlags Type;
+		FCoreRedirect Expected;
 
 		FRedirectTest(const FString& InOldName, const FString& InNewName, ECoreRedirectFlags InType)
-			: OldName(InOldName), NewName(InNewName), Type(InType)
+			: Expected(InType, InOldName, InNewName)
 		{}
 	};
 
@@ -2167,15 +2177,67 @@ bool FCoreRedirects::RunTests()
 	Tests.Emplace(TEXT("/Game/Package/古い接尾辞"), TEXT("/Game/Package/新しい接尾辞"), ECoreRedirectFlags::Type_Package);
 	Tests.Emplace(TEXT("/Game/Package/混合接尾辞"), TEXT("/Game/Package/mixed_suffix"), ECoreRedirectFlags::Type_Package);
 
+	// Enum value tests
+	FRedirectTest EnumValueTest(TEXT("/Script/ClothingSystemRuntimeCommon.MaskTarget_PhysMesh"), TEXT("/Script/ClothingSystemRuntimeCommon.EWeightMapTargetCommon"), ECoreRedirectFlags::Type_Enum);
+	EnumValueTest.Expected.ValueChanges.Add(TEXT("MaskTarget_PhysMesh::MaxDistance"), TEXT("EWeightMapTargetCommon::MaxDistance"));
+	EnumValueTest.Expected.ValueChanges.Add(TEXT("MaskTarget_PhysMesh::BackstopDistance"), TEXT("EWeightMapTargetCommon::BackstopDistance"));
+	EnumValueTest.Expected.ValueChanges.Add(TEXT("MaskTarget_PhysMesh::BackstopRadius"), TEXT("EWeightMapTargetCommon::BackstopRadius"));
+	EnumValueTest.Expected.ValueChanges.Add(TEXT("MaskTarget_PhysMesh::AnimDriveMultiplier"), TEXT("EWeightMapTargetCommon::AnimDriveStiffness"));
+	Tests.Add(EnumValueTest);
+
+	// Confirms we inherit all the mappings from the MaskTarget_PhysMesh -> EWeightMapTargetCommon mapping
+	// plus the explicit addition to EWeightMapTargetCommon added via a separate value redirect
+	FRedirectTest EnumInheritedAndAppendedValueTest(TEXT("/Script/ClothingSystemRuntimeCommon.EWeightMapTargetCommon"), TEXT("/Script/ClothingSystemRuntimeCommon.EWeightMapTargetCommon"), ECoreRedirectFlags::Type_Enum);
+	EnumInheritedAndAppendedValueTest.Expected.ValueChanges.Add(TEXT("MaskTarget_PhysMesh::MaxDistance"), TEXT("EWeightMapTargetCommon::MaxDistance"));
+	EnumInheritedAndAppendedValueTest.Expected.ValueChanges.Add(TEXT("MaskTarget_PhysMesh::BackstopDistance"), TEXT("EWeightMapTargetCommon::BackstopDistance"));
+	EnumInheritedAndAppendedValueTest.Expected.ValueChanges.Add(TEXT("MaskTarget_PhysMesh::BackstopRadius"), TEXT("EWeightMapTargetCommon::BackstopRadius"));
+	EnumInheritedAndAppendedValueTest.Expected.ValueChanges.Add(TEXT("MaskTarget_PhysMesh::AnimDriveMultiplier"), TEXT("EWeightMapTargetCommon::AnimDriveStiffness"));
+	EnumInheritedAndAppendedValueTest.Expected.ValueChanges.Add(TEXT("EWeightMapTargetCommon::AnimDriveMultiplier"), TEXT("EWeightMapTargetCommon::AnimDriveStiffness"));
+	Tests.Add(EnumInheritedAndAppendedValueTest);
+
 	for (FRedirectTest& Test : Tests)
 	{
-		FCoreRedirectObjectName OldName = FCoreRedirectObjectName(Test.OldName);
-		FCoreRedirectObjectName NewName = GetRedirectedName(Test.Type, OldName);
-
-		if (NewName.ToString() != Test.NewName)
+		FCoreRedirectObjectName ActualNewName;
+		const FCoreRedirect* ValueRedirects = nullptr;
+		RedirectNameAndValues(Test.Expected.RedirectFlags, Test.Expected.OldName, ActualNewName, &ValueRedirects);
+		bool bTestSuccess = ActualNewName == Test.Expected.NewName;
+		if (bTestSuccess && !Test.Expected.ValueChanges.IsEmpty())
+		{
+			bTestSuccess &= ValueRedirects != nullptr;
+			if (ValueRedirects)
+			{
+				for (const TPair<FString, FString>& Pair : Test.Expected.ValueChanges)
+				{
+					const FString* ActualChangeValue = ValueRedirects->ValueChanges.Find(Pair.Key);
+					bTestSuccess &= ActualChangeValue != nullptr;
+					if (ActualChangeValue)
+					{
+						bTestSuccess &= Pair.Value == *ActualChangeValue;
+					}
+					if (!bTestSuccess)
+					{
+						break;
+					}
+				}
+			}
+		}
+		if (!bTestSuccess)
 		{
 			bSuccess = false;
-			UE_LOG(LogCoreRedirects, Error, TEXT("FCoreRedirect Test Failed: %s to %s, should be %s!"), *OldName.ToString(), *NewName.ToString(), *Test.NewName);
+			TStringBuilder<256> ExpectedBuilder;
+			TStringBuilder<256> ActualBuilder;
+			auto PrintRedirect = [](FCoreRedirect& Redirect, TStringBuilderBase<TCHAR>& Builder)
+				{
+					Builder.Appendf(TEXT("%s -> %s"), *Redirect.OldName.ToString(), *Redirect.NewName.ToString());
+					for (TPair<FString, FString>& Pair : Redirect.ValueChanges)
+					{
+						Builder.Appendf(TEXT("\n\t[%s -> %s]"), *Pair.Key, *Pair.Value);
+					}
+				};
+			PrintRedirect(Test.Expected, ExpectedBuilder);
+			FCoreRedirect ActualRedirect = ValueRedirects ? *ValueRedirects : FCoreRedirect(Test.Expected.RedirectFlags, Test.Expected.OldName, ActualNewName);
+			PrintRedirect(ActualRedirect, ActualBuilder);
+			UE_LOG(LogCoreRedirects, Error, TEXT("FCoreRedirect Test Failed:\nExpected:\t%s,\nActual:\t\t%s"), ExpectedBuilder.ToString(), ActualBuilder.ToString());
 		}
 	}
 
