Preventing an object handle memory access when IsForceReturnObjectHandles() when resolving a weak object pointer referencing a remote object

#Synced-CL 0

diff --git a/Engine/Source/Runtime/CoreUObject/Private/UObject/RemoteObjectPrivate.h b/Engine/Source/Runtime/CoreUObject/Private/UObject/RemoteObjectPrivate.h
--- a/Engine/Source/Runtime/CoreUObject/Private/UObject/RemoteObjectPrivate.h
+++ b/Engine/Source/Runtime/CoreUObject/Private/UObject/RemoteObjectPrivate.h
@@ -115,6 +115,11 @@ namespace UE::RemoteObject::Private
 			return InId.GetGlobalized();
 		}
 	};
+
+	/**
+	* Returns true if UE::RemoteObject::Handle::Resolve() functions are forced to return FObjectHandles instead of the actual UObject pointers
+	*/
+	bool IsForceReturnObjectHandles();
 }
 
 namespace UE::RemoteObject::Transfer::Private
diff --git a/Engine/Source/Runtime/CoreUObject/Private/UObject/WeakObjectPtr.cpp b/Engine/Source/Runtime/CoreUObject/Private/UObject/WeakObjectPtr.cpp
--- a/Engine/Source/Runtime/CoreUObject/Private/UObject/WeakObjectPtr.cpp
+++ b/Engine/Source/Runtime/CoreUObject/Private/UObject/WeakObjectPtr.cpp
@@ -258,6 +258,18 @@ FUObjectItem* FWeakObjectPtr::TryResolveRemoteObject() const
 
 	if (ResolvedObject)
 	{
+		if (IsForceReturnObjectHandles())
+		{
+			// In this mode the result from ResolveObject() is actually an FObjectHandle which means it can also be a pointer to FRemoteObjectStub
+			FObjectHandle ResolvedHandle = *reinterpret_cast<FObjectHandle*>(&ResolvedObject);
+			checkf(IsInGameThread(), TEXT("Trying to dereference %s but when ForceReturnObjectHandles is active, weak object pointers can only be resolved on the game thread"), *ResolvedHandle.GetRemoteId().ToString());
+			static FUObjectItem TempObjectItem;
+			FMemory::Memzero(&TempObjectItem, sizeof(FUObjectItem));
+			TempObjectItem.SetRemoteId(ResolvedHandle.GetRemoteId());
+			TempObjectItem.SetObject(ResolvedObject);
+			return &TempObjectItem;
+		}
+
 		int32 NewIndex = GUObjectArray.ObjectToIndex(ResolvedObject);
 		FUObjectItem* ObjectItem = GUObjectArray.IndexToObject(NewIndex);
 		checkf(ObjectItem->GetRemoteId() == RemoteId, TEXT("Remote object %s has a different id (%s) than the weak object pointer that resolved it (%s)"),
