Allow the construction of FRemoteServerId FromIdNumber even for reserved server id values

#Synced-CL 0

diff --git a/Engine/Source/Runtime/CoreUObject/Private/UObject/RemoteObject.cpp b/Engine/Source/Runtime/CoreUObject/Private/UObject/RemoteObject.cpp
--- a/Engine/Source/Runtime/CoreUObject/Private/UObject/RemoteObject.cpp
+++ b/Engine/Source/Runtime/CoreUObject/Private/UObject/RemoteObject.cpp
@@ -1318,7 +1318,7 @@ FArchive& operator<<(FArchive& Ar, FRemoteObjectId& Id)
 
 FRemoteServerId FRemoteServerId::FromIdNumber(uint32 InNumber)
 {
-	checkf(InNumber < (uint32)ERemoteServerIdConstants::FirstReserved, TEXT("Remote server id can not be greater than %u, got: %u"), (uint32)ERemoteServerIdConstants::FirstReserved - 1, InNumber);
+	checkf(InNumber <= (uint32)ERemoteServerIdConstants::Max, TEXT("Remote server id can not be greater than %u, got: %u"), (uint32)ERemoteServerIdConstants::Max, InNumber);
 	FRemoteServerId Result;
 	Result.Id = InNumber;
 	return Result;
