When serializing remote object headers, serialize their remote ids the same way the ids are serialized for UObject references: by storing them in the unique id map and serializing the id index (2 bytes instead of 8 and fewer ids duplicated). Added a number of helper function to easier access to the ids in the unique id tables.
Added a list of serialized object ids that is accessible without deserializing the remote object data (allows access to the list of serialized object ids before data deserialization)

#Synced-CL 0

diff --git a/Engine/Source/Runtime/CoreUObject/Private/UObject/RemoteObjectSerialization.cpp b/Engine/Source/Runtime/CoreUObject/Private/UObject/RemoteObjectSerialization.cpp
--- a/Engine/Source/Runtime/CoreUObject/Private/UObject/RemoteObjectSerialization.cpp
+++ b/Engine/Source/Runtime/CoreUObject/Private/UObject/RemoteObjectSerialization.cpp
@@ -66,6 +66,7 @@ void SerializeNameTables(FArchive& Ar, const FRemoteObjectData& InObjectData)
 	}
 	Ar << const_cast<TArray<FRemoteObjectId>&>(InObjectData.Tables.RemoteIds);
 	Ar << const_cast<TArray<FPackedRemoteObjectPathName>&>(InObjectData.PathNames);
+	Ar << const_cast<TArray<FPackedRemoteObjectPathName::FNameIndexType>&>(InObjectData.SerializedObjectIds);
 }
 
 void DeserializeNameTables(FArchive& Ar, FRemoteObjectData& OutObjectData)
@@ -81,6 +82,7 @@ void DeserializeNameTables(FArchive& Ar, FRemoteObjectData& OutObjectData)
 	}
 	Ar << OutObjectData.Tables.RemoteIds;
 	Ar << OutObjectData.PathNames;
+	Ar << OutObjectData.SerializedObjectIds;
 }
 
 FString GenerateRemoteObjectFilename(FRemoteObjectId ObjectId, FRemoteServerId OwnerServerId)
@@ -427,20 +429,22 @@ public:
 		return Info;
 	}
 
-	virtual bool PopulateObjectHeader(UObject* Object, FRemoteObjectHeader& OutHeader) const
+	virtual bool PopulateObjectHeader(UObject* Object, FRemoteObjectHeader& OutHeader)
 	{
 		using namespace UE::RemoteObject;
 		using namespace UE::RemoteObject::Private;
 		using namespace UE::CoreUObject::Private;
 
 		OutHeader.Name = Object->GetFName();
-		OutHeader.RemoteId = FObjectHandleUtils::GetRemoteId(Object);
+		OutHeader.RemoteId = AddRemoteIdToIdMap(FObjectHandleUtils::GetRemoteId(Object));
 		OutHeader.Class = Object->GetClass();
 		UObject* Outer = FObjectHandleUtils::GetNonAccessTrackedOuterNoResolve(Object);
 		OutHeader.Outer = Outer;
 		OutHeader.Archetype = Object->GetArchetype();
 		OutHeader.InternalFlags = (int32)(Object->GetInternalFlags() & EInternalObjectFlags::Garbage);
 
+		ObjectData.SerializedObjectIds.Add(OutHeader.RemoteId);
+
 		return true;
 	}
 
@@ -716,7 +720,7 @@ protected:
 
 public:
 
-	virtual bool PopulateObjectHeader(UObject* Object, FRemoteObjectHeader& OutHeader) const override
+	virtual bool PopulateObjectHeader(UObject* Object, FRemoteObjectHeader& OutHeader) override
 	{
 		// As per comment in WriteObjectReference - when serializing archetype data pretend we're actually serializing references to instances of the archetype 
 		// so substitute the archetype object with its instance when serializing the object header
@@ -965,7 +969,7 @@ public:
 	}
 };
 
-FRemoteObjectConstructionOverrides::FRemoteObjectConstructionOverrides(const TArray<FRemoteObjectHeader>& InObjectHeaders)
+FRemoteObjectConstructionOverrides::FRemoteObjectConstructionOverrides(const FRemoteObjectData& ObjectData, const TArray<FRemoteObjectHeader>& InObjectHeaders)
 {
 	using namespace UE::RemoteObject::Handle;
 	using namespace UE::RemoteObject::Private;
@@ -978,12 +982,12 @@ FRemoteObjectConstructionOverrides::FRemoteObjectConstructionOverrides(const TAr
 #if UE_WITH_REMOTE_OBJECT_HANDLE
 		Params.OuterId = ::GetRemoteObjectId(Header.Outer.GetHandle());
 #endif
-		Params.RemoteId = Header.RemoteId;
+		Params.RemoteId = ObjectData.Tables.RemoteIds[Header.RemoteId];
 		Params.SerialNumber = Header.SerialNumber;
 	}
 }
 
-UObject* ConstructRemoteObject(const FRemoteObjectHeader& Header, ERemoteObjectSerializationFlags DeserializeFlags)
+UObject* ConstructRemoteObject(const FRemoteObjectId SerializedObjectId, const FRemoteObjectHeader& Header, ERemoteObjectSerializationFlags DeserializeFlags)
 {
 	using namespace UE::RemoteObject::Handle;
 	using namespace UE::RemoteObject::Private;
@@ -1007,16 +1011,16 @@ UObject* ConstructRemoteObject(const FRemoteObjectHeader& Header, ERemoteObjectS
 	UObject* Object = StaticFindObjectFast(Class, Outer, Name);
 	if (Object)
 	{
-		if (FRemoteObjectId(Object) != Header.RemoteId)
+		if (FRemoteObjectId(Object) != SerializedObjectId)
 		{
-			UE_LOG(LogRemoteSerialization, Warning, TEXT("Received remote object %s with identical pathname (%s) as a local object %s. Remote object will be renamed."), *Header.RemoteId.ToString(), *Object->GetPathName(), *FRemoteObjectId(Object).ToString());
+			UE_LOG(LogRemoteSerialization, Warning, TEXT("Received remote object %s with identical pathname (%s) as a local object %s. Remote object will be renamed."), *SerializedObjectId.ToString(), *Object->GetPathName(), *FRemoteObjectId(Object).ToString());
 			Name = FName();
 			Object = nullptr;
 		}
 	}
 	else
 	{
-		Object = StaticFindObjectFastInternal(Header.RemoteId);
+		Object = StaticFindObjectFastInternal(SerializedObjectId);
 		if (Object)
 		{
 			// The object already exists on this server but has been renamed
@@ -1036,7 +1040,7 @@ UObject* ConstructRemoteObject(const FRemoteObjectHeader& Header, ERemoteObjectS
 		Params.SerialNumber = Header.SerialNumber;
 		Params.Template = Header.Archetype.Get();
 #if UE_WITH_REMOTE_OBJECT_HANDLE
-		Params.RemoteId = Header.RemoteId;
+		Params.RemoteId = SerializedObjectId;
 #endif
 
 		{
@@ -1053,7 +1057,7 @@ UObject* ConstructRemoteObject(const FRemoteObjectHeader& Header, ERemoteObjectS
 		return nullptr;
 	}
 
-	checkf(FRemoteObjectId(Object) == Header.RemoteId, TEXT("Created an object with a different ID:%s than requested:%s"), *FRemoteObjectId(Object).ToString(), *Header.RemoteId.ToString());
+	checkf(FRemoteObjectId(Object) == SerializedObjectId, TEXT("Created an object with a different ID:%s than requested:%s"), *FRemoteObjectId(Object).ToString(), *SerializedObjectId.ToString());
 
 	MarkAsLocal(Object);
 
@@ -1283,7 +1287,8 @@ int32 DeserializeObjectData(FRemoteObjectData& ObjectData,
 		// Try to find any existing WeakObject serial numbers for the objects that are about to be constructed
 		for (FRemoteObjectHeader& ObjectHeader : ObjectHeaders)
 		{
-			if (FRemoteObjectStub* Stub = FindRemoteObjectStub(ObjectHeader.RemoteId))
+			FRemoteObjectId SerializedObjectId = ObjectData.GetRemoteObjectId(ObjectHeader.RemoteId);
+			if (FRemoteObjectStub* Stub = FindRemoteObjectStub(SerializedObjectId))
 			{
 				ObjectHeader.SerialNumber = Stub->SerialNumber;
 				if (!Stub->Name.IsNone())
@@ -1294,7 +1299,7 @@ int32 DeserializeObjectData(FRemoteObjectData& ObjectData,
 			}
 		}
 
-		const bool bReturningBorrowedObject = GResetBorrowedObjects && IsOwned(ObjectHeaders[0].RemoteId);
+		const bool bReturningBorrowedObject = GResetBorrowedObjects && IsOwned(ObjectData.GetRootSerializedObjectId());
 		if (bReturningBorrowedObject)
 		{
 			// Root object is owned by this server so we're receiving an object that was borrowed by another server.
@@ -1306,16 +1311,27 @@ int32 DeserializeObjectData(FRemoteObjectData& ObjectData,
 		{
 			TRACE_CPUPROFILER_EVENT_SCOPE(ConstructRemoteObjects);
 
-			FRemoteObjectConstructionOverrides ConstructionOverrides(ObjectHeaders);
+			FRemoteObjectConstructionOverrides ConstructionOverrides(ObjectData, ObjectHeaders);
 			FRemoteObjectConstructionOverridesScope OverridesScope(&ConstructionOverrides);
 
 			OutObjectRemoteIds.Reserve(ObjectHeaders.Num());
 			OutObjects.Reserve(ObjectHeaders.Num());
-			for (int32 ObjectIndex = 0; ObjectIndex < ObjectHeaders.Num(); ++ObjectIndex)
+			for (FRemoteObjectHeader& ObjectHeader : ObjectHeaders)
 			{
-				OutObjectRemoteIds.Add(ObjectHeaders[ObjectIndex].RemoteId);
-				OutObjects.Add(ConstructRemoteObject(ObjectHeaders[ObjectIndex], DeserializeFlags));
+				FRemoteObjectId SerializedObjectId = ObjectData.GetRemoteObjectId(ObjectHeader.RemoteId);
+				OutObjectRemoteIds.Add(SerializedObjectId);
+				OutObjects.Add(ConstructRemoteObject(SerializedObjectId, ObjectHeader, DeserializeFlags));
 			}
+
+			checkf(OutObjectRemoteIds.Num() == ObjectData.SerializedObjectIds.Num(), TEXT("The number of constructed object ids (%d) does not match the number of serialized object ids (%d)"), OutObjectRemoteIds.Num(), ObjectData.SerializedObjectIds.Num());
+#if DO_CHECK
+			for (int32 ObjectIndex = 0; ObjectIndex < OutObjectRemoteIds.Num(); ++ObjectIndex)
+			{
+				FRemoteObjectId ConstructedObjectId = OutObjectRemoteIds[ObjectIndex];
+				FRemoteObjectId SerializedObjectId = ObjectData.GetSerializedObjectId(ObjectIndex);
+				checkf(ConstructedObjectId == SerializedObjectId, TEXT("Constructed object id[%d] (%s) does not match the serialized object id (%s)"), ObjectIndex, *ConstructedObjectId.ToString(), *SerializedObjectId.ToString());
+			}
+#endif
 		}
 
 		UObject* RootObject = OutObjects.Num() ? OutObjects[0] : nullptr;
diff --git a/Engine/Source/Runtime/CoreUObject/Private/UObject/RemoteObjectSerialization.h b/Engine/Source/Runtime/CoreUObject/Private/UObject/RemoteObjectSerialization.h
--- a/Engine/Source/Runtime/CoreUObject/Private/UObject/RemoteObjectSerialization.h
+++ b/Engine/Source/Runtime/CoreUObject/Private/UObject/RemoteObjectSerialization.h
@@ -27,7 +27,7 @@ namespace UE::RemoteObject::Serialization
 	struct FRemoteObjectHeader
 	{
 		FName Name;
-		FRemoteObjectId RemoteId;
+		FPackedRemoteObjectPathName::FNameIndexType RemoteId;
 		TObjectPtr<UClass> Class;
 		TObjectPtr<UObject> Outer;
 		TObjectPtr<UObject> Archetype;
@@ -61,7 +61,7 @@ namespace UE::RemoteObject::Serialization
 	public:
 
 		FRemoteObjectConstructionOverrides() = default;
-		explicit FRemoteObjectConstructionOverrides(const TArray<FRemoteObjectHeader>& InObjectHeaders);
+		explicit FRemoteObjectConstructionOverrides(const FRemoteObjectData& ObjectData, const TArray<FRemoteObjectHeader>& InObjectHeaders);
 
 		/**
 		* Finds object construction overrides for an object that will be constricted with the specified Name and Outer
diff --git a/Engine/Source/Runtime/CoreUObject/Public/UObject/NoExportTypes.h b/Engine/Source/Runtime/CoreUObject/Public/UObject/NoExportTypes.h
--- a/Engine/Source/Runtime/CoreUObject/Public/UObject/NoExportTypes.h
+++ b/Engine/Source/Runtime/CoreUObject/Public/UObject/NoExportTypes.h
@@ -2837,6 +2837,9 @@ struct FRemoteObjectData
 	UPROPERTY()
 	TArray<FPackedRemoteObjectPathName> PathNames;
 
+	UPROPERTY()
+	TArray<uint16> SerializedObjectIds;
+
 	UPROPERTY()
 	TArray<FRemoteObjectBytes> Bytes;
 };
diff --git a/Engine/Source/Runtime/CoreUObject/Public/UObject/RemoteObjectTransfer.h b/Engine/Source/Runtime/CoreUObject/Public/UObject/RemoteObjectTransfer.h
--- a/Engine/Source/Runtime/CoreUObject/Public/UObject/RemoteObjectTransfer.h
+++ b/Engine/Source/Runtime/CoreUObject/Public/UObject/RemoteObjectTransfer.h
@@ -28,10 +28,18 @@ struct FRemoteObjectBytes
 */
 struct FRemoteObjectData
 {
+	/** Tables containing lists of unique remote ids and FNames */
 	FRemoteObjectTables Tables;
+	/** List of unique pathnames (stored as a list of indices of unique FNames) */
 	TArray<FPackedRemoteObjectPathName> PathNames;
+	/** List of remote id indices of the objects serialized in this ObjectData struct (in the same order they were serialized so it matches the order of object headers after deserialization) */
+	TArray<FPackedRemoteObjectPathName::FNameIndexType> SerializedObjectIds;
+	/** Serialized object headers and data (properties) */
 	TArray<FRemoteObjectBytes> Bytes;
 
+	/**
+	* Returns the size of serialized data
+	*/
 	inline int32 GetNumBytes() const
 	{
 		int32 Num = 0;
@@ -41,6 +49,45 @@ struct FRemoteObjectData
 		}
 		return Num;
 	}
+
+	/**
+	* Returns a unique remote object id based on its index the id tables
+	*/
+	inline FRemoteObjectId GetRemoteObjectId(FPackedRemoteObjectPathName::FNameIndexType IdIndex) const
+	{
+		return Tables.RemoteIds[IdIndex];
+	}
+
+	/**
+	* Returns a remote object id for the serialized object
+	* @param ObjectIndex Index of the serialized object
+	* @return Remote Object Id of the serialized object
+	*/
+	inline FRemoteObjectId GetSerializedObjectId(int32 ObjectIndex) const
+	{
+		return GetRemoteObjectId(SerializedObjectIds[ObjectIndex]);
+	}
+
+	/**
+	* Returns the root serialized object's remote id (the first serialized object's id)
+	* @return Remote Object Id of the root serialized object
+	*/
+	inline FRemoteObjectId GetRootSerializedObjectId() const
+	{
+		return SerializedObjectIds.Num() ? GetSerializedObjectId(0) : FRemoteObjectId();
+	}
+
+	/**
+	* Calls the specified function for each serialized object remote id
+	* @param Op Function to call (can be a lambda)
+	*/
+	inline void ForEachSerializedObjectId(TFunctionRef<void(FRemoteObjectId)> Op) const
+	{
+		for (FPackedRemoteObjectPathName::FNameIndexType Index : SerializedObjectIds)
+		{
+			Op(Tables.RemoteIds[Index]);
+		}
+	}
 };
 
 namespace UE::RemoteObject::Transfer
