[ZenLoader]
- Fix crash where we tried to call SetAllPublicExportsLoaded for package that were in failure state

#Synced-CL Unknown

diff --git a/Engine/Source/Runtime/CoreUObject/Private/Serialization/AsyncLoading2.cpp b/Engine/Source/Runtime/CoreUObject/Private/Serialization/AsyncLoading2.cpp
--- a/Engine/Source/Runtime/CoreUObject/Private/Serialization/AsyncLoading2.cpp
+++ b/Engine/Source/Runtime/CoreUObject/Private/Serialization/AsyncLoading2.cpp
@@ -2154,7 +2154,7 @@ public:
 			}
 			if (PackageRef->HasPackage())
 			{
-				if (PackageRef->GetPackage()->bHasBeenFullyLoaded)
+				if (PackageRef->GetPackage()->bHasBeenFullyLoaded && !PackageRef->IsFailed())
 				{
 					const bool bSnapshotExportCount = false;
 					PackageRef->SetAllPublicExportsLoaded(bSnapshotExportCount);
