[LinkerLoad/ALT]
- Fix race condition where a redirector could be reserialized while it was being postloaded

#Synced-CL Unknown

diff --git a/Engine/Source/Runtime/CoreUObject/Private/UObject/LinkerLoad.cpp b/Engine/Source/Runtime/CoreUObject/Private/UObject/LinkerLoad.cpp
--- a/Engine/Source/Runtime/CoreUObject/Private/UObject/LinkerLoad.cpp
+++ b/Engine/Source/Runtime/CoreUObject/Private/UObject/LinkerLoad.cpp
@@ -3531,9 +3531,13 @@ FLinkerLoad::EVerifyResult FLinkerLoad::VerifyImport(int32 ImportIndex)
 				// serialize in the properties of the redirector (to get the object the redirector point to)
 				// Always load redirectors in case there was a circular dependency. This will allow inner redirector
 				// references to always serialize fully here before accessing the DestinationObject
-				check(!GEventDrivenLoaderEnabled || !EVENT_DRIVEN_ASYNC_LOAD_ACTIVE_AT_RUNTIME);
-				Redir->SetFlags(RF_NeedLoad);
-				Preload(Redir);
+				// If the redirector is now in phase2, it already has been through serialization
+				// and we don't want to risk racing with its PostLoad by forcing its reload needlessly.
+				if (!Redir->HasAnyFlags(RF_LoadCompleted) && !Redir->HasAnyInternalFlags(EInternalObjectFlags::AsyncLoadingPhase2))
+				{
+					Redir->SetFlags(RF_NeedLoad);
+					Preload(Redir);
+				}
 
 				UObject* DestObject = Redir->DestinationObject;
 
@@ -6452,7 +6456,9 @@ UObject* FLinkerLoad::CreateImport( int32 Index )
 						// but RF_NeedLoad has been removed already. We need to reresolve in that case right away
 						// otherwise the DestinationObject wouldn't be set until we unwind the stack and finish
 						// the deserialization, which may be too late.
-						if (!Redirector->HasAnyFlags(RF_NeedLoad|RF_LoadCompleted))
+						// If the redirector is now in phase2, it already has been through serialization
+						// and we don't want to risk racing with its PostLoad by forcing its reload needlessly.
+						if (!Redirector->HasAnyFlags(RF_NeedLoad|RF_LoadCompleted) && !Redirector->HasAnyInternalFlags(EInternalObjectFlags::AsyncLoadingPhase2))
 						{
 							// Set the flag back if missing and preload hasn't completed yet so that
 							// the preload we're going to run does something.
