Support redirected properties in FFieldPath
Fix issues where objects serializing FFieldPaths (like FControlRigOverrideValue) would fail to resolve the path if it was referencing a property that has been redirected

#Synced-CL 48789429

diff --git a/Engine/Source/Runtime/CoreUObject/Private/UObject/FieldPath.cpp b/Engine/Source/Runtime/CoreUObject/Private/UObject/FieldPath.cpp
--- a/Engine/Source/Runtime/CoreUObject/Private/UObject/FieldPath.cpp
+++ b/Engine/Source/Runtime/CoreUObject/Private/UObject/FieldPath.cpp
@@ -278,7 +278,16 @@ FField* FFieldPath::TryToResolvePath(UStruct* InCurrentStruct, FFieldPath::EPath
 	if (Owner && Path.Num() > 0)
 	{
 		int32 PathIndex = Path.Num() - 1;
-		Result = FindFProperty<FField>(Owner, Path[PathIndex--]);
+		Result = FindFProperty<FField>(Owner, Path[PathIndex]);
+		if (!Result)
+		{
+			const FName RedirectedName = FProperty::FindRedirectedPropertyName(Owner, Path[PathIndex]);
+			if (!RedirectedName.IsNone())
+			{
+				Result = FindFProperty<FField>(Owner, RedirectedName);
+			}
+		}
+		PathIndex--;
 		while (Result && PathIndex >= 0)
 		{
 			// Nested property
