Store remotely referenced objects to database before the incremental purge phase of garbage collection starts. This is to avoid a scenario where an unreachable object is dereferenced before it gets stored into database.

#Synced-CL 0

diff --git a/Engine/Source/Runtime/CoreUObject/Private/UObject/GarbageCollection.cpp b/Engine/Source/Runtime/CoreUObject/Private/UObject/GarbageCollection.cpp
--- a/Engine/Source/Runtime/CoreUObject/Private/UObject/GarbageCollection.cpp
+++ b/Engine/Source/Runtime/CoreUObject/Private/UObject/GarbageCollection.cpp
@@ -866,12 +866,12 @@ public:
 		// Global UObject Array needs to be locked only when freeing UObject indices
 		GUObjectArray.LockInternalArray();
 		{
-			// This loops replaces every entry in GUnreachableObjects which up untill this point has been an FUObjectItem 
+			// This loops replaces every entry in GUnreachableObjects which up until this point has been an FUObjectItem 
 			// with the actual UObject that the FUObjectItem represented.
 			// This is because in this loop we free UObject indices meaning that they are removed from GUObjectArray and 
 			// the associated FUObjectItem is reset and no longer pointing at UObject.
 			// This approach has the benefit that we don't need to lock GUObjectArray when calling UObject destructors and 
-			// we can reclaim GUObjectArray entires faster.
+			// we can reclaim GUObjectArray entries faster.
 			while (ObjCurrentFreeIndexObjectIndex < GUnreachableObjects.Num())
 			{
 				UE::GC::FUnreachableObject& UnreachableObject = GUnreachableObjects[ObjCurrentFreeIndexObjectIndex];
@@ -905,7 +905,7 @@ public:
 
 		if (ObjCurrentFreeIndexObjectIndex == GUnreachableObjects.Num())
 		{
-			// At this point all entires in GUnreachableObjects point at UObject memory instead of FUObjectItems
+			// At this point all entries in GUnreachableObjects point at UObject memory instead of FUObjectItems
 			while (ObjCurrentPurgeObjectIndex < GUnreachableObjects.Num())
 			{
 				UE::GC::FUnreachableObject& UnreachableObject = GUnreachableObjects[ObjCurrentPurgeObjectIndex];
@@ -5363,10 +5363,18 @@ bool GatherUnreachableObjects(UE::GC::EGatherOptions Options, double TimeLimit /
 			AcquireGCLock();
 		}
 		{			
-			TArrayView<UE::GC::FUnreachableObject> UnrachableObjectsView(GUnreachableObjects);
+			TArrayView<UE::GC::FUnreachableObject> UnreachableObjectsView(GUnreachableObjects);
 			// Although the below cast is not pretty, we don't want to expose UE::GC::FUnreachableObject to async loading API 
 			// as it may be confusing what UE::GC::FUnreachableObject currently represents
-			NotifyUnreachableObjects(*reinterpret_cast<TArrayView<FUObjectItem*>*>(&UnrachableObjectsView));
+			const TArrayView<FUObjectItem*>& UnreachableObjectItemsView = *reinterpret_cast<TArrayView<FUObjectItem*>*>(&UnreachableObjectsView);
+			NotifyUnreachableObjects(UnreachableObjectItemsView);
+
+#if UE_WITH_REMOTE_OBJECT_HANDLE
+			if (!GExitPurge)
+			{
+				UE::RemoteObject::Transfer::Private::StoreUnreachableRemoteObjectsToDatabase(UnreachableObjectItemsView);
+			}
+#endif // UE_WITH_REMOTE_OBJECT_HANDLE
 		}
 		if (bNeedsGCLock)
 		{
@@ -6206,12 +6214,6 @@ bool UnhashUnreachableObjects(bool bUseTimeLimit, double TimeLimit)
 				UObject* Object = static_cast<UObject*>(ObjectItem->GetObject());
 				FScopedCBDProfile Profile(Object);
 				// Begin the object's asynchronous destruction.
-#if UE_WITH_REMOTE_OBJECT_HANDLE
-				if (!GExitPurge && ObjectItem->HasAnyFlags(EInternalObjectFlags::RemoteReference) && !ObjectItem->HasAnyFlags(EInternalObjectFlags::Remote))
-				{
-					UE::RemoteObject::Transfer::Private::StoreObjectToDatabase(Object);
-				}
-#endif // UE_WITH_REMOTE_OBJECT_HANDLE
 				Object->ConditionalBeginDestroy();
 			}
 
diff --git a/Engine/Source/Runtime/CoreUObject/Private/UObject/RemoteObjectPrivate.h b/Engine/Source/Runtime/CoreUObject/Private/UObject/RemoteObjectPrivate.h
--- a/Engine/Source/Runtime/CoreUObject/Private/UObject/RemoteObjectPrivate.h
+++ b/Engine/Source/Runtime/CoreUObject/Private/UObject/RemoteObjectPrivate.h
@@ -127,7 +127,7 @@ namespace UE::RemoteObject::Transfer::Private
 
 
 	/**
-	* Stores unreachable objects to database
+	* Stores remotely referenced unreachable objects to database
 	*/
-	void StoreObjectToDatabase(UObject* Object);
+	void StoreUnreachableRemoteObjectsToDatabase(const TArrayView<FUObjectItem*>& UnreachableObjects);
 }
diff --git a/Engine/Source/Runtime/CoreUObject/Private/UObject/RemoteObjectTransfer.cpp b/Engine/Source/Runtime/CoreUObject/Private/UObject/RemoteObjectTransfer.cpp
--- a/Engine/Source/Runtime/CoreUObject/Private/UObject/RemoteObjectTransfer.cpp
+++ b/Engine/Source/Runtime/CoreUObject/Private/UObject/RemoteObjectTransfer.cpp
@@ -1241,20 +1241,12 @@ UObject* GetOutermostUnreachableRemoteReference(UObject* Object)
 	return OutermostUnreachableRemoteReference;
 }
 
-void StoreObjectToDatabase(UObject* Object)
+void StoreObjectToDatabase(UObject* Object, FRemoteObjectId ObjectId)
 {
 	using namespace UE::RemoteObject::Serialization;
 	using namespace UE::RemoteObject::Private;
 	using namespace UE::RemoteObject::Transfer;
 
-	FRemoteObjectId ObjectId(Object);
-	if (ObjectId.IsAsset())
-	{
-		// We don't need to store assets to database since they already exist on disk so we just store their pathname so that we can reload them if needed
-		StoreAssetPath(Object);
-		return;
-	}
-
 	FRemoteServerId LocalServerId = FRemoteServerId::GetLocalServerId();
 	FRemoteServerId OwnerServerId = UE::RemoteObject::Handle::GetOwnerServerId(Object);
 	if (LocalServerId != OwnerServerId)
@@ -1308,7 +1300,7 @@ void StoreObjectToDatabase(UObject* Object)
 			UE::RemoteObject::Transfer::OnObjectsSavedToDiskDelegate.Broadcast(SentObjects, Params.ObjectData, StoreToDatabaseContext);
 		}
 
-		checkf(SentObjects.Num() > 0, TEXT("PeerId:%s unable to store object %s to DatabaseId:%s"), *FRemoteServerId::GetLocalServerId().ToString(), *FRemoteObjectId(Object).ToString(), *DatabaseId.ToString());
+		checkf(SentObjects.Num() > 0, TEXT("PeerId:%s unable to store object %s to DatabaseId:%s"), *FRemoteServerId::GetLocalServerId().ToString(), *ObjectId.ToString(), *DatabaseId.ToString());
 
 		for (UObject* ReferencedObject : ReferencedObjects)
 		{
@@ -1329,6 +1321,31 @@ void StoreObjectToDatabase(UObject* Object)
 	} while (IsRemoteReference(Object));
 }
 
+void StoreUnreachableRemoteObjectsToDatabase(const TArrayView<FUObjectItem*>& UnreachableObjects)
+{
+	for (const FUObjectItem* ObjectItem : UnreachableObjects)
+	{
+		if (ObjectItem->HasAnyFlags(EInternalObjectFlags::RemoteReference) && !ObjectItem->HasAnyFlags(EInternalObjectFlags::Remote))
+		{
+			FRemoteObjectId ObjectItemId;
+			UObject* Object = (UObject*)ObjectItem->GetObject();
+#if UE_WITH_REMOTE_OBJECT_HANDLE
+			// We can avoid looking up the object id (which happens in FRemoteObjectId(Object)) since we already have the FUObjectItem where the id is stored
+			ObjectItemId = ObjectItem->GetRemoteId();
+#endif
+			if (ObjectItemId.IsAsset())
+			{
+				// We don't need to store assets to database since they already exist on disk so we just store their pathname so that we can reload them if needed
+				UE::RemoteObject::Private::StoreAssetPath(Object);
+			}
+			else
+			{
+				StoreObjectToDatabase(Object, ObjectItemId);
+			}			
+		}
+	}
+}
+
 } // namespace UE::RemoteObject::Transfer::Private
 
 
