[AsyncLoading2]

Remove unnecessary base class from FAsyncPackage2 as we already internally track refcounts and the IRefCountBase interface is too intrusive.

#rnx
#Synced-CL 48850838

diff --git a/Engine/Source/Runtime/CoreUObject/Private/Serialization/AsyncLoading2.cpp b/Engine/Source/Runtime/CoreUObject/Private/Serialization/AsyncLoading2.cpp
--- a/Engine/Source/Runtime/CoreUObject/Private/Serialization/AsyncLoading2.cpp
+++ b/Engine/Source/Runtime/CoreUObject/Private/Serialization/AsyncLoading2.cpp
@@ -3948,7 +3948,7 @@ private:
 * Structure containing intermediate data required for async loading of all exports of a package.
 */
 
-struct FAsyncPackage2 : public FRefCountedObject
+struct FAsyncPackage2
 {
 	friend struct FAsyncPackageScope2;
 	friend class FAsyncLoadingThread2;
