[UObjectBase]
- Make PendingRegistrant thread-safe

#Synced-CL 48817740

diff --git a/Engine/Source/Runtime/CoreUObject/Private/UObject/UObjectBase.cpp b/Engine/Source/Runtime/CoreUObject/Private/UObject/UObjectBase.cpp
--- a/Engine/Source/Runtime/CoreUObject/Private/UObject/UObjectBase.cpp
+++ b/Engine/Source/Runtime/CoreUObject/Private/UObject/UObjectBase.cpp
@@ -13,6 +13,7 @@
 #include "Logging/StructuredLog.h"
 #include "Misc/FeedbackContext.h"
 #include "Modules/ModuleManager.h"
+#include "Containers/StripedMap.h"
 #include "UObject/UObjectAllocator.h"
 #include "UObject/UObjectHash.h"
 #include "UObject/Class.h"
@@ -70,7 +71,7 @@ struct FPendingRegistrantInfo
 #if UE_WITH_REMOTE_OBJECT_HANDLE
 	const FRemoteObjectId RemoteId;
 #endif // UE_WITH_REMOTE_OBJECT_HANDLE
-
+	FPendingRegistrantInfo() = default;
 	FPendingRegistrantInfo(class UClass* (*InStaticClassFn)(), const TCHAR* InName, const TCHAR* InPackageName
 #if UE_WITH_REMOTE_OBJECT_HANDLE
 		, const FRemoteObjectId InRemoteId
@@ -83,9 +84,9 @@ struct FPendingRegistrantInfo
 		, RemoteId(InRemoteId)
 #endif
 	{}
-	static TMap<UObjectBase*, FPendingRegistrantInfo>& GetMap()
+	static TStripedMap<4, UObjectBase*, FPendingRegistrantInfo>& GetMap()
 	{
-		static TMap<UObjectBase*, FPendingRegistrantInfo> PendingRegistrantInfo;
+		static TStripedMap<4, UObjectBase*, FPendingRegistrantInfo> PendingRegistrantInfo;
 		return PendingRegistrantInfo;
 	}
 };
@@ -545,7 +546,7 @@ void UObjectBase::Register(class UClass* (*StaticClassFn)(), const TCHAR* Packag
 {
 	LLM_SCOPE(ELLMTag::UObject);
 	LLM_SCOPE_BYTAG(UObject_UObjectBase);
-	TMap<UObjectBase*, FPendingRegistrantInfo>& PendingRegistrants = FPendingRegistrantInfo::GetMap();
+	TStripedMap<4, UObjectBase*, FPendingRegistrantInfo>& PendingRegistrants = FPendingRegistrantInfo::GetMap();
 
 	FPendingRegistrant* PendingRegistration = new FPendingRegistrant(this);
 	PendingRegistrants.Add(this, FPendingRegistrantInfo(StaticClassFn, InName, PackageName
@@ -581,9 +582,16 @@ void UObjectBase::Register(class UClass* (*StaticClassFn)(), const TCHAR* Packag
 #if UE_WITH_REMOTE_OBJECT_HANDLE
 FRemoteObjectId UObjectBase::GetPendingRegistrantRemoteId() const
 {
-	TMap<UObjectBase*, FPendingRegistrantInfo>& PendingRegistrants = FPendingRegistrantInfo::GetMap();
-	FPendingRegistrantInfo& Info = PendingRegistrants.FindChecked(this);
-	return Info.RemoteId;
+	FRemoteObjectId RemoteId;
+	TStripedMap<4, UObjectBase*, FPendingRegistrantInfo>& PendingRegistrants = FPendingRegistrantInfo::GetMap();
+	bool bFound = PendingRegistrants.FindAndApply(this,
+		[&RemoteId](const FPendingRegistrantInfo& Info)
+		{
+			RemoteId = Info.RemoteId;
+		}
+	);
+	check(bFound);
+	return RemoteId;
 }
 #endif
 
@@ -637,24 +645,32 @@ void UObjectForceRegistration(UObjectBase* Object, bool bCheckForModuleRelease)
 {
 	LLM_SCOPE(ELLMTag::UObject);
 	LLM_SCOPE_BYTAG(UObject_UObjectBase);
-	TMap<UObjectBase*, FPendingRegistrantInfo>& PendingRegistrants = FPendingRegistrantInfo::GetMap();
+	TStripedMap<4, UObjectBase*, FPendingRegistrantInfo>& PendingRegistrants = FPendingRegistrantInfo::GetMap();
 
-	FPendingRegistrantInfo* Info = PendingRegistrants.Find(Object);
-	if (Info)
+	const TCHAR* Name;
+	const TCHAR* PackageName;
+	UClass* StaticClass;
+#if UE_WITH_REMOTE_OBJECT_HANDLE
+	FRemoteObjectId RemoteId;
+#endif
+	if (PendingRegistrants.RemoveIf(Object,
+		[&](FPendingRegistrantInfo& Info)
+		{
+			PackageName = Info.PackageName;
+			Name = Info.Name;
+			StaticClass = Info.StaticClassFn();
+#if UE_WITH_REMOTE_OBJECT_HANDLE
+			RemoteId = Info.RemoteId;
+#endif
+			return true;
+		}))
 	{
-		const TCHAR* PackageName = Info->PackageName;
 #if USE_PER_MODULE_UOBJECT_BOOTSTRAP
 		if (bCheckForModuleRelease)
 		{
 			UObjectReleaseModuleRegistrants(FName(PackageName));
 		}
 #endif
-		const TCHAR* Name = Info->Name;
-		UClass* StaticClass = Info->StaticClassFn();
-#if UE_WITH_REMOTE_OBJECT_HANDLE
-		FRemoteObjectId RemoteId = Info->RemoteId;
-#endif
-		PendingRegistrants.Remove(Object);  // delete this first so that it doesn't try to do it twice
 		Object->DeferredRegister(StaticClass, PackageName, Name
 #if UE_WITH_REMOTE_OBJECT_HANDLE
 			, RemoteId
@@ -1098,12 +1114,10 @@ void ProcessNewlyLoadedUObjects(FName InModuleName, bool bCanProcessNewlyLoadedO
 	StructRegistry.EmptyRegistrations();
 	ClassRegistry.EmptyRegistrations();
 
-	if (TMap<UObjectBase*, FPendingRegistrantInfo>& PendingRegistrants = FPendingRegistrantInfo::GetMap(); PendingRegistrants.IsEmpty())
+	TStripedMap<4, UObjectBase*, FPendingRegistrantInfo>& PendingRegistrants = FPendingRegistrantInfo::GetMap();
+	if (PendingRegistrants.Num() == 0 && PendingRegistrants.GetAllocatedSize() > 0)
 	{
-		if (PendingRegistrants.GetMaxIndex() > 0)
-		{
-			PendingRegistrants.Empty();
-		}
+		PendingRegistrants.Shrink();
 	}
 
 	if (bNewUObjects && !GIsInitialLoad)
