[LinkerLoad/ALT]
- Skip deferred dependencices recursion on package that have finished serialization

#Synced-CL 48817740

diff --git a/Engine/Source/Runtime/CoreUObject/Private/Blueprint/BlueprintSupport.cpp b/Engine/Source/Runtime/CoreUObject/Private/Blueprint/BlueprintSupport.cpp
--- a/Engine/Source/Runtime/CoreUObject/Private/Blueprint/BlueprintSupport.cpp
+++ b/Engine/Source/Runtime/CoreUObject/Private/Blueprint/BlueprintSupport.cpp
@@ -1487,6 +1487,24 @@ int32 FLinkerLoad::FindCDOExportIndex(UClass* LoadClass)
 
 UPackage* LoadPackageInternal(UPackage* InOuter, const FPackagePath& PackagePath, uint32 LoadFlags, FLinkerLoad* ImportLinker, FArchive* InReaderOverride, const FLinkerInstancingContext* InstancingContext, const FPackagePath* DiffPackagePath);
 
+bool ShouldSkipResolveDeferredDependencies(UStruct* LoadStruct)
+{
+	if (IsInActualLoadingThread() || IsInParallelLoadingThread())
+	{
+		if (!LoadStruct->HasAnyInternalFlags(EInternalObjectFlags::AsyncLoadingPhase1))
+		{
+			// If we're in the loader, it would be a race hazard to try to recurse into
+			// packages and linkers that have finished loading, so we skip unless the
+			// package we target is also owned by the loader.
+			// We know that a struct can't contain deferred dependencies anymore once it
+			// leaves its phase 1 and is ready for postload.
+			return true;
+		}
+	}
+
+	return false;
+}
+
 void FLinkerLoad::ResolveDeferredDependencies(UStruct* LoadStruct)
 {
 #if USE_CIRCULAR_DEPENDENCY_LOAD_DEFERRING
@@ -1634,7 +1652,7 @@ void FLinkerLoad::ResolveDeferredDependencies(UStruct* LoadStruct)
 			{
 				// in case this is a user defined struct, we have to resolve any 
 				// deferred dependencies in the struct 
-				if (SourceLinker != nullptr)
+				if (SourceLinker != nullptr && !ShouldSkipResolveDeferredDependencies(StructObj))
 				{
 					SourceLinker->ResolveDeferredDependencies(StructObj);
 				}
@@ -1657,20 +1675,23 @@ void FLinkerLoad::ResolveDeferredDependencies(UStruct* LoadStruct)
 
 		if (UStruct* SuperStruct = LoadStruct->GetSuperStruct())
 		{
-			FLinkerLoad* SuperLinker = SuperStruct->GetLinker();
-			// NOTE: there is no harm in calling this when the super is not 
-			//       "actively resolving deferred dependencies"... this condition  
-			//       just saves on wasted ops, looping over the super's ImportMap
-			if ((SuperLinker != nullptr) && SuperLinker->HasUnresolvedDependencies())
-			{
-				// a resolve could have already been started up the stack, and in turn  
-				// started loading a different package that resulted in another (this) 
-				// resolve beginning... in that scenario, the original resolve could be 
-				// for this class's super and we want to make sure that finishes before
-				// we regenerate this class (else the generated script code could end up 
-				// with unwanted placeholder references; ones that would have been
-				// resolved by the super's linker)
-				SuperLinker->ResolveDeferredDependencies(SuperStruct);
+			if (!ShouldSkipResolveDeferredDependencies(SuperStruct))
+			{
+				FLinkerLoad* SuperLinker = SuperStruct->GetLinker();
+				// NOTE: there is no harm in calling this when the super is not 
+				//       "actively resolving deferred dependencies"... this condition  
+				//       just saves on wasted ops, looping over the super's ImportMap
+				if ((SuperLinker != nullptr) && SuperLinker->HasUnresolvedDependencies())
+				{
+					// a resolve could have already been started up the stack, and in turn  
+					// started loading a different package that resulted in another (this) 
+					// resolve beginning... in that scenario, the original resolve could be 
+					// for this class's super and we want to make sure that finishes before
+					// we regenerate this class (else the generated script code could end up 
+					// with unwanted placeholder references; ones that would have been
+					// resolved by the super's linker)
+					SuperLinker->ResolveDeferredDependencies(SuperStruct);
+				}
 			}
 		}
 
