LLM Tracking during cooking: Mark that allocations made during PreSave are in system-specific code rather than cooker code.
#rnx

#Synced-CL 48806790

diff --git a/Engine/Source/Runtime/CoreUObject/Private/UObject/SavePackage/SavePackageUtilities.cpp b/Engine/Source/Runtime/CoreUObject/Private/UObject/SavePackage/SavePackageUtilities.cpp
--- a/Engine/Source/Runtime/CoreUObject/Private/UObject/SavePackage/SavePackageUtilities.cpp
+++ b/Engine/Source/Runtime/CoreUObject/Private/UObject/SavePackage/SavePackageUtilities.cpp
@@ -11,6 +11,7 @@
 #include "Cooker/CookEvents.h"
 #include "CoreMinimal.h"
 #include "HAL/FileManager.h"
+#include "HAL/LowLevelMemTracker.h"
 #include "Hash/Blake3.h"
 #include "Interfaces/ITargetPlatform.h"
 #include "IO/IoHash.h"
@@ -655,6 +656,8 @@ bool IsProceduralSave(bool bIsCooking, const FPackagePath& TargetPackagePath, ui
 void CallPreSave(UObject* Object, FObjectSaveContextData& ObjectSaveContext)
 {
 	SCOPED_SAVETIMER_TEXT(*WriteToString<256>(GetClassTraceScope(Object), TEXTVIEW("_PreSave")));
+	// Reset the scope so that untagged memory allocated during PreSave shows up as Untagged rather than the caller of SavePackage
+	LLM_SCOPE(ELLMTag::Untagged);
 
 	FObjectPreSaveContext ObjectPreSaveContext(ObjectSaveContext);
 	ObjectSaveContext.bBaseClassCalled = false;
