[AutoRTFM] Use CloseIfTransactional rather than a manual check.

CloseIfTransactional is a new API which is designed to encapsulate this sort of check. This isn't intended to change behavior, just make the code read more naturally.

CloseIfTransactional is expected to optimize out in builds which don't enable AutoRTFM.

#Synced-CL 48922620

diff --git a/Engine/Source/Runtime/CoreUObject/Private/UObject/UObjectGlobals.cpp b/Engine/Source/Runtime/CoreUObject/Private/UObject/UObjectGlobals.cpp
--- a/Engine/Source/Runtime/CoreUObject/Private/UObject/UObjectGlobals.cpp
+++ b/Engine/Source/Runtime/CoreUObject/Private/UObject/UObjectGlobals.cpp
@@ -4213,17 +4213,11 @@ void FObjectInitializer::PostConstructInit()
 		// The Class of the Class may not be available during very early initialization of CoreUObject.
 		if (UClass* ClassClass = Class->GetClass(); ClassClass && ClassClass->IsChildOf<UVerseClass>())
 		{
-			if (AutoRTFM::IsTransactional())
-			{
-#if UE_AUTORTFM_ENABLED // #if purely to silence unreachable code warning.
-				AutoRTFM::EContextStatus Status = AutoRTFM::Close([&] { UE::FInitializedPropertyValueState(Obj).ActivateTracking(); });
-				check(Status == AutoRTFM::EContextStatus::OnTrack);
-#endif
-			}
-			else
+			AutoRTFM::EContextStatus Status = AutoRTFM::CloseIfTransactional([&]
 			{
 				UE::FInitializedPropertyValueState(Obj).ActivateTracking();
-			}
+			});
+			check(Status == AutoRTFM::EContextStatus::OnTrack);
 		}
 	}
 #endif
