ChangePhysicsIslandId will now also mark the passed in objects as a remote reference after creating a stub for the object.
Cleaned up FindOrAddRemoteObjectStub functions so they behave consistently.

#Synced-CL 48777527

diff --git a/Engine/Source/Runtime/CoreUObject/Private/UObject/RemoteObject.cpp b/Engine/Source/Runtime/CoreUObject/Private/UObject/RemoteObject.cpp
--- a/Engine/Source/Runtime/CoreUObject/Private/UObject/RemoteObject.cpp
+++ b/Engine/Source/Runtime/CoreUObject/Private/UObject/RemoteObject.cpp
@@ -146,43 +146,46 @@ public:
 		{
 			Stub = new FRemoteObjectStub();
 			Stub->Id = Id;
-			Stub->ResidentServerId = ResidentServerId.IsValid() ? ResidentServerId : Id.GetServerId();
 
 			// if we are creating the stub, then this object's owner is deduced from its ID
 			// if the server ID is invalid then it's a local native object that was created before the local server had its ID assigned
 			FRemoteServerId ObjectServerId = Id.GetServerId();
 			Stub->OwningServerId = ObjectServerId.IsValid() ? ObjectServerId : FRemoteServerId::GetLocalServerId();
-		}		
+			Stub->ResidentServerId = ResidentServerId.IsValid() ? ResidentServerId : Stub->OwningServerId;
+		}
+		else if (ResidentServerId.IsValid())
+		{
+			// Update the ResidentServerId only if it's valid
+			Stub->ResidentServerId = ResidentServerId;
+		}
 		return Stub;
 	}
 
-	UE::RemoteObject::Handle::FRemoteObjectStub* FindOrAddRemoteObjectStub(UObject* Object, FRemoteServerId DestinationServerId)
+	UE::RemoteObject::Handle::FRemoteObjectStub* FindOrAddRemoteObjectStub(const UObject* Object, FRemoteServerId ResidentServerId = FRemoteServerId())
 	{
 		using namespace UE::RemoteObject::Handle;
+		using namespace UE::CoreUObject::Private;
 
-		UE::TScopeLock MapLock(ObjectMapCritical);
-		FRemoteObjectStub*& Stub = RemoteObjects.FindOrAdd(FRemoteObjectId(Object));
-		if (!Stub)
-		{
-			Stub = new FRemoteObjectStub(Object);
+		checkf(Object, TEXT("Unable to add a stub for a null object"));
 
-			// if we are creating the stub, then this object's owner is deduced from its ID
-			// if the server ID is invalid then it's a local native object that was created before the local server had its ID assigned
-			FRemoteServerId ObjectServerId = Stub->Id.GetServerId();
-			Stub->OwningServerId = ObjectServerId.IsValid() ? ObjectServerId : FRemoteServerId::GetLocalServerId();
-		}
-		else if (!Stub->Class.IsValid())
-		{
-			// Always make sure we store the object class in the stub (it's possible the stub was created only with a remote id in which case the class wouldn't be stored)
-			Stub->Class = FRemoteObjectClass(Object->GetClass());
-		}
+		UE::TScopeLock MapLock(ObjectMapCritical);
+		FRemoteObjectStub* Stub = FindOrAddRemoteObjectStub(FRemoteObjectId(Object), ResidentServerId);
 
+		// Always try and update the stub with the most recent information about the object
 		FUObjectItem* ObjectItem = GUObjectArray.ObjectToObjectItem(Object);
 		checkf(ObjectItem, TEXT("Attempting to get a serial number for an object that does not exist in the global UObject array (it's possible GUObjectArray is not initialized yet, ObjectIndex=%d)"), GUObjectArray.ObjectToIndex(Object));
 		Stub->SerialNumber = ObjectItem->GetSerialNumber();
 		Stub->Name = Object->GetFName();
-		Stub->ResidentServerId = DestinationServerId;
 		Stub->bWasGarbage = !IsValid(Object);
+		if (UObject* Outer = FObjectHandleUtils::GetNonAccessTrackedOuterNoResolve(Object))
+		{
+			Stub->OuterId = FObjectHandleUtils::GetRemoteId(Outer);
+		}
+		// The class of an object never changes so we only need to set it once
+		if (!Stub->Class.IsValid())
+		{
+			Stub->Class = FRemoteObjectClass(Object->GetClass());
+		}
 
 		return Stub;
 	}
@@ -368,10 +371,7 @@ void RegisterRemoteObjectId(FRemoteObjectId ObjectId, FRemoteServerId ResidentSe
 
 void RegisterSharedObject(UObject* Object)
 {
-	FUObjectItem* ObjectItem = GUObjectArray.ObjectToObjectItem(Object);
-#if UE_WITH_REMOTE_OBJECT_HANDLE
-	ObjectItem->SetFlags(EInternalObjectFlags::RemoteReference);
-#endif
+	MarkAsRemoteReference(Object);
 }
 
 void MarkAsRemote(UObject* Object, FRemoteServerId DestinationServerId)
@@ -387,7 +387,7 @@ void MarkAsRemote(UObject* Object, FRemoteServerId DestinationServerId)
 	ObjectMaps->FindOrAddRemoteObjectStub(Object, DestinationServerId);
 }
 
-void MarkAsRemoteReference(UObject* Object)
+void MarkAsRemoteReference(const UObject* Object)
 {
 #if UE_WITH_REMOTE_OBJECT_HANDLE
 	FUObjectItem* ObjectItem = GUObjectArray.ObjectToObjectItem(Object);
@@ -721,7 +721,8 @@ void ChangePhysicsIslandId(const UObject* Object, FPhysicsIslandId NewPhysicsIsl
 	using namespace UE::RemoteObject::Private;
 
 #if UE_WITH_REMOTE_OBJECT_HANDLE
-	UE::RemoteObject::Handle::FRemoteObjectStub* RemoteStub = ObjectMaps->FindOrAddRemoteObjectStub(FObjectHandleUtils::GetRemoteId(Object));
+	UE::RemoteObject::Handle::FRemoteObjectStub* RemoteStub = ObjectMaps->FindOrAddRemoteObjectStub(Object);
+	MarkAsRemoteReference(Object);
 	if (RemoteStub)
 	{
 		RemoteStub->PhysicsIslandId = NewPhysicsIslandId;
diff --git a/Engine/Source/Runtime/CoreUObject/Private/UObject/RemoteObjectPrivate.h b/Engine/Source/Runtime/CoreUObject/Private/UObject/RemoteObjectPrivate.h
--- a/Engine/Source/Runtime/CoreUObject/Private/UObject/RemoteObjectPrivate.h
+++ b/Engine/Source/Runtime/CoreUObject/Private/UObject/RemoteObjectPrivate.h
@@ -32,7 +32,7 @@ namespace UE::RemoteObject::Private
 	/**
 	* Marks object as referenced by a remote object
 	*/
-	void MarkAsRemoteReference(UObject* Object);
+	void MarkAsRemoteReference(const UObject* Object);
 
 	/**
 	* @return true if object is marked as remote reference
diff --git a/Engine/Source/Runtime/CoreUObject/Public/UObject/RemoteObject.h b/Engine/Source/Runtime/CoreUObject/Public/UObject/RemoteObject.h
--- a/Engine/Source/Runtime/CoreUObject/Public/UObject/RemoteObject.h
+++ b/Engine/Source/Runtime/CoreUObject/Public/UObject/RemoteObject.h
@@ -135,6 +135,7 @@ namespace UE::RemoteObject::Handle
 		FPhysicsIslandId PhysicsIslandId;
 
 		FRemoteObjectStub() = default;
+		UE_DEPRECATED(5.8, "Stubs should only be initialized with UObject properties when they are being added to the internal stub map.")
 		COREUOBJECT_API explicit FRemoteObjectStub(UObject* Object);
 	};
 
