Enable constinit UObjects with merged modules

- Add RemoveFromRootForUnload to remove objects from GC roots
- Replace explicit garbage marks with the new call
- Skip destruction of Native objects which are destroyed statically when the library unloads
- Fix plugin packages not included in removal searches

#rnx

#Synced-CL 48939516

diff --git a/Engine/Source/Runtime/CoreUObject/Private/UObject/GarbageCollection.cpp b/Engine/Source/Runtime/CoreUObject/Private/UObject/GarbageCollection.cpp
--- a/Engine/Source/Runtime/CoreUObject/Private/UObject/GarbageCollection.cpp
+++ b/Engine/Source/Runtime/CoreUObject/Private/UObject/GarbageCollection.cpp
@@ -649,6 +649,31 @@ namespace UE::GC::Private
 	{
 		GSkipDestroyGarbage.store(bSkip, std::memory_order_relaxed);
 	}
+
+	// When using constinit UObjects with merged modules, it becomes tricky to unload Native objects.
+	// We need to know during GC that objects are Native in order to not destroy them because they have a static lifetime.
+	// Native is part of EInternalObjectFlags_GarbageCollectionKeepFlags so we need a way to let GC know these objects are no longer relevant while keeping the flag.
+	void RemoveFromRootForUnload(TSet<const UObject*> ObjectsToRemove)
+	{
+		TRACE_CPUPROFILER_EVENT_SCOPE(RemoveFromRootForUnload);
+
+		GRootsMutex.Lock();
+
+		for (const UObject* Object : ObjectsToRemove)
+		{
+			const int32 Index = GUObjectArray.ObjectToIndex(Object);
+			if (GUObjectArray.IndexToObject(Index)->GetRefCount() > 0)
+			{
+				UE_LOG(LogGarbage, Warning, TEXT("Cannot remove refcounted object %s from roots!"), *Object->GetName());
+			}
+			else
+			{
+				GRoots.Remove(Index);
+			}
+		}
+
+		GRootsMutex.Unlock();
+	}
 	
 	void ProcessDirtyRootNoLock(int32 Index)
 	{
@@ -881,11 +906,21 @@ public:
 				UObject* Object = (UObject*)ObjectItem->GetObject();
 				check(Object->HasAllFlags(RF_FinishDestroyed | RF_BeginDestroyed));
 
-				GUObjectArray.FreeUObjectIndex(Object);
+#if UE_WITH_CONSTINIT_UOBJECT
+				if (ObjectItem->HasAllFlags(EInternalObjectFlags::Native))
+				{
+					// Skip the destructor/delete of compiled-in constinit objects
+					UnreachableObject.Object = nullptr;
+				}
+				else
+#endif // UE_WITH_CONSTINIT_UOBJECT
+				{ 
+					// Replace the entry in GUnreachableObjects with the actual UObject so that we can iterate over the same array when 
+					// we call UObject destructors and free their memory in the loop below
+					UnreachableObject.Object = Object;
+				}
 
-				// Replace the entry in GUnreachableObjects with the actual UObject so that we can iterate over the same array when 
-				// we call UObject destructors and free their memory in the loop below
-				UnreachableObject.Object = Object;
+				GUObjectArray.FreeUObjectIndex(Object);
 
 				++ProcessedObjectsCount;
 				++ObjCurrentFreeIndexObjectIndex;
@@ -911,9 +946,14 @@ public:
 				UE::GC::FUnreachableObject& UnreachableObject = GUnreachableObjects[ObjCurrentPurgeObjectIndex];
 				UObject* Object = UnreachableObject.Object;
 
-				Object->~UObject();
-				GUObjectAllocator.FreeUObject(Object);
-				UnreachableObject.Object = nullptr;
+#if UE_WITH_CONSTINIT_UOBJECT
+				if (Object)
+#endif // UE_WITH_CONSTINIT_UOBJECT
+				{
+					Object->~UObject();
+					GUObjectAllocator.FreeUObject(Object);
+					UnreachableObject.Object = nullptr;
+				}
 
 				++ProcessedObjectsCount;
 				++ObjectsDestroyedSinceLastMarkPhase;
diff --git a/Engine/Source/Runtime/CoreUObject/Private/UObject/UObjectBase.cpp b/Engine/Source/Runtime/CoreUObject/Private/UObject/UObjectBase.cpp
--- a/Engine/Source/Runtime/CoreUObject/Private/UObject/UObjectBase.cpp
+++ b/Engine/Source/Runtime/CoreUObject/Private/UObject/UObjectBase.cpp
@@ -1138,18 +1138,24 @@ static TArray<UPackage*> GetScriptPackagesFromModuleNames(TConstArrayView<FName>
 {
 	TArray<UPackage*> Packages;
 
-	Packages.Reserve(2 * InModuleNames.Num());
+	Packages.Reserve(3 * InModuleNames.Num());
 
 	for (const FName ModuleName : InModuleNames)
 	{
-		// Identify regular script packages. The expected package name is '/Script/ModuleName'.
+		// Identify module script packages. The expected package name is '/Script/ModuleName'.
 		if (UPackage* Package = FindObjectFast<UPackage>(/* Outer */ nullptr, WriteToString<256>(TEXT("/Script/"), ModuleName).ToString()))
 		{
 			Packages.Add(Package);
 		}
 
-		// Identify Verse packages. The expected package name is '/PluginName/_Verse/VNI/ModuleName'.
+		// Identify plugin script packages. The expected package name is '/Script/PluginName'.
 		FName PluginName = IPluginManager::Get().PackageNameFromModuleName(ModuleName);
+		if (UPackage* Package = FindObjectFast<UPackage>(/* Outer */ nullptr, WriteToString<256>(TEXT("/Script/"), PluginName).ToString()))
+		{
+			Packages.AddUnique(Package);
+		}
+
+		// Identify Verse packages. The expected package name is '/PluginName/_Verse/VNI/ModuleName'.
 		if (UPackage* Package = FindObjectFast<UPackage>(/* Outer */ nullptr, WriteToString<256>(TEXT("/"), PluginName, TEXT("/_Verse/VNI/"), ModuleName).ToString()))
 		{
 			Packages.Add(Package);
@@ -1165,38 +1171,37 @@ static void RemoveObjectsInModules(TConstArrayView<UPackage*> Packages)
 	LLM_SCOPE_BYTAG(UObject_UObjectBase);
 	DECLARE_SCOPE_CYCLE_COUNTER(TEXT("RemoveObjectsInModules"), STAT_RemoveObjectsInModules, STATGROUP_ObjectVerbose);
 
-	FPermanentObjectPoolExtents PermanentPool;
+	TSet<const UObject*> ObjectsToRemove;
+	const FPermanentObjectPoolExtents PermanentPool;
 
 	for (UPackage* Package : Packages)
 	{
 		UE_LOG(LogUObjectBase, Log, TEXT("RemoveObjectsInModules: removing package %s"), *Package->GetName());
 
-		ForEachObjectWithPackage(Package, [&PermanentPool](UObject* InObject) mutable -> bool
+		ForEachObjectWithPackage(Package, [&PermanentPool, &ObjectsToRemove](UObject* InObject) mutable -> bool
 		{
 			if (!PermanentPool.Contains(InObject) && !GUObjectArray.IsDisregardForGC(InObject))
 			{
 				if (InObject->IsA<UClass>())
 				{
-					ForEachObjectOfClass(Cast<UClass>(InObject), [InObject](UObject* InClassObject) -> bool
+					ForEachObjectOfClass(Cast<UClass>(InObject), [InObject, &ObjectsToRemove](UObject* InClassObject) -> bool
 					{
-						InClassObject->ClearInternalFlags(EInternalObjectFlags::RootSet | EInternalObjectFlags::Native);
-						InClassObject->ClearFlags(RF_Standalone);
-						InClassObject->MarkAsGarbage();
-
+						ObjectsToRemove.Add(InClassObject);
 						return true;
 					}, true);
 				}
 
-				InObject->ClearInternalFlags(EInternalObjectFlags::RootSet | EInternalObjectFlags::Native);
-				InObject->ClearFlags(RF_Standalone);
-				InObject->MarkAsGarbage();
+				ObjectsToRemove.Add(InObject);
 			}
 
 			return true;
 		}, true);
 
 		UScriptStruct::RemoveAllDeferredCppStructOps(*Package->GetName());
+		ObjectsToRemove.Add(Package);
 	}
+
+	UE::GC::Private::RemoveFromRootForUnload(ObjectsToRemove);
 }
 
 static bool CheckLiveObjectsInModules(TConstArrayView<FName> InModuleNames)
diff --git a/Engine/Source/Runtime/CoreUObject/Public/UObject/GarbageCollection.h b/Engine/Source/Runtime/CoreUObject/Public/UObject/GarbageCollection.h
--- a/Engine/Source/Runtime/CoreUObject/Public/UObject/GarbageCollection.h
+++ b/Engine/Source/Runtime/CoreUObject/Public/UObject/GarbageCollection.h
@@ -209,4 +209,5 @@ COREUOBJECT_API bool IsGarbageCollectingAndLockingUObjectHashTables();
 namespace UE::GC::Private
 {
 	COREUOBJECT_API void SetSkipDestroy(bool bSkip);
+	COREUOBJECT_API void RemoveFromRootForUnload(TSet<const UObject*> ObjectsToRemove);
 }
