- Refactored UObject::LoadConfig into some Context objets to break up the wall of code and also be able to more easily add functionality (like passing a preloaded FConfigFile to LoadConfig)

#Synced-CL 48934447

diff --git a/Engine/Source/Runtime/CoreUObject/Private/UObject/Obj.cpp b/Engine/Source/Runtime/CoreUObject/Private/UObject/Obj.cpp
--- a/Engine/Source/Runtime/CoreUObject/Private/UObject/Obj.cpp
+++ b/Engine/Source/Runtime/CoreUObject/Private/UObject/Obj.cpp
@@ -2868,651 +2868,759 @@ void CheckMissingSection(const FString& SectionName, const FString& IniFilename)
 static FName GConsoleVariableFName("ConsoleVariable");
 #endif
 
+
+
 void UObject::LoadConfig(UClass* ConfigClass/*=nullptr*/, const TCHAR* InFilename/*=nullptr*/,
 	uint32 PropagationFlags/*=LCPF_None*/, FProperty* PropertyToLoad/*=nullptr*/,
 	TArray<UE::ConfigAccessTracking::FConfigAccessData>* OutAccessedValues/*=nullptr*/)
 {
-	SCOPE_CYCLE_COUNTER(STAT_LoadConfig);
+	FLoadConfigParams Context;
+	Context.ConfigClass = ConfigClass;
+	Context.Filename = InFilename;
+	Context.PropagationFlags = PropagationFlags;
+	Context.PropertyToLoad = PropertyToLoad;
+	Context.AccessedValues = OutAccessedValues;
 
-	// OriginalClass is the class that LoadConfig() was originally called on
-	static thread_local UClass* OriginalClass = nullptr;
+	LoadConfig(Context);
+}
 
-	if (!ConfigClass)
+class FPropertyLoadConfigContext
+{
+public:
+	FPropertyLoadConfigContext(const UObject::FLoadConfigParams& InLoadParams, UObject* InObject)
+		: LoadParams(InLoadParams)
+		, ConfigSystem(LoadParams.ConfigSystem ? LoadParams.ConfigSystem : GConfig)
+		, ActiveObject(InObject)
+		, bTrackingAccessOnly(LoadParams.AccessedValues != nullptr) // if we are tracking, we don't acutally LoadConfig onto properties
 	{
-		// if no class was specified in the call, this is the OriginalClass
-		ConfigClass = GetClass();
-		OriginalClass = ConfigClass;
 	}
-
-	if (!ConfigClass->HasAnyClassFlags(CLASS_Config))
+	
+	bool InitContext()
 	{
-		return;
-	}
-
+		// all must pass for the Init to succeed, and for LoadConfig to continue
+		if (PrepareObject() && Propagate() && PrepareConfig())
+		{
+			if (LoadParams.PropertyToLoad == nullptr)
+			{
+				UE_LOG(LogConfig, VeryVerbose, TEXT("(%s) '%s' loading configuration from %s"),
+					*ConfigClass->GetName(), *ActiveObject->GetName(), *Filename);
+			}
+			else
+			{
+				UE_LOG(LogConfig, VeryVerbose, TEXT("(%s) '%s' loading configuration for property %s from %s"),
+					*ConfigClass->GetName(), *ActiveObject->GetName(), *LoadParams.PropertyToLoad->GetName(), *Filename);
+			}
 
-#if DO_CHECK
-	// When a CDO is created before the Class is linked, it will have an empty set of properties. Among other issues, this can cause loading of config data to fail silently. 
-	if (!ConfigClass->DebugIsPropertyChainReady())
+			return true;
+		}
+		
+		return false;
+	}
+	
+	bool InitForProperty(const FProperty* Property)
 	{
-		FString ErrorMessage = FString::Printf(TEXT("Attempting to load config data for %s before the Class has been constructed/registered/linked (likely during module loading or early startup). This will result in the load silently failing and should be fixed."), *GetName());
-		if (GPrematureConfigLoadErrorLevel == 0)
+		if (!Property->HasAnyPropertyFlags(CPF_Config))
 		{
-			UE_LOG(LogObj, Display, TEXT("%s"), *ErrorMessage);
+			return false;
 		}
-		else if (GPrematureConfigLoadErrorLevel == 1)
+
+		// if we're only supposed to load the value for a specific property, skip all others
+		if (LoadParams.PropertyToLoad != nullptr && LoadParams.PropertyToLoad != Property)
 		{
-			UE_LOG(LogObj, Warning, TEXT("%s"), *ErrorMessage);
+			return false;
 		}
-		else if (GPrematureConfigLoadErrorLevel == 2)
+		
+		const bool bGlobalConfig = EnumHasAnyFlags(Property->PropertyFlags, CPF_GlobalConfig);
+		FString ModifiedConfigFilename;
+		
+		// if the property is GlobalConfig, then the class to load the property from is the class that owns the property
+		// so we may need to look in a different section
+		if (bGlobalConfig && Property->GetOwnerClass() != ConfigClass)
 		{
-			// Repeating the condition here just so it ends up in the output message
-			ensureAlwaysMsgf(ConfigClass->DebugIsPropertyChainReady(), TEXT("%s"), *ErrorMessage);
+			const UClass* OwnerClass = Property->GetOwnerClass();
+			
+			if (!bPerObject)
+			{
+				ClassSection = OwnerClass->GetPathName();
+			}
+			// if we forced a filename, then we will use the ConfigFile already found for the class
+			if (LoadParams.Filename.IsEmpty())
+			{
+				// look up in the owning class' branch, which in rare cases is different than this class's branch
+				// however, to see if it's different is a slowish function call and a string compare, so just always look up
+				ModifiedConfigFilename = OwnerClass->GetConfigName();
+				ConfigFile = ConfigSystem->Find(ModifiedConfigFilename);				
+			}
+			else
+			{
+				ConfigFile = ClassConfigFile;
+			}
+//			
 		}
 		else
 		{
-			checkf(ConfigClass->DebugIsPropertyChainReady(), TEXT("%s"), *ErrorMessage);
+			// @todo precalc this in the PrepConfig above?
+			if (!bPerObject)
+			{
+				ClassSection = ConfigClass->GetPathName();
+			}
+			ConfigFile = ClassConfigFile;
+		}
+		
+		if (!bPerObject)
+		{
+			// allow the class to override the expected section name
+			ActiveObject->OverrideConfigSection(ClassSection);
 		}
-	}
-#endif
 
-#if !WITH_EDITOR
-	if (OutAccessedValues)
-	{
-		return;
-	}
+		Key = Property->GetName();
+		PortFlags = EPropertyPortFlags::PPF_SerializedAsImportText;
+
+#if WITH_EDITOR
+		const FString& CVarName = Property->GetMetaData(GConsoleVariableFName);
+		if (!CVarName.IsEmpty())
+		{
+			Key = CVarName;
+			PortFlags |= PPF_ConsoleVariable;
+		}
+
+		if (bTrackingAccessOnly)
+		{
+			const FString& PropFileName = ModifiedConfigFilename.IsEmpty() ? Filename : ModifiedConfigFilename;
+			LoadParams.AccessedValues->Add(UE::ConfigAccessTracking::FConfigAccessData(
+				UE::ConfigAccessTracking::ELoadType::ConfigSystem, ConfigPlatform, FName(*PropFileName),
+				FName(*ClassSection), FName(*Key), nullptr /* InRequestingPlatform */));
+			
+			return false;
+		}
 #endif
-	const bool bModifyingObjects = OutAccessedValues == nullptr;
 
-#if !IS_PROGRAM
-	auto HaveSameProperties = [](const UStruct* Struct1, const UStruct* Struct2) -> bool
-	{
-		TFieldIterator<FProperty> It1(Struct1);
-		TFieldIterator<FProperty> It2(Struct2);
-		for (;; ++It1, ++It2)
+		ConfigSection = ConfigFile->FindSection(ClassSection);
+		// look in alternate section name for POC
+		if (ConfigSection == nullptr && bPerObject && ClassPathSection.Len())
 		{
-			bool bAtEnd1 = !It1;
-			bool bAtEnd2 = !It2;
-			// If one iterator is at the end and one isn't, the property lists are different
-			if (bAtEnd1 != bAtEnd2)
-			{
-				return false;
-			}
-			// If both iterators have reached the end, the property lists are the same
-			if (bAtEnd1)
-			{
-				return true;
-			}
-			// If the properties are different, the property lists are different
-			if (*It1 != *It2)
-			{
-				return false;
-			}
+			ConfigSection = ConfigFile->FindSection(ClassPathSection);
+		}
+
+#if !UE_BUILD_SHIPPING
+		if (ConfigSection == nullptr && !FPlatformProperties::RequiresCookedData())
+		{
+			const FString& PropFileName = ModifiedConfigFilename.IsEmpty() ? Filename : ModifiedConfigFilename;
+			CheckMissingSection(ClassSection, PropFileName);
 		}
-	};
-	// Do we have properties that don't exist yet?
-	// If this happens then we're trying to load the config for an object that doesn't
-	// know what its layout is. Usually a call to GetDefaultObject that occurs too early
-	// because ProcessNewlyLoadedUObjects hasn't happened yet
-	checkf(ConfigClass->PropertyLink != nullptr
-		|| (ConfigClass->GetSuperStruct() && HaveSameProperties(ConfigClass, ConfigClass->GetSuperStruct()))
-		|| (UE_WITH_CONSTINIT_UOBJECT ? !ConfigClass->HasAnyClassFlags(CLASS_Constructed) : ConfigClass->PropertiesSize == 0)
-		|| IsEngineExitRequested(), // Ignore this check when exiting as we may have requested exit during init when not everything is initialized
-		TEXT("class %s has uninitialized properties. Accessed too early?"), *ConfigClass->GetPathName());
 #endif
 
-	UClass* ParentClass = ConfigClass->GetSuperClass();
-	if (ParentClass != nullptr)
+		UE_LOG(LogConfig, VeryVerbose, TEXT("   Loading value for %s from [%s]"), *Key, *ClassSection);
+
+		// if we didn't find a section, no need to process this property - 
+		// @todo: can we actually return false? some code may want to empty the property then skip if no Section was found
+//		return ConfigSection == nullptr;
+		return true;
+	}
+
+	bool ProcessArrayProperty(const FArrayProperty* ArrayProperty)
 	{
-		if (ParentClass->HasAnyClassFlags(CLASS_Config))
+		FScriptArrayHelper_InContainer ArrayHelper(ArrayProperty, ActiveObject);
+
+#if WITH_EDITOR
+		// Empty out any array properties if this is a PerPlatformConfig class
+		// as we are replacing the values with the Platform's version when entering
+		// a new preview platform.
+		if (bIsPerPlatformConfig)
 		{
-			if ((PropagationFlags & UE::LCPF_ReadParentSections) != 0)
-			{
-				// call LoadConfig on the parent class
-				LoadConfig(ParentClass, nullptr, PropagationFlags, PropertyToLoad, OutAccessedValues);
+			ArrayHelper.EmptyValues();
+		}
+#endif
+		if (ConfigSection)
+		{
+			TArray<FConfigValue> List;
+			const FName KeyName(*Key, FNAME_Find);
+			ConfigSection->MultiFind(KeyName, List);
+
+			const int32 Size = ArrayProperty->Inner->GetElementSize();
 
-				// if we are also notifying child classes or instances, stop here as this object's properties will be imported as a result of notifying the others
-				if ((PropagationFlags & (UE::LCPF_PropagateToChildDefaultObjects | UE::LCPF_PropagateToInstances)) != 0)
+			// Only override default properties if there is something to override them with, or
+			// if the config system had a "initialize to empty" entry
+			if (!List.IsEmpty() || ConfigSection->EmptyInitializedKeys.Contains(KeyName))
+			{
+				ArrayHelper.EmptyAndAddValues(List.Num());
+				for (int32 i = List.Num() - 1, c = 0; i >= 0; i--, c++)
 				{
-					return;
+					ArrayProperty->Inner->ImportText_Direct(*List[i].GetValue(), ArrayHelper.GetRawPtr(c), ActiveObject, PortFlags);
 				}
 			}
-			else if ((PropagationFlags & UE::LCPF_PropagateToChildDefaultObjects) != 0)
+			else
 			{
-				// not propagating the call upwards, but we are propagating the call to all child classes
-				for (TObjectIterator<UClass> It; It; ++It)
+				int32 Index = 0;
+				const FConfigValue* ElementValue = nullptr;
+				do
 				{
-					if (It->IsChildOf(ConfigClass))
+					// Add array index number to end of key
+					FString IndexedKey = FString::Printf(TEXT("%s[%i]"), *Key, Index);
+
+					// Try to find value of key
+					const FName IndexedName(*IndexedKey, FNAME_Find);
+					if (IndexedName == NAME_None)
 					{
-						// mask out the PropgateToParent and PropagateToChildren values
-						It->GetDefaultObject()->LoadConfig(*It, nullptr,
-							(PropagationFlags & (UE::LCPF_PersistentFlags | UE::LCPF_PropagateToInstances)),
-							PropertyToLoad, OutAccessedValues);
+						break;
 					}
-				}
+					ElementValue = ConfigSection->Find(IndexedName);
 
-				// LoadConfig() was called on this object during iteration, so stop here 
-				return;
-			}
-			else if ((PropagationFlags & UE::LCPF_PropagateToInstances) != 0)
-			{
-				// call LoadConfig() on all instances of this class (except the CDO)
-				// Do not propagate this call to parents, and do not propagate to children or instances (would be redundant) 
-				for (TObjectIterator<UObject> It; It; ++It)
-				{
-					if (It->IsA(ConfigClass))
+					// If found, import the element
+					if (ElementValue != nullptr)
 					{
-						if (!GIsEditor)
-						{
-							// make sure to pass in the class so that OriginalClass isn't reset
-							It->LoadConfig(It->GetClass(), nullptr, (PropagationFlags & UE::LCPF_PersistentFlags),
-								PropertyToLoad, OutAccessedValues);
-						}
-#if WITH_EDITOR
-						else
-						{
-							if (bModifyingObjects)
-							{
-								It->PreEditChange(nullptr);
-							}
-
-							// make sure to pass in the class so that OriginalClass isn't reset
-							It->LoadConfig(It->GetClass(), nullptr, (PropagationFlags & UE::LCPF_PersistentFlags),
-								PropertyToLoad, OutAccessedValues);
-
-							if (bModifyingObjects)
-							{
-								It->PostEditChange();
-							}
-						}
-#endif // WITH_EDITOR
+						// expand the array if necessary so that Index is a valid element
+						ArrayHelper.ExpandForIndex(Index);
+						ArrayProperty->Inner->ImportText_Direct(*ElementValue->GetValue(), ArrayHelper.GetRawPtr(Index), ActiveObject, PortFlags);
 					}
-				}
+
+					Index++;
+				} while (ElementValue || Index < ArrayHelper.Num());
 			}
 		}
-		else if ((PropagationFlags & UE::LCPF_PropagateToChildDefaultObjects) != 0)
-		{
-			// we're at the base-most config class
-			for (TObjectIterator<UClass> It; It; ++It)
-			{
-				if (It->IsChildOf(ConfigClass))
-				{
-					if (!GIsEditor)
-					{
-						// make sure to pass in the class so that OriginalClass isn't reset
-						It->GetDefaultObject()->LoadConfig(*It, nullptr,
-							(PropagationFlags & (UE::LCPF_PersistentFlags | UE::LCPF_PropagateToInstances)),
-							PropertyToLoad, OutAccessedValues);
-					}
+		
+		return true;
+	}
+	
+	bool ProcessSetProperty(const FSetProperty* SetProperty)
+	{
+		bool bProcessedProperty = true;
+
+		FScriptSetHelper_InContainer SetHelper(SetProperty, ActiveObject);
+
 #if WITH_EDITOR
-					else
-					{
-						if (bModifyingObjects)
-						{
-							It->PreEditChange(nullptr);
-						}
+		// Empty out any set properties if this is a PerPlatformConfig class
+		// as we are replacing the values with the Platform's version when entering
+		// a new preview platform.
+		if (bIsPerPlatformConfig)
+		{
+			SetHelper.EmptyElements();
+		}
+#endif
 
-						// make sure to pass in the class so that OriginalClass isn't reset
-						It->GetDefaultObject()->LoadConfig(*It, nullptr,
-							(PropagationFlags & (UE::LCPF_PersistentFlags | UE::LCPF_PropagateToInstances)),
-							PropertyToLoad, OutAccessedValues);
+		if (ConfigSection)
+		{
+			TArray<FConfigValue> List;
+			const FName KeyName(*Key, FNAME_Find);
+			ConfigSection->MultiFind(KeyName, List);
 
-						if (bModifyingObjects)
-						{
-							It->PostEditChange();
-						}
-					}
-#endif // WITH_EDITOR
+			bool bSingleSetEntry = false;
+			if (List.Num() == 1)
+			{
+				const FString& SingleListValue = List[0].GetValue();
+				if (SingleListValue.Len() > 1
+					&& SingleListValue[0] == TEXT('(')
+					&& SingleListValue[SingleListValue.Len() - 1] == TEXT(')'))
+				{
+					// If we have a single entry in the set that is surrounded with parentheses, fall back to
+					// the old processing method
+					bSingleSetEntry = true;
 				}
 			}
 
-			return;
-		}
-		else if ((PropagationFlags & UE::LCPF_PropagateToInstances) != 0)
-		{
-			for (TObjectIterator<UObject> It; It; ++It)
+			// Only override default properties if there is something to override them with.
+			if (!bSingleSetEntry && !List.IsEmpty())
 			{
-				if (It->GetClass() == ConfigClass)
+				bProcessedProperty = true;
+				SetHelper.EmptyElements(List.Num());
+
+				// Each config value entry can possibly specify multiple set elements - create a temporary
+				// set here that we can import to and add its elements to the object's property
+				void* TempSet = FMemory::Malloc(SetProperty->GetSize(), SetProperty->GetMinAlignment());
+				SetProperty->InitializeValue(TempSet);
+
+				// Importing elements to the set inline does not check for duplicates - create a temporary
+				// element that we can import to add uniquely to the set
+				void* TempElement = FMemory::Malloc(SetProperty->ElementProp->GetSize(),
+					SetProperty->ElementProp->GetMinAlignment());
+				SetProperty->ElementProp->InitializeValue(TempElement);
+
+				for (const FConfigValue& ListValue : List)
 				{
-					if (!GIsEditor)
+					const FString& ListString = ListValue.GetValue();
+
+					// Try to import the config value as an entire set first
+					const TCHAR* SetImportResult =
+						SetProperty->ImportText_Direct(*ListString, TempSet, ActiveObject, PortFlags);
+
+					if (SetImportResult && SetImportResult != *ListString)
 					{
-						// make sure to pass in the class so that OriginalClass isn't reset
-						It->LoadConfig(It->GetClass(), nullptr,	(PropagationFlags & UE::LCPF_PersistentFlags),
-							PropertyToLoad, OutAccessedValues);
+						FScriptSetHelper TempSetHelper(SetProperty, TempSet);
+						for (FScriptSetHelper::FIterator Itr = TempSetHelper.CreateIterator(); Itr; ++Itr)
+						{
+							SetHelper.AddElement(TempSetHelper.GetElementPtr(Itr));
+						}
 					}
-#if WITH_EDITOR
 					else
 					{
-						if (bModifyingObjects)
-						{
-							It->PreEditChange(nullptr);
-						}
+						// If we failed to import the value as an entire set, try to import it as a single
+						// element
+						const TCHAR* ElementImportResult =
+						SetProperty->ElementProp->ImportText_Direct(*ListString, TempElement, ActiveObject, PortFlags);
 
-						// make sure to pass in the class so that OriginalClass isn't reset
-						It->LoadConfig(It->GetClass(), nullptr, (PropagationFlags & UE::LCPF_PersistentFlags),
-							PropertyToLoad, OutAccessedValues);
-
-						if (bModifyingObjects)
+						if (ElementImportResult && ElementImportResult != *ListString)
 						{
-							It->PostEditChange();
+							SetHelper.AddElement(TempElement);
 						}
 					}
-#endif // WITH_EDITOR
-				}
-			}
-		}
-	}
-
-	// if a filename was specified, load all properties from that file no matter which class they come from
-	const FString Filename = InFilename ? InFilename : GetConfigFilename(this);
 
-	const bool bPerObject = UsesPerObjectConfig(this);
-
-	// does the class want to override the platform hierarchy (ignored if we passed in a specific ini file),
-	// and if the name isn't the current running platform (no need to load extra files if already in GConfig)
-	bool bUseConfigOverride = false;
-	FConfigFile* OverrideConfigFile = nullptr;
-	FConfigFile LocalOverrideConfig;
-	const TCHAR* ConfigOverridePlatform = GetConfigOverridePlatform();
-#if WITH_EDITOR
-	FName ConfigPlatform;
-#endif
+				}
 
-	if (InFilename == nullptr && ConfigOverridePlatform != nullptr &&
-		FCString::Stricmp(ConfigOverridePlatform, ANSI_TO_TCHAR(FPlatformProperties::IniPlatformName())) != 0)
-	{
-#if WITH_EDITOR
-		ConfigPlatform = FName(ConfigOverridePlatform);
-#endif
-		if (bModifyingObjects)
-		{
-			OverrideConfigFile = FConfigCacheIni::FindOrLoadPlatformConfig(LocalOverrideConfig,
-				*GetClass()->ClassConfigName.ToString(), ConfigOverridePlatform);
-			bUseConfigOverride = true;
+				SetProperty->DestroyAndFreeValue(TempSet);
+				SetProperty->ElementProp->DestroyAndFreeValue(TempElement);
+				SetHelper.Rehash();
+				
+				return true;
+			}
 		}
+		
+		// if we get here, we didn't process the property as a Set, so return false, and we will try again 
+		// with ProcessStandardProperty
+		return false;
 	}
-#if WITH_EDITOR
-	else if (GetClass()->HasAnyClassFlags(CLASS_PerPlatformConfig) && UObject::OnGetPreviewPlatform.IsBound())
+	
+	void ProcessStandardProperty(FProperty* Property)
 	{
-		FName PreviewPlatform;
-		if (UObject::OnGetPreviewPlatform.Execute(PreviewPlatform))
+		for (int32 i = 0; i < Property->ArrayDim; i++)
 		{
-			FString PreviewPlatformStr = PreviewPlatform.ToString();
-			ConfigPlatform = FName(*PreviewPlatformStr);
-			if (bModifyingObjects)
+			// @todo can we optimize this somehow?
+			if (Property->ArrayDim != 1)
 			{
-				OverrideConfigFile = FConfigCacheIni::FindOrLoadPlatformConfig(LocalOverrideConfig,
-					*GetClass()->ClassConfigName.ToString(), *PreviewPlatformStr);
-				bUseConfigOverride = true;
+				Key = FString::Printf(TEXT("%s[%i]"), *Property->GetName(), i);
 			}
-		}
-	}
-#endif
-
-	FString ClassSection;
-	FString ClassPathSection;
-	FName LongCommitName;
 
-	if (bPerObject)
+			if (ConfigSection)
+			{
+				FString Value;
+				if (ConfigSection->GetString(*Key, Value))
+				{
+					if (Property->ImportText_Direct(*Value, Property->ContainerPtrToValuePtr<uint8>(ActiveObject, i), ActiveObject, PortFlags) == nullptr)
+					{
+						// this should be an error as the properties from the .ini / .int file are not correctly being read in and probably are affecting things in subtle ways
+						UE_LOG(LogObj, Error, TEXT("LoadConfig (%s): import failed for %s in: %s"), *ActiveObject->GetPathName(), *Property->GetName(), *Value);
+					}
+				}
+			}
+		}		
+	}
+private:	
+	
+	bool PrepareObject()
 	{
-		FString PathNameString;
-		UObject* Outermost = GetOutermost();
-
-		if (Outermost == GetTransientPackage())
-		{
-			PathNameString = GetName();
-		}
-		else
+		ConfigClass = LoadParams.ConfigClass ? LoadParams.ConfigClass : ActiveObject->GetClass();
+		if (!ConfigClass->HasAnyClassFlags(CLASS_Config))
 		{
-			GetPathName(Outermost, PathNameString);
-			LongCommitName = Outermost->GetFName();
+			return false;
 		}
 
-		ClassSection = PathNameString + TEXT(" ") + GetClass()->GetName();
+		bPerObject = UsesPerObjectConfig(ActiveObject);
+		bIsPerPlatformConfig = ActiveObject->GetClass()->HasAnyClassFlags(CLASS_PerPlatformConfig);
 
-		FString OverrideClassSection;
-		OverridePerObjectConfigSection(OverrideClassSection);
-		if (OverrideClassSection.Len() && OverrideClassSection != ClassSection)
+		// handle some special cases and validations
+		if (bPerObject)
 		{
-			// If a we got a section name override no need to perform short class name checks
-			ClassSection = MoveTemp(OverrideClassSection);
-			// Keep ClassPathSection empty so that we don't check for it when the section name has been overridden
+			HandlePerObject();
 		}
-		else
+		CheckEarlyAccess();
+		
+		// we only support tracking values in with the editor
+		if (!WITH_EDITOR && bTrackingAccessOnly)
 		{
-			// Cache both version of per object config section name
-			ClassPathSection = PathNameString + TEXT(" ") + GetClass()->GetPathName();
+			return false;
 		}
+		
+		return true;
 	}
-
-	// If any of my properties are class variables, then LoadConfig() would also be called for each one of those classes.
-	// Since OriginalClass is a static variable, if the value of a class variable is a class different from the current class, 
-	// we'll lose our nice reference to the original class - and cause any variables which were declared after this class variable to fail 
-	// the 'if (OriginalClass != Class)' check....better store it in a temporary place while we do the actual loading of our properties 
-	UClass* MyOrigClass = OriginalClass;
-
-	if (PropertyToLoad == nullptr)
-	{
-		UE_LOG(LogConfig, VeryVerbose, TEXT("(%s) '%s' loading configuration from %s"),
-			*ConfigClass->GetName(), *GetName(), *Filename);
-	}
-	else
-	{
-		UE_LOG(LogConfig, VeryVerbose, TEXT("(%s) '%s' loading configuration for property %s from %s"),
-			*ConfigClass->GetName(), *GetName(), *PropertyToLoad->GetName(), *Filename);
-	}
-
-	auto GetConfigValue = [&OverrideConfigFile, &bUseConfigOverride]
-	(const TCHAR* ClassSection, const TCHAR* Key, const TCHAR* ConfigName, FString& OutValue)
+	
+	bool PrepareConfig()
 	{
-		if (bUseConfigOverride)
+		if (LoadParams.OverrideFile != nullptr)
 		{
-			return OverrideConfigFile->GetString(ClassSection, Key, OutValue);
-		}
-		else
-		{
-			return GConfig->GetString(ClassSection, Key, OutValue, ConfigName);
+			checkf(LoadParams.Filename.IsEmpty(), TEXT("Cannot have FLoadConfigParams::OverrideFile and FLoadConfigParams::Filename both set, it indicates a logic issue."));
+			ClassConfigFile = LoadParams.OverrideFile;
+			return true;
 		}
-	};
+		
+		// if a filename was specified, load all properties from that file no matter which class they come from
+		Filename = LoadParams.Filename.IsEmpty() ? GetConfigFilename(ActiveObject) : LoadParams.Filename;
 
-	auto GetConfigSection = [&OverrideConfigFile, &bUseConfigOverride]
-	(const TCHAR* SectionName, const TCHAR* ConfigFilename)
-	{
-		if (bUseConfigOverride)
-		{
-			return OverrideConfigFile->FindSection(SectionName);
-		}
-		else
-		{
-			return GConfig->GetSection(SectionName, false, ConfigFilename);
-		}
-	};
+		// does the class want to override the platform hierarchy (ignored if we passed in a specific ini file),
+		// and if the name isn't the current running platform (no need to load extra files if already in GConfig)
+		const TCHAR* ConfigOverridePlatform = ActiveObject->GetConfigOverridePlatform();
 
-	for (FProperty* Property = ConfigClass->PropertyLink; Property; Property = Property->PropertyLinkNext)
-	{
+		if (LoadParams.Filename.IsEmpty() && ConfigOverridePlatform != nullptr &&
+			FCString::Stricmp(ConfigOverridePlatform, ANSI_TO_TCHAR(FPlatformProperties::IniPlatformName())) != 0)
+		{
 #if WITH_EDITOR
-		FSoftObjectPathSerializationScope SerializationScope(NAME_None, Property->GetFName(),
-			Property->GetOwnerProperty()->HasMetaData(FSoftObjectPath::NAME_Untracked) ? ESoftObjectPathCollectType::NeverCollect :
-			Property->IsEditorOnlyProperty() ? ESoftObjectPathCollectType::EditorOnlyCollect : ESoftObjectPathCollectType::AlwaysCollect,
-			ESoftObjectPathSerializeType::AlwaysSerialize);
+			ConfigPlatform = FName(ConfigOverridePlatform);
 #endif
-
-		if (!Property->HasAnyPropertyFlags(CPF_Config))
-		{
-			continue;
-		}
-
-		// if we're only supposed to load the value for a specific property, skip all others
-		if (PropertyToLoad != nullptr && PropertyToLoad != Property)
-		{
-			continue;
-		}
-
-		const bool bGlobalConfig = EnumHasAnyFlags(Property->PropertyFlags, CPF_GlobalConfig);
-		UClass* OwnerClass = Property->GetOwnerClass();
-
-		UClass* BaseClass = bGlobalConfig ? OwnerClass : ConfigClass;
-		if (!bPerObject)
-		{
-			ClassSection = BaseClass->GetPathName();
-			LongCommitName = BaseClass->GetOutermost()->GetFName();
-
-			// allow the class to override the expected section name
-			OverrideConfigSection(ClassSection);
-		}
-
-		// globalconfig properties should always use the owning class's config file
-		// specifying a value for InFilename will override this behavior (as it does with normal properties)
-		const FString* PropFileNamePtr = &Filename;
-		FString Temp;
-		if (bGlobalConfig && InFilename == nullptr)
-		{
-			Temp = OwnerClass->GetConfigName();
-			PropFileNamePtr = &Temp;
+			if (!bTrackingAccessOnly)
+			{
+				ClassConfigFile = FConfigCacheIni::FindOrLoadPlatformConfig(PlatformOverrideConfig,
+					*ActiveObject->GetClass()->ClassConfigName.ToString(), ConfigOverridePlatform);
+//				bUseConfigOverride = true;
+			}
 		}
-		const FString& PropFileName = *PropFileNamePtr;
-
-		FString Key = Property->GetName();
-		int32 PortFlags = EPropertyPortFlags::PPF_SerializedAsImportText;
-
 #if WITH_EDITOR
-		const FString& CVarName = Property->GetMetaData(GConsoleVariableFName);
-		if (!CVarName.IsEmpty())
+		else if (bIsPerPlatformConfig && UObject::OnGetPreviewPlatform.IsBound())
 		{
-			Key = CVarName;
-			PortFlags |= PPF_ConsoleVariable;
+			FName PreviewPlatform;
+			if (UObject::OnGetPreviewPlatform.Execute(PreviewPlatform))
+			{
+				FString PreviewPlatformStr = PreviewPlatform.ToString();
+				ConfigPlatform = FName(*PreviewPlatformStr);
+				if (!bTrackingAccessOnly)
+				{
+					ClassConfigFile = FConfigCacheIni::FindOrLoadPlatformConfig(PlatformOverrideConfig,
+						*ActiveObject->GetClass()->ClassConfigName.ToString(), *PreviewPlatformStr);
+//					bUseConfigOverride = true;
+				}
+			}
 		}
-
-		const bool bIsPerPlatformConfig = GetClass()->HasAnyClassFlags(CLASS_PerPlatformConfig);
-#endif // #if WITH_EDITOR
-
-		if (OutAccessedValues)
-		{
-#if WITH_EDITOR
-			OutAccessedValues->Add(UE::ConfigAccessTracking::FConfigAccessData(
-				UE::ConfigAccessTracking::ELoadType::ConfigSystem, ConfigPlatform, FName(*PropFileName),
-				FName(*ClassSection), FName(*Key), nullptr /* InRequestingPlatform */));
 #endif
-		}
-		if (!bModifyingObjects)
+		
+		// if we didn't have one specified already, we will use the ConfigSystem's branch for this class 
+		if (ClassConfigFile == nullptr)
 		{
-			continue;
+			ClassConfigFile = ConfigSystem->Find(Filename);
 		}
+		
+		// if we can't find the branch, nothing to do here (unless tracking, then let it continue tracking what would have been read)
+		return ClassConfigFile != nullptr || bTrackingAccessOnly;
+	}
 
-		// Track if we loaded this config value using special handling (e.g. array or set)
-		bool bProcessedProperty = false;
 
-		UE_LOG(LogConfig, VeryVerbose, TEXT("   Loading value for %s from [%s]"), *Key, *ClassSection);
 
-		FArrayProperty* Array = CastField<FArrayProperty>(Property);
-		FSetProperty* SetProperty = CastField<FSetProperty>(Property);
 
-		if ((Array || SetProperty))
+	bool Propagate()
+	{
+		UClass* ParentClass = ConfigClass->GetSuperClass();
+		if (ParentClass != nullptr)
 		{
-			const FConfigSection* Sec = GetConfigSection(*ClassSection, *PropFileName);
-			if (!Sec && bPerObject && ClassPathSection.Len())
-			{
-				Sec = GetConfigSection(*ClassPathSection, *PropFileName);
-			}
-
-#if !UE_BUILD_SHIPPING
-			if (!Sec && !FPlatformProperties::RequiresCookedData())
-			{
-				CheckMissingSection(ClassSection, PropFileName);
-			}
-#endif
-
-			if (Array)
+			if (ParentClass->HasAnyClassFlags(CLASS_Config))
 			{
-				FScriptArrayHelper_InContainer ArrayHelper(Array, this);
-
-				bProcessedProperty = true;
-#if WITH_EDITOR
-				// Empty out any array properties if this is a PerPlatformConfig class
-				// as we are replacing the values with the Platform's version when entering
-				// a new preview platform.
-				if (bIsPerPlatformConfig)
+				if ((LoadParams.PropagationFlags & UE::LCPF_ReadParentSections) != 0)
 				{
-					ArrayHelper.EmptyValues();
+					// call LoadConfig on the parent class
+					UObject::FLoadConfigParams ParentContext = LoadParams;
+					ParentContext.ConfigClass = ParentClass;
+					ParentContext.Filename = TEXT("");
+					ActiveObject->LoadConfig(ParentContext);
+
+					// if we are also notifying child classes or instances, stop here as this object's properties will be imported as a result of notifying the others
+					if ((LoadParams.PropagationFlags & (UE::LCPF_PropagateToChildDefaultObjects | UE::LCPF_PropagateToInstances)) != 0)
+					{
+						return false;
+					}
 				}
-#endif
-				if (Sec)
+				else if ((LoadParams.PropagationFlags & UE::LCPF_PropagateToChildDefaultObjects) != 0)
 				{
-					TArray<FConfigValue> List;
-					const FName KeyName(*Key, FNAME_Find);
-					Sec->MultiFind(KeyName, List);
-
-					const int32 Size = Array->Inner->GetElementSize();
+					// not propagating the call upwards, but we are propagating the call to all child classes
 
-					// Only override default properties if there is something to override them with, or
-					// if the config system had a "initialize to empty" entry
-					if (!List.IsEmpty() || Sec->EmptyInitializedKeys.Contains(KeyName))
+					UObject::FLoadConfigParams ChildContext = LoadParams;
+					// mask out the PropgateToParent and PropagateToChildren values
+					ChildContext.PropagationFlags = (LoadParams.PropagationFlags & (UE::LCPF_PersistentFlags | UE::LCPF_PropagateToInstances));
+					ChildContext.Filename = TEXT("");
+					for (TObjectIterator<UClass> It; It; ++It)
 					{
-						ArrayHelper.EmptyAndAddValues(List.Num());
-						for (int32 i = List.Num() - 1, c = 0; i >= 0; i--, c++)
+						if (It->IsChildOf(ConfigClass))
 						{
-							Array->Inner->ImportText_Direct(*List[i].GetValue(), ArrayHelper.GetRawPtr(c), this, PortFlags);
+							ChildContext.ConfigClass = *It;
+							It->GetDefaultObject()->LoadConfig(ChildContext);
 						}
 					}
-					else
+
+					// LoadConfig() was called on this object during iteration, so stop here 
+					return false;
+				}
+				else if ((LoadParams.PropagationFlags & UE::LCPF_PropagateToInstances) != 0)
+				{
+					// call LoadConfig() on all instances of this class (except the CDO)
+					// Do not propagate this call to parents, and do not propagate to children or instances (would be redundant) 
+
+					UObject::FLoadConfigParams InstanceContext = LoadParams;
+					// mask out the PropgateToParent and PropagateToChildren values
+					InstanceContext.PropagationFlags = (LoadParams.PropagationFlags & (UE::LCPF_PersistentFlags));
+					InstanceContext.Filename = TEXT("");
+					for (TObjectIterator<UObject> It; It; ++It)
 					{
-						int32 Index = 0;
-						const FConfigValue* ElementValue = nullptr;
-						do
+						if (It->IsA(ConfigClass))
 						{
-							// Add array index number to end of key
-							FString IndexedKey = FString::Printf(TEXT("%s[%i]"), *Key, Index);
-
-							// Try to find value of key
-							const FName IndexedName(*IndexedKey, FNAME_Find);
-							if (IndexedName == NAME_None)
+							// null would also work
+							InstanceContext.ConfigClass = It->GetClass();
+							if (!GIsEditor || bTrackingAccessOnly)
 							{
-								break;
+								It->LoadConfig(InstanceContext);
 							}
-							ElementValue = Sec->Find(IndexedName);
-
-							// If found, import the element
-							if (ElementValue != nullptr)
+	#if WITH_EDITOR
+							else
 							{
-								// expand the array if necessary so that Index is a valid element
-								ArrayHelper.ExpandForIndex(Index);
-								Array->Inner->ImportText_Direct(*ElementValue->GetValue(), ArrayHelper.GetRawPtr(Index), this, PortFlags);
+								It->PreEditChange(nullptr);
+								It->LoadConfig(InstanceContext);
+								It->PostEditChange();
 							}
-
-							Index++;
-						} while (ElementValue || Index < ArrayHelper.Num());
+	#endif // WITH_EDITOR
+						}
 					}
 				}
 			}
-			else if (SetProperty)
+			else if ((LoadParams.PropagationFlags & UE::LCPF_PropagateToChildDefaultObjects) != 0)
 			{
-				FScriptSetHelper_InContainer SetHelper(SetProperty, this);
+				// we're at the base-most config class
 
-#if WITH_EDITOR
-				// Empty out any set properties if this is a PerPlatformConfig class
-				// as we are replacing the values with the Platform's version when entering
-				// a new preview platform.
-				if (bIsPerPlatformConfig)
+				UObject::FLoadConfigParams ChildContext = LoadParams;
+				// mask out the PropgateToParent and PropagateToChildren values
+				ChildContext.PropagationFlags = (LoadParams.PropagationFlags & (UE::LCPF_PersistentFlags | UE::LCPF_PropagateToInstances));
+				ChildContext.Filename = TEXT("");
+				for (TObjectIterator<UClass> It; It; ++It)
 				{
-					SetHelper.EmptyElements();
+					if (It->IsChildOf(ConfigClass))
+					{
+						ChildContext.ConfigClass = *It;
+						if (!GIsEditor || bTrackingAccessOnly)
+						{
+							It->GetDefaultObject()->LoadConfig(ChildContext);
+						}
+	#if WITH_EDITOR
+						else
+						{
+							It->PreEditChange(nullptr);
+							It->GetDefaultObject()->LoadConfig(ChildContext);
+							It->PostEditChange();
+						}
+	#endif // WITH_EDITOR
+					}
 				}
-#endif
 
-				if (Sec)
-				{
-					TArray<FConfigValue> List;
-					const FName KeyName(*Key, FNAME_Find);
-					Sec->MultiFind(KeyName, List);
+				return false;
+			}
+			else if ((LoadParams.PropagationFlags & UE::LCPF_PropagateToInstances) != 0)
+			{
+				UObject::FLoadConfigParams InstanceContext = LoadParams;
+				// mask out the PropgateToParent and PropagateToChildren values
+				InstanceContext.PropagationFlags = (LoadParams.PropagationFlags & (UE::LCPF_PersistentFlags));
+				InstanceContext.Filename = TEXT("");
 
-					bool bSingleSetEntry = false;
-					if (List.Num() == 1)
+				for (TObjectIterator<UObject> It; It; ++It)
+				{
+					if (It->GetClass() == ConfigClass)
 					{
-						const FString& SingleListValue = List[0].GetValue();
-						if (SingleListValue.Len() > 1
-							&& SingleListValue[0] == TEXT('(')
-							&& SingleListValue[SingleListValue.Len() - 1] == TEXT(')'))
+						InstanceContext.ConfigClass = It->GetClass();
+						
+						if (!GIsEditor || bTrackingAccessOnly)
 						{
-							// If we have a single entry in the set that is surrounded with parentheses, fall back to
-							// the old processing method
-							bSingleSetEntry = true;
+							It->LoadConfig(InstanceContext);
+						}
+	#if WITH_EDITOR
+						else
+						{
+							It->PreEditChange(nullptr);
+							It->LoadConfig(InstanceContext);
+							It->PostEditChange();
 						}
+	#endif // WITH_EDITOR
 					}
+				}
+			}
+		}
+		
+		return true;
+	}
 
-					// Only override default properties if there is something to override them with.
-					if (!bSingleSetEntry && !List.IsEmpty())
-					{
-						bProcessedProperty = true;
-						SetHelper.EmptyElements(List.Num());
-
-						// Each config value entry can possibly specify multiple set elements - create a temporary
-						// set here that we can import to and add its elements to the object's property
-						void* TempSet = FMemory::Malloc(SetProperty->GetSize(), SetProperty->GetMinAlignment());
-						SetProperty->InitializeValue(TempSet);
-
-						// Importing elements to the set inline does not check for duplicates - create a temporary
-						// element that we can import to add uniquely to the set
-						void* TempElement = FMemory::Malloc(SetProperty->ElementProp->GetSize(),
-							SetProperty->ElementProp->GetMinAlignment());
-						SetProperty->ElementProp->InitializeValue(TempElement);
-
-						for (const FConfigValue& ListValue : List)
-						{
-							const FString& ListString = ListValue.GetValue();
 
-							// Try to import the config value as an entire set first
-							const TCHAR* SetImportResult =
-								SetProperty->ImportText_Direct(*ListString, TempSet, this, PortFlags);
+	//------------------------------------------------------------------------------------
+	
+	void HandlePerObject()
+	{
+		FString PathNameString;
+		UObject* Outermost = ActiveObject->GetOutermost();
 
-							if (SetImportResult && SetImportResult != *ListString)
-							{
-								FScriptSetHelper TempSetHelper(SetProperty, TempSet);
-								for (FScriptSetHelper::FIterator Itr = TempSetHelper.CreateIterator(); Itr; ++Itr)
-								{
-									SetHelper.AddElement(TempSetHelper.GetElementPtr(Itr));
-								}
-							}
-							else
-							{
-								// If we failed to import the value as an entire set, try to import it as a single
-								// element
-								const TCHAR* ElementImportResult =
-									SetProperty->ElementProp->ImportText_Direct(*ListString, TempElement, this, PortFlags);
+		if (Outermost == GetTransientPackage())
+		{
+			PathNameString = ActiveObject->GetName();
+		}
+		else
+		{
+			ActiveObject->GetPathName(Outermost, PathNameString);
+		}
 
-								if (ElementImportResult && ElementImportResult != *ListString)
-								{
-									SetHelper.AddElement(TempElement);
-								}
-							}
+		// this will be set per-property for non-per-object
+		ClassSection = PathNameString + TEXT(" ") + ActiveObject->GetClass()->GetName();
 
-						}
+		FString OverrideClassSection;
+		ActiveObject->OverridePerObjectConfigSection(OverrideClassSection);
+		if (OverrideClassSection.Len() && OverrideClassSection != ClassSection)
+		{
+			// If a we got a section name override no need to perform short class name checks
+			ClassSection = MoveTemp(OverrideClassSection);
+			// Keep ClassPathSection empty so that we don't check for it when the section name has been overridden
+		}
+		else
+		{
+			// Cache both version of per object config section name
+			ClassPathSection = PathNameString + TEXT(" ") + ActiveObject->GetClass()->GetPathName();
+		}
+	}
 
-						SetProperty->DestroyAndFreeValue(TempSet);
-						SetProperty->ElementProp->DestroyAndFreeValue(TempElement);
-						SetHelper.Rehash();
-					}
-				}
+	void CheckEarlyAccess()
+	{
+#if DO_CHECK
+		// When a CDO is created before the Class is linked, it will have an empty set of properties. Among other issues, this can cause loading of config data to fail silently. 
+		if (!ConfigClass->DebugIsPropertyChainReady())
+		{
+			FString ErrorMessage = FString::Printf(TEXT("Attempting to load config data for %s before the Class has been constructed/registered/linked (likely during module loading or early startup). This will result in the load silently failing and should be fixed."), *ActiveObject->GetName());
+			if (GPrematureConfigLoadErrorLevel == 0)
+			{
+				UE_LOG(LogObj, Display, TEXT("%s"), *ErrorMessage);
+			}
+			else if (GPrematureConfigLoadErrorLevel == 1)
+			{
+				UE_LOG(LogObj, Warning, TEXT("%s"), *ErrorMessage);
+			}
+			else if (GPrematureConfigLoadErrorLevel == 2)
+			{
+				// Repeating the condition here just so it ends up in the output message
+				ensureAlwaysMsgf(ConfigClass->DebugIsPropertyChainReady(), TEXT("%s"), *ErrorMessage);
+			}
+			else
+			{
+				checkf(ConfigClass->DebugIsPropertyChainReady(), TEXT("%s"), *ErrorMessage);
 			}
 		}
+#endif
 
-		if (!bProcessedProperty)
+#if !IS_PROGRAM
+		auto HaveSameProperties = [](const UStruct* Struct1, const UStruct* Struct2) -> bool
 		{
-			for (int32 i = 0; i < Property->ArrayDim; i++)
+			TFieldIterator<FProperty> It1(Struct1);
+			TFieldIterator<FProperty> It2(Struct2);
+			for (;; ++It1, ++It2)
 			{
-				if (Property->ArrayDim != 1)
+				bool bAtEnd1 = !It1;
+				bool bAtEnd2 = !It2;
+				// If one iterator is at the end and one isn't, the property lists are different
+				if (bAtEnd1 != bAtEnd2)
 				{
-					Key = FString::Printf(TEXT("%s[%i]"), *Property->GetName(), i);
+					return false;
 				}
-
-				FString Value;
-				bool bFoundValue = GetConfigValue(*ClassSection, *Key, *PropFileName, Value);
-				if (!bFoundValue && bPerObject && ClassPathSection.Len())
+				// If both iterators have reached the end, the property lists are the same
+				if (bAtEnd1)
 				{
-					// Try to get the value from POC config section with class path name
-					bFoundValue = GetConfigValue(*ClassPathSection, *Key, *PropFileName, Value);
+					return true;
 				}
-
-				if (bFoundValue)
+				// If the properties are different, the property lists are different
+				if (*It1 != *It2)
 				{
-					if (Property->ImportText_Direct(*Value, Property->ContainerPtrToValuePtr<uint8>(this, i), this, PortFlags) == nullptr)
-					{
-						// this should be an error as the properties from the .ini / .int file are not correctly being read in and probably are affecting things in subtle ways
-						UE_LOG(LogObj, Error, TEXT("LoadConfig (%s): import failed for %s in: %s"), *GetPathName(), *Property->GetName(), *Value);
-					}
+					return false;
 				}
+			}
+		};
+		// Do we have properties that don't exist yet?
+		// If this happens then we're trying to load the config for an object that doesn't
+		// know what its layout is. Usually a call to GetDefaultObject that occurs too early
+		// because ProcessNewlyLoadedUObjects hasn't happened yet
+		checkf(ConfigClass->PropertyLink != nullptr
+			|| (ConfigClass->GetSuperStruct() && HaveSameProperties(ConfigClass, ConfigClass->GetSuperStruct()))
+			|| (UE_WITH_CONSTINIT_UOBJECT ? !ConfigClass->HasAnyClassFlags(CLASS_Constructed) : ConfigClass->PropertiesSize == 0)
+			|| IsEngineExitRequested(), // Ignore this check when exiting as we may have requested exit during init when not everything is initialized
+			TEXT("class %s has uninitialized properties. Accessed too early?"), *ConfigClass->GetPathName());
+#endif
+	}
 
-#if !UE_BUILD_SHIPPING
-				if (!bFoundValue && !FPlatformProperties::RequiresCookedData())
-				{
-					CheckMissingSection(ClassSection, PropFileName);
-				}
+public:
+	// the final class to walk over properties in LoadConfig
+	const UClass* ConfigClass = nullptr;
+
+private:
+
+	///// settings for the entire class
+
+	const UObject::FLoadConfigParams& LoadParams;
+	FConfigCacheIni* ConfigSystem = nullptr;
+	UObject* ActiveObject = nullptr;
+	bool bPerObject = false;
+	bool bIsPerPlatformConfig = false;
+	// the default section for the class, but each property can potentially override it when setting ConfigSection
+	FString ClassSection;
+	FString ClassPathSection;
+	// the filename to look up in ConfigSystem,
+	FString Filename;
+	// storage for a local file when loading another platform's config
+	FConfigFile PlatformOverrideConfig;
+	// the config file for the class, but it can change per property (rarely)
+	const FConfigFile* ClassConfigFile = nullptr;
+
+	///// per-property settings
+
+	// flags to control the text import
+	int32 PortFlags = 0;
+	// the final file to load from
+	const FConfigFile* ConfigFile = nullptr;
+	// the final Section to read Key from
+	const FConfigSection* ConfigSection = nullptr;
+	// the key to use to look in ConfigSection
+	FString Key;
+
+	
+	// the ConfigPlatform is only used for tracking
+	bool bTrackingAccessOnly = false;
+#if WITH_EDITOR
+	FName ConfigPlatform;
 #endif
-			}
+};
+
+void UObject::LoadConfig(const FLoadConfigParams& Params)
+{
+	SCOPE_CYCLE_COUNTER(STAT_LoadConfig);
+
+	// make the local context that manages the complexity
+	FPropertyLoadConfigContext PropertyContext(Params, this);
+	
+	// there are various reasons why we won't need to continue, so early out
+	if (PropertyContext.InitContext() == false)
+	{
+		return;
+	}
+		
+	for (FProperty* Property = PropertyContext.ConfigClass->PropertyLink; Property; Property = Property->PropertyLinkNext)
+	{
+#if WITH_EDITOR
+		FSoftObjectPathSerializationScope SerializationScope(NAME_None, Property->GetFName(),
+			Property->GetOwnerProperty()->HasMetaData(FSoftObjectPath::NAME_Untracked) ? ESoftObjectPathCollectType::NeverCollect :
+			Property->IsEditorOnlyProperty() ? ESoftObjectPathCollectType::EditorOnlyCollect : ESoftObjectPathCollectType::AlwaysCollect,
+			ESoftObjectPathSerializeType::AlwaysSerialize);
+#endif
+
+		if (!PropertyContext.InitForProperty(Property))
+		{
+			continue;
+		}
+		
+		// Track if we loaded this config value using special handling (e.g. array or set)
+		bool bProcessedProperty = false;
+
+		FArrayProperty* ArrayProperty = CastField<FArrayProperty>(Property);
+		FSetProperty* SetProperty = CastField<FSetProperty>(Property);
+
+		if (ArrayProperty)
+		{
+			bProcessedProperty = PropertyContext.ProcessArrayProperty(ArrayProperty);
+		}
+		else if (SetProperty)
+		{
+			bProcessedProperty = PropertyContext.ProcessSetProperty(SetProperty);
+		}
+		
+		if (!bProcessedProperty)
+		{
+			PropertyContext.ProcessStandardProperty(Property);			
 		}
 	}
 
 	// if we are reloading config data after the initial class load, fire the callback now
-	if (bModifyingObjects && (PropagationFlags & UE::LCPF_ReloadingConfigData) != 0)
+	if (Params.AccessedValues == nullptr && (Params.PropagationFlags & UE::LCPF_ReloadingConfigData) != 0)
 	{
-		PostReloadConfig(PropertyToLoad);
+		PostReloadConfig(Params.PropertyToLoad);
 	}
 }
 
diff --git a/Engine/Source/Runtime/CoreUObject/Public/UObject/Object.h b/Engine/Source/Runtime/CoreUObject/Public/UObject/Object.h
--- a/Engine/Source/Runtime/CoreUObject/Public/UObject/Object.h
+++ b/Engine/Source/Runtime/CoreUObject/Public/UObject/Object.h
@@ -1427,6 +1427,29 @@ public:
 		uint32 PropagationFlags = UE::LCPF_None, class FProperty* PropertyToLoad = nullptr,
 		TArray<UE::ConfigAccessTracking::FConfigAccessData>* OutAccessedKeys = nullptr);
 
+	struct FLoadConfigParams 
+	{
+		FLoadConfigParams() = default;
+				
+		// the class to look in the config system for (defaults to Obj's Class if null)
+		UClass* ConfigClass = nullptr;
+		// the config system to read from (defaults to GConfig if null)
+		FConfigCacheIni* ConfigSystem = nullptr;
+		
+		// an optional filename to use to look up in the ConfigSystem
+		FString Filename;
+		// flags to control if we want to propagate LoadConfig to instances of the class, child classes, etc. see ELoadConfigPropagationFlags
+		uint32 PropagationFlags = UE::LCPF_None;
+		// if this is set, only a single property of the class will be loaded from config
+		class FProperty* PropertyToLoad = nullptr;
+		// if this is set, we only track what sections/keys are read from the config system, without modifying the Object at all
+		TArray<UE::ConfigAccessTracking::FConfigAccessData>* AccessedValues = nullptr;
+		// if set, this values will be read from this COnfigFile instead of looking up in ConfigSystem
+		FConfigFile* OverrideFile = nullptr;
+	};
+
+	COREUOBJECT_API void LoadConfig(const FLoadConfigParams& Context);
+	
 	/**
 	 * Wrapper method for LoadConfig that is used when reloading the config data for objects at runtime which have already loaded their config data at least once.
 	 * Allows the objects the receive a callback that its configuration data has been reloaded.
