- Fixed a static analysis warning when using a #define in an if()

#Synced-CL 48811006

diff --git a/Engine/Source/Runtime/CoreUObject/Private/UObject/Obj.cpp b/Engine/Source/Runtime/CoreUObject/Private/UObject/Obj.cpp
--- a/Engine/Source/Runtime/CoreUObject/Private/UObject/Obj.cpp
+++ b/Engine/Source/Runtime/CoreUObject/Private/UObject/Obj.cpp
@@ -3224,10 +3224,12 @@ private:
 		CheckEarlyAccess();
 		
 		// we only support tracking values in with the editor
-		if (!WITH_EDITOR && bTrackingAccessOnly)
+#if !WITH_EDITOR
+		if (bTrackingAccessOnly)
 		{
 			return false;
 		}
+#endif
 		
 		return true;
 	}
