Fix package import/export index handling for cells

#Synced-CL 48675545

diff --git a/Engine/Source/Runtime/CoreUObject/Private/Serialization/AsyncLoading2.cpp b/Engine/Source/Runtime/CoreUObject/Private/Serialization/AsyncLoading2.cpp
--- a/Engine/Source/Runtime/CoreUObject/Private/Serialization/AsyncLoading2.cpp
+++ b/Engine/Source/Runtime/CoreUObject/Private/Serialization/AsyncLoading2.cpp
@@ -4410,11 +4410,11 @@ private:
 	void ConditionalSerializeExport(const FAsyncPackageHeaderData& Header, int32 LocalExportIndex);
 	void ConditionalCreateImport(const FAsyncPackageHeaderData& Header, int32 LocalImportIndex);
 	void ConditionalSerializeImport(const FAsyncPackageHeaderData& Header, int32 LocalImportIndex);
-	int32 GetPublicCellExportIndex(uint64 CellExportHash, FAsyncPackageHeaderData*& OutHeader);
-	void ConditionalCreateCellExport(const FAsyncPackageHeaderData& Header, int32 CellExportIndex);
-	void ConditionalSerializeCellExport(const FAsyncPackageHeaderData& Header, int32 CellExportIndex);
-	void ConditionalCreateCellImport(const FAsyncPackageHeaderData& Header, int32 CellImportIndex);
-	void ConditionalSerializeCellImport(const FAsyncPackageHeaderData& Header, int32 CellImportIndex);
+	int32 GetPublicCellExportIndexAsLocalIndex(uint64 CellExportHash, FAsyncPackageHeaderData*& OutHeader);
+	void ConditionalCreateCellExport(const FAsyncPackageHeaderData& Header, int32 LocalExportIndex);
+	void ConditionalSerializeCellExport(const FAsyncPackageHeaderData& Header, int32 LocalExportIndex);
+	void ConditionalCreateCellImport(const FAsyncPackageHeaderData& Header, int32 LocalImportIndex);
+	void ConditionalSerializeCellImport(const FAsyncPackageHeaderData& Header, int32 LocalImportIndex);
 
 	/**
 	 * Begin async loading process. Simulates parts of BeginLoad.
@@ -8445,8 +8445,10 @@ bool FAsyncPackage2::CreateRuntimeExports(FAsyncLoadingThreadState2& ThreadState
 
 				if (BundleEntry.CommandType == FExportBundleEntry::ExportCommandType_Create)
 				{
-					check(!CellExport.Cell);
-					EventDrivenCreateCellExport(*LocalHeaderData, BundleEntry.LocalExportIndex, &Ar);
+					if (!CellExport.Cell)
+					{
+						EventDrivenCreateCellExport(*LocalHeaderData, BundleEntry.LocalExportIndex, &Ar);
+					}
 				}
 				else
 				{
@@ -8653,7 +8655,7 @@ void FAsyncPackage2::ProcessExportDependencies(const FAsyncPackageHeaderData& He
 			}
 			else
 			{
-				ConditionalCreateCellExport(Header, Dep.LocalImportOrExportIndex.ToExport() - Header.ExportsView.Num());
+				ConditionalCreateCellExport(Header, Dep.LocalImportOrExportIndex.ToExport());
 			}
 		}
 		else
@@ -8664,7 +8666,7 @@ void FAsyncPackage2::ProcessExportDependencies(const FAsyncPackageHeaderData& He
 			}
 			else
 			{
-				ConditionalCreateCellImport(Header, Dep.LocalImportOrExportIndex.ToImport() - Header.ImportMap.Num());
+				ConditionalCreateCellImport(Header, Dep.LocalImportOrExportIndex.ToImport());
 			}
 		}
 	}
@@ -8680,7 +8682,7 @@ void FAsyncPackage2::ProcessExportDependencies(const FAsyncPackageHeaderData& He
 			}
 			else
 			{
-				ConditionalSerializeCellExport(Header, Dep.LocalImportOrExportIndex.ToExport() - Header.ExportsView.Num());
+				ConditionalSerializeCellExport(Header, Dep.LocalImportOrExportIndex.ToExport());
 			}
 		}
 		else
@@ -8691,7 +8693,7 @@ void FAsyncPackage2::ProcessExportDependencies(const FAsyncPackageHeaderData& He
 			}
 			else
 			{
-				ConditionalSerializeCellImport(Header, Dep.LocalImportOrExportIndex.ToImport() - Header.ImportMap.Num());
+				ConditionalSerializeCellImport(Header, Dep.LocalImportOrExportIndex.ToImport());
 			}
 		}
 	}
@@ -8817,41 +8819,44 @@ void FAsyncPackage2::ConditionalSerializeImport(const FAsyncPackageHeaderData& H
 	}
 }
 
-int32 FAsyncPackage2::GetPublicCellExportIndex(uint64 CellExportHash, FAsyncPackageHeaderData*& OutHeader)
+int32 FAsyncPackage2::GetPublicCellExportIndexAsLocalIndex(uint64 CellExportHash, FAsyncPackageHeaderData*& OutHeader)
 {
 	for (int32 CellExportIndex = 0; CellExportIndex < HeaderData.CellExportMap.Num(); ++CellExportIndex)
 	{
 		if (HeaderData.CellExportMap[CellExportIndex].PublicExportHash == CellExportHash)
 		{
 			OutHeader = &HeaderData;
-			return CellExportIndex;
+			return CellExportIndex + HeaderData.ExportMap.Num();
 		}
 	}
 	return -1;
 }
 
-void FAsyncPackage2::ConditionalCreateCellExport(const FAsyncPackageHeaderData& Header, int32 CellExportIndex)
+void FAsyncPackage2::ConditionalCreateCellExport(const FAsyncPackageHeaderData& Header, int32 LocalExportIndex)
 {
+	const int32 CellExportIndex = LocalExportIndex - Header.ExportMap.Num();
 	if (!Header.CellExportsView[CellExportIndex].Cell)
 	{
 		FAsyncPackageScope2 Scope(this);
-		EventDrivenCreateCellExport(Header, CellExportIndex, nullptr);
+		EventDrivenCreateCellExport(Header, LocalExportIndex, nullptr);
 	}
 }
 
-void FAsyncPackage2::ConditionalSerializeCellExport(const FAsyncPackageHeaderData& Header, int32 CellExportIndex)
+void FAsyncPackage2::ConditionalSerializeCellExport(const FAsyncPackageHeaderData& Header, int32 LocalExportIndex)
 {
-	ConditionalCreateCellExport(Header, CellExportIndex);
+	ConditionalCreateCellExport(Header, LocalExportIndex);
 
+	const int32 CellExportIndex = LocalExportIndex - Header.ExportMap.Num();
 	if (!Header.CellExportsView[CellExportIndex].bSerialized)
 	{
 		FAsyncPackageScope2 Scope(this);
-		EventDrivenSerializeCellExport(Header, CellExportIndex, nullptr);
+		EventDrivenSerializeCellExport(Header, LocalExportIndex, nullptr);
 	}
 }
 
-void FAsyncPackage2::ConditionalCreateCellImport(const FAsyncPackageHeaderData& Header, int32 CellImportIndex)
+void FAsyncPackage2::ConditionalCreateCellImport(const FAsyncPackageHeaderData& Header, int32 LocalImportIndex)
 {
+	const int32 CellImportIndex = LocalImportIndex - Header.ImportMap.Num();
 	const FPackageObjectIndex& CellIndex = Header.CellImportMap[CellImportIndex];
 	check(CellIndex.IsPackageImport());
 	if (Verse::VCell* FromImportStore = ImportStore.FindOrGetImportCell(Header, CellIndex))
@@ -8871,17 +8876,18 @@ void FAsyncPackage2::ConditionalCreateCellImport(const FAsyncPackageHeaderData&
 	// If the header ref is at 0, it means all the serialization is done for that package, we just raced with it and missed the cue.
 	if (SourcePackage->TryAddHeaderRef())
 	{
-		int32 CellExportIndex = SourcePackage->GetPublicCellExportIndex(CellExportHash, SourcePackageHeader);
-		if (CellExportIndex >= 0)
+		const int32 LocalExportIndex = SourcePackage->GetPublicCellExportIndexAsLocalIndex(CellExportHash, SourcePackageHeader);
+		if (LocalExportIndex >= 0)
 		{
-			SourcePackage->ConditionalCreateCellExport(*SourcePackageHeader, CellExportIndex);
+			SourcePackage->ConditionalCreateCellExport(*SourcePackageHeader, LocalExportIndex);
 		}
 		SourcePackage->ReleaseHeaderRef();
 	}
 }
 
-void FAsyncPackage2::ConditionalSerializeCellImport(const FAsyncPackageHeaderData& Header, int32 CellImportIndex)
+void FAsyncPackage2::ConditionalSerializeCellImport(const FAsyncPackageHeaderData& Header, int32 LocalImportIndex)
 {
+	const int32 CellImportIndex = LocalImportIndex - Header.ImportMap.Num();
 	const FPackageObjectIndex& CellIndex = Header.CellImportMap[CellImportIndex];
 	check(CellIndex.IsPackageImport());
 
@@ -8897,10 +8903,10 @@ void FAsyncPackage2::ConditionalSerializeCellImport(const FAsyncPackageHeaderDat
 	// If the header ref is at 0, it means all the serialization is done for that package, we just raced with it and missed the cue.
 	if (SourcePackage->TryAddHeaderRef())
 	{
-		int32 CellExportIndex = SourcePackage->GetPublicCellExportIndex(CellExportHash, SourcePackageHeader);
-		if (CellExportIndex >= 0)
+		const int32 LocalExportIndex = SourcePackage->GetPublicCellExportIndexAsLocalIndex(CellExportHash, SourcePackageHeader);
+		if (LocalExportIndex >= 0)
 		{
-			SourcePackage->ConditionalSerializeCellExport(*SourcePackageHeader, CellExportIndex);
+			SourcePackage->ConditionalSerializeCellExport(*SourcePackageHeader, LocalExportIndex);
 		}
 		SourcePackage->ReleaseHeaderRef();
 	}
@@ -13458,4 +13464,4 @@ bool FPublicExportMapPreInsertPublicExportsTest::RunTest(const FString& Paramete
 	return true;
 }
 
-#endif // WITH_LOW_LEVEL_TESTS
\ No newline at end of file
+#endif // WITH_LOW_LEVEL_TESTS
