[ZenLoader]
- Add support for s.AllowMultithreadedLoading to be defined in the /Script/Engine.[Editor]StreamingSettings ini section if desired

#Synced-CL 48817740

diff --git a/Engine/Source/Runtime/CoreUObject/Private/Serialization/AsyncLoading2.cpp b/Engine/Source/Runtime/CoreUObject/Private/Serialization/AsyncLoading2.cpp
--- a/Engine/Source/Runtime/CoreUObject/Private/Serialization/AsyncLoading2.cpp
+++ b/Engine/Source/Runtime/CoreUObject/Private/Serialization/AsyncLoading2.cpp
@@ -10942,6 +10942,17 @@ FAsyncLoadingThread2::FAsyncLoadingThread2(FIoDispatcher& InIoDispatcher)
 	NotifyCompiledVersePackageDelegate.BindRaw(this, &FAsyncLoadingThread2::NotifyCompiledVersePackage);
 #endif
 
+	// if we are on an uncooked platform query the editor streaming settings
+	const TCHAR* SectionName = FPlatformProperties::RequiresCookedData() ? TEXT("/Script/Engine.StreamingSettings") : TEXT("/Script/Engine.EditorStreamingSettings");
+	
+	bool bConfigValue = false;
+
+	// This is a CVar but for consistency, we allow setting it with the other streaming settings
+	if (GConfig->GetBool(SectionName, TEXT("s.AllowMultithreadedLoading"), bConfigValue, GEngineIni))
+	{
+		GAllowMultithreadedLoading = bConfigValue;
+	}
+
 	UE_LOG(LogStreaming, Log, TEXT("AsyncLoading2 - Created: Event Driven Loader: %s, Async Loading Thread: %s, Async Post Load: %s"),
 		GEventDrivenLoaderEnabled ? TEXT("true") : TEXT("false"),
 		FAsyncLoadingThreadSettings::Get().bAsyncLoadingThreadEnabled ? TEXT("true") : TEXT("false"),
