[LinkerLoad]
- Unify loading path between -game and editor/commandlet by removing diverging code that was causing race conditions on CreatePackage.

#Synced-CL 48817740

diff --git a/Engine/Source/Runtime/CoreUObject/Private/UObject/LinkerLoad.cpp b/Engine/Source/Runtime/CoreUObject/Private/UObject/LinkerLoad.cpp
--- a/Engine/Source/Runtime/CoreUObject/Private/UObject/LinkerLoad.cpp
+++ b/Engine/Source/Runtime/CoreUObject/Private/UObject/LinkerLoad.cpp
@@ -6307,125 +6307,6 @@ UObject* FLinkerLoad::CreateImport( int32 Index )
 	// Imports can have no name if they were filtered out due to package redirects, skip in that case
 	if (Import.XObject == nullptr && Import.ObjectName != NAME_None)
 	{
-		FUObjectSerializeContext* CurrentLoadContext = FUObjectThreadContext::Get().GetSerializeContext();
-		if (!GIsEditor && !IsRunningCommandlet())
-		{
-			// Try to find existing version in memory first.
-			if( UPackage* ClassPackage = FindObjectFast<UPackage>( nullptr, Import.ClassPackage, EFindObjectFlags::None ) )
-			{
-				if( UClass*	FindClass = FindObjectFast<UClass>( ClassPackage, Import.ClassName, EFindObjectFlags::None ) ) //
-				{
-					// Make sure the class has been loaded and linked before creating a CDO.
-					// This is an edge case, but can happen if a blueprint package has not finished creating exports for a class
-					// during async loading, and another package creates the class via CreateImport while in cooked builds because
-					// we don't call preload immediately after creating a class in CreateExport like in non-cooked builds.
-					Preload( FindClass );
-
-					FindClass->GetDefaultObject(); // build the CDO if it isn't already built
-					UObject*	FindObject		= nullptr;
-
-					// Import is a toplevel package.
-					if( Import.OuterIndex.IsNull() )
-					{
-						FName ObjectName = InstancingContext.RemapPackage(Import.ObjectName);
-						// Instancing context supports remapping editor-only references to none, handle this case here.
-						UPackage* Pkg = !ObjectName.IsNone() ? CreatePackage(*ObjectName.ToString()) : nullptr;
-						if (IsPackageReferenceAllowed(Pkg))
-						{
-							FindObject = Pkg;
-						}
-					}
-					// Import is regular import/ export.
-					else
-					{
-						// Find the imports' outer.
-						UObject* FindOuter = nullptr;
-						// Import.
-						if( Import.OuterIndex.IsImport() )
-						{
-							FObjectImport& OuterImport = Imp(Import.OuterIndex);
-							// Outer already in memory.
-							if( OuterImport.XObject )
-							{
-								FindOuter = OuterImport.XObject;
-							}
-							// Outer is toplevel package, create/ find it.
-							else if( OuterImport.OuterIndex.IsNull() )
-							{
-								FName ObjectName = InstancingContext.RemapPackage(OuterImport.ObjectName);
-								// Instancing context supports remapping editor-only references to none, handle this case here as well.
-								UPackage* Pkg = !ObjectName.IsNone() ? CreatePackage(*ObjectName.ToString()) : nullptr;
-								if (IsPackageReferenceAllowed(Pkg))
-								{
-									FindOuter = Pkg;
-								}
-
-							}
-							// Outer is regular import/ export, use IndexToObject to potentially recursively load/ find it.
-							else
-							{
-								FindOuter = IndexToObject( Import.OuterIndex );
-							}
-						}
-						// Export.
-						else
-						{
-							// Create/ find the object's outer.
-							FindOuter = IndexToObject( Import.OuterIndex );
-						}
-						if (!FindOuter)
-						{
-							// This can happen when deleting native properties or restructing blueprints. If there is an actual problem it will be caught when trying to resolve the outer itself
-							FString OuterName = Import.OuterIndex.IsNull() ? LinkerRoot->GetFullName() : GetFullImpExpName(Import.OuterIndex);
-							UE_ASSET_LOG(LogLinker, Verbose, PackagePath, TEXT("CreateImport: Failed to load Outer for resource '%s': %s"), *Import.ObjectName.ToString(), *OuterName);
-							return NULL;
-						}
-
-						// Find object now that we know it's class, outer and name.
-						FindObject = FindImportFast(FindClass, FindOuter, Import.ObjectName);
-#if UE_INTERNAL_UNIVERSAL_PLACEHOLDERS
-						if(!FindObject)
-						{
-							if(FindClass->IsChildOf(UClass::StaticClass()))
-							{
-								ForceCreatePlaceholderClassImport(Index);
-							}
-							else
-							{
-								FindObject = NewObject<UObject>(FindOuter, FindClass, Import.ObjectName, RF_Public);
-								if(UScriptStruct* Struct = Cast<UScriptStruct>(FindObject))
-								{
-									// add a single property to avoid creating a 0 sized struct. Zero sized allocations
-									// are supported by some systems, but containers are generally unhappy with zero stride
-									// container elements. We can't just use the fallback struct as we want a proper
-									// object name:
-									FBoolProperty* Property = new FBoolProperty(Struct, "__PlaceholderSentinel", RF_NoFlags);
-									Struct->ChildProperties = Property;
-									Struct->StaticLink(true); // make the struct roughly usable, it's just a dummy type
-								}
-							}
-						}
-#endif // UE_INTERNAL_UNIVERSAL_PLACEHOLDERS
-					}
-
-					if( FindObject )
-					{
-#if USE_CIRCULAR_DEPENDENCY_LOAD_DEFERRING
-						// Don't use the object if it's still waiting on some part of a deferred load!
-						const FLinkerLoad* ObjLinker = FindObject->GetLinker();
-						if (!ObjLinker || !ObjLinker->IsBlueprintFinalizationPending())
-#endif	// USE_CIRCULAR_DEPENDENCY_LOAD_DEFERRING
-						{
-							// Associate import and indicate that we associated an import for later cleanup.
-							Import.XObject = FindObject;
-							CurrentLoadContext->IncrementImportCount();
-							FLinkerManager::Get().AddLoaderWithNewImports(this);
-						}
-					}
-				}
-			}
-		}
-
 		if( Import.XObject == NULL )
 		{
 			EVerifyResult VerifyImportResult = VERIFY_Success;
@@ -6469,6 +6350,7 @@ UObject* FLinkerLoad::CreateImport( int32 Index )
 					}
 				}
 #endif
+				FUObjectSerializeContext* CurrentLoadContext = FUObjectThreadContext::Get().GetSerializeContext();
 				CurrentLoadContext->IncrementImportCount();
 				FLinkerManager::Get().AddLoaderWithNewImports(this);
 			}
