[Zen] Crash earlier on corrupted ZenPackageHeader.

Adds an assertion when the header size is larger than the memory size in `FZenPackageHeader::MakeView()`.

#rnx
#Synced-CL 49029036

diff --git a/Engine/Source/Runtime/CoreUObject/Private/Serialization/ZenPackageHeader.cpp b/Engine/Source/Runtime/CoreUObject/Private/Serialization/ZenPackageHeader.cpp
--- a/Engine/Source/Runtime/CoreUObject/Private/Serialization/ZenPackageHeader.cpp
+++ b/Engine/Source/Runtime/CoreUObject/Private/Serialization/ZenPackageHeader.cpp
@@ -23,6 +23,7 @@ FZenPackageHeader FZenPackageHeader::MakeView(FMemoryView Memory, FString& OutEr
 	const uint8* PackageHeaderDataPtr = reinterpret_cast<const uint8*>(Memory.GetData());
 	const FZenPackageSummary* PackageSummary = reinterpret_cast<const FZenPackageSummary*>(PackageHeaderDataPtr);
 	PackageHeader.PackageSummary = PackageSummary;
+	checkf(PackageSummary->HeaderSize <= Memory.GetSize(), TEXT("Corrupted package: HeaderSize=%u exceeds MemorySize=%llu"), PackageSummary->HeaderSize, Memory.GetSize());
 
 	TArrayView<const uint8> PackageHeaderDataView(PackageHeaderDataPtr + sizeof(FZenPackageSummary), PackageHeader.PackageSummary->HeaderSize - sizeof(FZenPackageSummary));
 	FMemoryReaderView PackageHeaderDataReader(PackageHeaderDataView);
