[StructUtils] Use the default values when serializing the InstancedStruct and InstancedStructContainer. If the defaults are not compatible, use the default value of the struct instead.
#Synced-CL 49067435

diff --git a/Engine/Source/Runtime/CoreUObject/Private/StructUtils/InstancedStruct.cpp b/Engine/Source/Runtime/CoreUObject/Private/StructUtils/InstancedStruct.cpp
--- a/Engine/Source/Runtime/CoreUObject/Private/StructUtils/InstancedStruct.cpp
+++ b/Engine/Source/Runtime/CoreUObject/Private/StructUtils/InstancedStruct.cpp
@@ -4,6 +4,7 @@
 
 #include "Misc/PackageName.h"
 #include "StructUtils/StructView.h"
+#include "StructUtils/UserDefinedStruct.h"
 #include "Serialization/CustomVersion.h"
 #include "UObject/Linker.h"
 #include "UObject/CoreRedirects.h"
@@ -126,6 +127,35 @@ bool FInstancedStruct::Serialize(FArchive& Ar, UStruct* DefaultsStruct, const vo
 {
 	Ar.UsingCustomVersion(GInstancedStructCustomVersion.GUID);
 
+	const FInstancedStruct* DefaultInstancedStruct = static_cast<const FInstancedStruct*>(Defaults);
+
+	auto SerializeWithDefaults = [Self = this, &Ar, DefaultInstancedStruct](UScriptStruct* ScriptStruct)
+		{
+			const bool bSameStruct = DefaultInstancedStruct && DefaultInstancedStruct->GetScriptStruct() == ScriptStruct;
+			if (bSameStruct)
+			{
+				// Use the parent Defaults
+				ScriptStruct->SerializeItem(Ar, Self->GetMutableMemory(), DefaultInstancedStruct->GetMemory());
+			}
+			else if (UScriptStruct::ICppStructOps* TheCppStructOps = ScriptStruct->GetCppStructOps(); TheCppStructOps && !TheCppStructOps->HasZeroConstructor())
+			{
+				// Create a temporary struct on the stack. The struct has default values.
+				const int32 MinAlignment = ScriptStruct->GetMinAlignment();
+				const int32 RequiredSize = ScriptStruct->GetStructureSize();
+				TArray<uint8, TInlineAllocator<64>> TmpMemoryBuffer;
+				TmpMemoryBuffer.SetNumUninitialized(FMath::Max(1, RequiredSize) + FMath::Max(0, MinAlignment - (int32)MIN_ALIGNMENT));
+				void* TempMemory = Align(TmpMemoryBuffer.GetData(), MinAlignment);
+				ScriptStruct->InitializeStruct(TempMemory);
+				ScriptStruct->SerializeItem(Ar, Self->GetMutableMemory(), TempMemory);
+				ScriptStruct->DestroyStruct(TempMemory);
+			}
+			else
+			{
+				// Serialize without defaults
+				ScriptStruct->SerializeItem(Ar, Self->GetMutableMemory(), nullptr);
+			}
+		};
+
 	if (Ar.IsLoading())
 	{
 		const int32 CustomVersion = Ar.CustomVer(GInstancedStructCustomVersion.GUID);
@@ -161,7 +191,14 @@ bool FInstancedStruct::Serialize(FArchive& Ar, UStruct* DefaultsStruct, const vo
 		// Initialize only if the type changes.
 		if (ScriptStruct != SerializedScriptStruct)
 		{
-			InitializeAs(SerializedScriptStruct);
+			if (DefaultInstancedStruct && DefaultInstancedStruct->GetScriptStruct() == SerializedScriptStruct)
+			{
+				InitializeAs(SerializedScriptStruct, DefaultInstancedStruct->GetMemory());
+			}
+			else
+			{
+				InitializeAs(SerializedScriptStruct);
+			}
 		}
 
 		// Size of the serialized memory
@@ -183,7 +220,7 @@ bool FInstancedStruct::Serialize(FArchive& Ar, UStruct* DefaultsStruct, const vo
 		}
 		else if (ScriptStruct != nullptr && ensureMsgf(GetMutableMemory() != nullptr, TEXT("A valid script struct should always have allocated memory")))
 		{
-			ConstCast(ScriptStruct)->SerializeItem(Ar, GetMutableMemory(), /* Defaults */ nullptr);
+			SerializeWithDefaults(ConstCast(ScriptStruct));
 		}
 	}
 	else if (Ar.IsSaving())
@@ -197,7 +234,7 @@ bool FInstancedStruct::Serialize(FArchive& Ar, UStruct* DefaultsStruct, const vo
 		{
 			// If saving a duplicated UDS, save the primary type instead, so that the data is loaded with the original struct.
 			// This is used as part of the user defined struct reinstancing logic.
-			UUserDefinedStruct* PrimaryUserDefinedStruct = UserDefinedStruct->PrimaryStruct.Get(); 
+			UUserDefinedStruct* PrimaryUserDefinedStruct = UserDefinedStruct->PrimaryStruct.Get();
 			Ar << PrimaryUserDefinedStruct;
 		}
 		else
@@ -215,7 +252,7 @@ bool FInstancedStruct::Serialize(FArchive& Ar, UStruct* DefaultsStruct, const vo
 		const int64 InitialOffset = Ar.Tell(); // Position before struct serialization to compute its serial size
 		if (ScriptStruct != nullptr && ensureMsgf(GetMutableMemory() != nullptr, TEXT("A valid script struct should always have allocated memory")))
 		{
-			ConstCast(ScriptStruct)->SerializeItem(Ar, GetMutableMemory(), /* Defaults */ nullptr);
+			SerializeWithDefaults(ConstCast(ScriptStruct));
 		}
 		const int64 FinalOffset = Ar.Tell(); // Keep current offset to reset the archive pos after write the serial size
 
diff --git a/Engine/Source/Runtime/CoreUObject/Private/StructUtils/InstancedStructContainer.cpp b/Engine/Source/Runtime/CoreUObject/Private/StructUtils/InstancedStructContainer.cpp
--- a/Engine/Source/Runtime/CoreUObject/Private/StructUtils/InstancedStructContainer.cpp
+++ b/Engine/Source/Runtime/CoreUObject/Private/StructUtils/InstancedStructContainer.cpp
@@ -1,6 +1,7 @@
 // Copyright Epic Games, Inc. All Rights Reserved.
 
 #include "StructUtils/InstancedStructContainer.h"
+#include "StructUtils/UserDefinedStruct.h"
 #include "UObject/Linker.h"
 
 #if WITH_EDITOR
@@ -505,8 +506,10 @@ bool FInstancedStructContainer::Identical(const FInstancedStructContainer* Other
 	return bResult;
 }
 
-bool FInstancedStructContainer::Serialize(FArchive& Ar)
+bool FInstancedStructContainer::Serialize(FArchive& Ar, UStruct* DefaultsStruct, const void* Defaults)
 {
+	const FInstancedStructContainer* DefaultInstancedStructContainer = static_cast<const FInstancedStructContainer*>(Defaults);
+
 	enum class EVersion : uint8
 	{
 		InitialVersion = 0,
@@ -528,6 +531,36 @@ bool FInstancedStructContainer::Serialize(FArchive& Ar)
 	int32 NumItemsSerialized = NumItems;
 	Ar << NumItemsSerialized;
 
+	auto SerializeWithDefaults = [Self = this, &Ar, DefaultInstancedStructContainer](const FItem& Item, int32 Index)
+		{
+			const bool bSameStruct = DefaultInstancedStructContainer
+				&& DefaultInstancedStructContainer->IsValidIndex(Index)
+				&& DefaultInstancedStructContainer->GetItem(Index).ScriptStruct == Item.ScriptStruct;
+			if (bSameStruct)
+			{
+				// Use the parent Defaults
+				const void* DefaultMemory = DefaultInstancedStructContainer->Memory + DefaultInstancedStructContainer->GetItem(Index).Offset;
+				ConstCast(Item.ScriptStruct)->SerializeItem(Ar, Self->Memory + Item.Offset, DefaultMemory);
+			}
+			else if (UScriptStruct::ICppStructOps* TheCppStructOps = Item.ScriptStruct->GetCppStructOps(); TheCppStructOps && !TheCppStructOps->HasZeroConstructor())
+			{
+				// Create a temporary struct on the stack. The struct has default values.
+				const int32 MinAlignment = Item.ScriptStruct->GetMinAlignment();
+				const int32 RequiredSize = Item.ScriptStruct->GetStructureSize();
+				TArray<uint8, TInlineAllocator<64>> TmpMemoryBuffer;
+				TmpMemoryBuffer.SetNumUninitialized(FMath::Max(1, RequiredSize) + FMath::Max(0, MinAlignment - (int32)MIN_ALIGNMENT));
+				void* TempMemory = Align(TmpMemoryBuffer.GetData(), MinAlignment);
+				ConstCast(Item.ScriptStruct)->InitializeStruct(TempMemory);
+				ConstCast(Item.ScriptStruct)->SerializeItem(Ar, Self->Memory + Item.Offset, TempMemory);
+				ConstCast(Item.ScriptStruct)->DestroyStruct(TempMemory);
+			}
+			else
+			{
+				// Serialize without defaults
+				ConstCast(Item.ScriptStruct)->SerializeItem(Ar, Self->Memory + Item.Offset, nullptr);
+			}
+		};
+
 	if (NumItemsSerialized > 0)
 	{
 		if (Ar.IsLoading())
@@ -573,7 +606,7 @@ bool FInstancedStructContainer::Serialize(FArchive& Ar)
 				}
 				else if (Item.ScriptStruct != nullptr)
 				{
-					ConstCast(Item.ScriptStruct)->SerializeItem(Ar, Memory + Item.Offset, /* Defaults */ nullptr);
+					SerializeWithDefaults(Item, Index);
 				}
 			}
 		}
@@ -616,7 +649,7 @@ bool FInstancedStructContainer::Serialize(FArchive& Ar)
 				const int64 InitialOffset = Ar.Tell(); // Position before struct serialization to compute its serial size
 				if (Item.ScriptStruct != nullptr)
 				{
-					ConstCast(Item.ScriptStruct)->SerializeItem(Ar, Memory + Item.Offset, /* Defaults */ nullptr);
+					SerializeWithDefaults(Item, Index);
 				}
 				const int64 FinalOffset = Ar.Tell(); // Keep current offset to reset the archive pos after write the serial size
 
diff --git a/Engine/Source/Runtime/CoreUObject/Public/StructUtils/InstancedStructContainer.h b/Engine/Source/Runtime/CoreUObject/Public/StructUtils/InstancedStructContainer.h
--- a/Engine/Source/Runtime/CoreUObject/Public/StructUtils/InstancedStructContainer.h
+++ b/Engine/Source/Runtime/CoreUObject/Public/StructUtils/InstancedStructContainer.h
@@ -230,7 +230,7 @@ public:
 	/** Type traits */
 	UE_API void AddStructReferencedObjects(FReferenceCollector& Collector);
 	UE_API bool Identical(const FInstancedStructContainer* Other, uint32 PortFlags) const;
-	UE_API bool Serialize(FArchive& Ar);
+	UE_API bool Serialize(FArchive& Ar, UStruct* DefaultsStruct = nullptr, const void* Defaults = nullptr);
 	UE_API void GetPreloadDependencies(TArray<UObject*>& OutDeps) const;
 	UE_API bool ExportTextItem(FString& ValueStr, FInstancedStructContainer const& DefaultValue, UObject* Parent, int32 PortFlags, UObject* ExportRootScope) const;
 	UE_API bool ImportTextItem(const TCHAR*& Buffer, int32 PortFlags, UObject* Parent, FOutputDevice* ErrorText, FArchive* InSerializingArchive = nullptr);
@@ -282,6 +282,8 @@ struct TStructOpsTypeTraits<FInstancedStructContainer> : public TStructOpsTypeTr
 		WithGetPreloadDependencies = true,
 		WithExportTextItem = true,
 		WithImportTextItem = true,
+		WithClearOnFinishDestroy = true,
+		WithZeroConstructor = true,
 	};
 };
 
