Dstm: Move PhysicsIslandId from RemoteObjectStub into DstmPhysicsMigration. PhysicsIslandIds are no longer migrated with object,

#Synced-CL 48954519

diff --git a/Engine/Source/Runtime/CoreUObject/Internal/UObject/UObjectMigrationContext.h b/Engine/Source/Runtime/CoreUObject/Internal/UObject/UObjectMigrationContext.h
--- a/Engine/Source/Runtime/CoreUObject/Internal/UObject/UObjectMigrationContext.h
+++ b/Engine/Source/Runtime/CoreUObject/Internal/UObject/UObjectMigrationContext.h
@@ -69,12 +69,6 @@ struct FUObjectMigrationContext final
 	/** The (New) Owner of the ObjectId as a result of this migration */
 	FRemoteServerId OwnerServerId;
 
-	/** The Server that this Object's Physics reside on */
-	FRemoteServerId PhysicsServerId;
-
-	/** The Local Island Index the Object belongs to for Physics */
-	uint32 PhysicsLocalIslandId;
-
 	/** Which side of the migration does this context apply to? Sending or Receiving? */
 	EObjectMigrationSide MigrationSide = EObjectMigrationSide::Receive;
 
diff --git a/Engine/Source/Runtime/CoreUObject/Private/UObject/RemoteObject.cpp b/Engine/Source/Runtime/CoreUObject/Private/UObject/RemoteObject.cpp
--- a/Engine/Source/Runtime/CoreUObject/Private/UObject/RemoteObject.cpp
+++ b/Engine/Source/Runtime/CoreUObject/Private/UObject/RemoteObject.cpp
@@ -190,20 +190,6 @@ public:
 		return Stub;
 	}
 
-	void ClearAllPhysicsIds()
-	{
-#if UE_WITH_REMOTE_OBJECT_HANDLE
-		for (TPair<FRemoteObjectId, UE::RemoteObject::Handle::FRemoteObjectStub*>& Elem : RemoteObjects)
-		{
-			if (Elem.Value)
-			{
-				Elem.Value->PhysicsIslandId.PhysicsServerId = FRemoteServerId();
-				Elem.Value->PhysicsIslandId.PhysicsLocalIslandId = 0;
-			}
-		}
-#endif
-	}
-
 	FRemoteObjectPathName* StoreAssetPath(UObject* InObject)
 	{
 		FRemoteObjectId ObjectId(InObject);
@@ -220,44 +206,6 @@ public:
 		FRemoteObjectPathName** PathName = AssetPaths.Find(ObjectId);
 		return PathName ? *PathName : nullptr;
 	}
-
-	/**
-	* This function is called by UE::RemoteObject::Handle::UpdateAllPhysicsServerId. More info is shown in that function's description.
-	*/
-	void UpdateAllPhysicsServerId(const TMap<UE::RemoteObject::Handle::FPhysicsIslandId, UE::RemoteObject::Handle::FPhysicsIslandId>& PhysicsIdMergingMap)
-	{
-#if UE_WITH_REMOTE_OBJECT_HANDLE
-		using UE::RemoteObject::Handle::FPhysicsIslandId;
-		for (TPair<FRemoteObjectId, UE::RemoteObject::Handle::FRemoteObjectStub*>& Elem : RemoteObjects)
-		{
-			if (Elem.Value)
-			{
-				const FPhysicsIslandId OriginalPhysicsId = Elem.Value->PhysicsIslandId;
-				if (const FPhysicsIslandId* MergedPhysicsId = PhysicsIdMergingMap.Find(OriginalPhysicsId))
-				{
-					Elem.Value->PhysicsIslandId.PhysicsServerId = MergedPhysicsId->PhysicsServerId;
-				}
-			}
-		}
-#endif
-	}
-
-	void UpdateAllPhysicsLocalIslandId(const TMap<uint32, uint32>& PhysicsLocalIslandMergingMap)
-	{
-#if UE_WITH_REMOTE_OBJECT_HANDLE
-		for (TPair<FRemoteObjectId, UE::RemoteObject::Handle::FRemoteObjectStub*>& Elem : RemoteObjects)
-		{
-			if (Elem.Value)
-			{
-				const uint32 OriginalLocalIslandId = Elem.Value->PhysicsIslandId.PhysicsLocalIslandId;
-				if (const uint32* MergedLocalIslandId = PhysicsLocalIslandMergingMap.Find(OriginalLocalIslandId))
-				{
-					Elem.Value->PhysicsIslandId.PhysicsLocalIslandId = *MergedLocalIslandId;
-				}
-			}
-		}
-#endif
-	}
 };
 FRemoteObjectMaps* ObjectMaps = nullptr;
 
@@ -328,7 +276,7 @@ void InitRemoteObjects()
 				// Turns a request into an immediate load
 				FUObjectMigrationContext MigrationContext {
 					.ObjectId = ObjectId, .RemoteServerId = DestinationServerId, .OwnerServerId = LastKnownResidentServerId,
-					.PhysicsServerId = LastKnownResidentServerId, .MigrationSide = EObjectMigrationSide::Receive
+					.MigrationSide = EObjectMigrationSide::Receive
 				};
 				UE::RemoteObject::Serialization::Disk::LoadObjectFromDisk(MigrationContext);
 			}
@@ -688,81 +636,6 @@ void ChangeOwnerServerId(const UObject* Object, FRemoteServerId NewOwnerServerId
 #endif
 }
 
-FPhysicsIslandId GetPhysicsIslandId(const FRemoteObjectId ObjectId)
-{
-	using namespace UE::CoreUObject::Private;
-	using namespace UE::RemoteObject::Private;
-	FPhysicsIslandId Result;
-#if UE_WITH_REMOTE_OBJECT_HANDLE
-	UE::RemoteObject::Handle::FRemoteObjectStub* RemoteStub = FindRemoteObjectStub(ObjectId);
-	if (RemoteStub)
-	{
-		Result = RemoteStub->PhysicsIslandId;
-	}
-#endif
-	return Result;
-}
-
-
-FPhysicsIslandId GetPhysicsIslandId(const UObject* Object)
-{
-	using namespace UE::CoreUObject::Private;
-	using namespace UE::RemoteObject::Private;
-#if UE_WITH_REMOTE_OBJECT_HANDLE
-	return GetPhysicsIslandId(FObjectHandleUtils::GetRemoteId(Object));
-#else
-	return FPhysicsIslandId();
-#endif
-}
-
-void ChangePhysicsIslandId(const UObject* Object, FPhysicsIslandId NewPhysicsIslandId)
-{
-	using namespace UE::CoreUObject::Private;
-	using namespace UE::RemoteObject::Private;
-
-#if UE_WITH_REMOTE_OBJECT_HANDLE
-	UE::RemoteObject::Handle::FRemoteObjectStub* RemoteStub = ObjectMaps->FindOrAddRemoteObjectStub(Object);
-	MarkAsRemoteReference(Object);
-	if (RemoteStub)
-	{
-		RemoteStub->PhysicsIslandId = NewPhysicsIslandId;
-	}
-#endif
-}
-
-void ClearAllPhysicsServerId()
-{
-#if UE_WITH_REMOTE_OBJECT_HANDLE
-	if (Private::ObjectMaps)
-	{
-		Private::ObjectMaps->ClearAllPhysicsIds();
-	}
-#endif
-}
-
-void UpdateAllPhysicsServerId(const TMap<FPhysicsIslandId, FPhysicsIslandId>& PhysicsServerMergingMap)
-{
-#if UE_WITH_REMOTE_OBJECT_HANDLE
-	if (Private::ObjectMaps)
-	{
-		Private::ObjectMaps->UpdateAllPhysicsServerId(PhysicsServerMergingMap);
-	}
-#endif
-}
-
-/**
- Updates all FPhysicsIslandId::PhysicsLocalIslandId based on input map. Usage can be found in description of UE::RemoteObject::Handle::UpdateAllPhysicsLocalIslandId
-*/
-void UpdateAllPhysicsLocalIslandId(const TMap<uint32, uint32>& PhysicsLocalIslandMergingMap)
-{
-#if UE_WITH_REMOTE_OBJECT_HANDLE
-	if (Private::ObjectMaps)
-	{
-		Private::ObjectMaps->UpdateAllPhysicsLocalIslandId(PhysicsLocalIslandMergingMap);
-	}
-#endif
-}
-
 FRemoteObjectId GetRemoteObjectId(TFunctionRef<const UObject*()> CodeThatReturnsRawUObjectPtr)
 {
 #if UE_WITH_REMOTE_OBJECT_HANDLE
@@ -1421,13 +1294,6 @@ FArchive& operator<<(FArchive& Ar, FRemoteServerId& Id)
 	return Ar;
 }
 
-void UE::RemoteObject::Handle::FPhysicsIslandId::Reset()
-{
-	PhysicsServerId = FRemoteServerId();
-	PhysicsLocalIslandId = 0;
-}
-
-
 // Debugging functionality to help us find these objects in debug builds
 #if !UE_BUILD_SHIPPING && UE_WITH_REMOTE_OBJECT_HANDLE
 /**
diff --git a/Engine/Source/Runtime/CoreUObject/Private/UObject/RemoteObjectSerialization.cpp b/Engine/Source/Runtime/CoreUObject/Private/UObject/RemoteObjectSerialization.cpp
--- a/Engine/Source/Runtime/CoreUObject/Private/UObject/RemoteObjectSerialization.cpp
+++ b/Engine/Source/Runtime/CoreUObject/Private/UObject/RemoteObjectSerialization.cpp
@@ -105,7 +105,7 @@ void LoadObjectFromDisk(const FUObjectMigrationContext& MigrationContext)
 	// We have transferred ownership from the Database to the local server
 	FRemoteServerId LocalServerId = FRemoteServerId::GetLocalServerId();
 	const FRemoteServerId& DatabaseId = UE::RemoteObject::Transfer::DatabaseId;
-	UE::RemoteObject::Transfer::OnObjectDataReceived(LocalServerId, LocalServerId, 0, MigrationContext.ObjectId, DatabaseId, ObjectData);
+	UE::RemoteObject::Transfer::OnObjectDataReceived(LocalServerId, MigrationContext.ObjectId, DatabaseId, ObjectData);
 }
 
 void SaveObjectToDisk(const UE::RemoteObject::Transfer::FMigrateSendParams& Params)
diff --git a/Engine/Source/Runtime/CoreUObject/Private/UObject/RemoteObjectTransfer.cpp b/Engine/Source/Runtime/CoreUObject/Private/UObject/RemoteObjectTransfer.cpp
--- a/Engine/Source/Runtime/CoreUObject/Private/UObject/RemoteObjectTransfer.cpp
+++ b/Engine/Source/Runtime/CoreUObject/Private/UObject/RemoteObjectTransfer.cpp
@@ -190,13 +190,10 @@ public:
 						*HighestPriorityRequest->RequestPriority.ToString());
 
 					FRemoteServerId OwnerServerId = Handle::GetOwnerServerId(RootObject);
-					FPhysicsIslandId PhysicsIslandId = Handle::GetPhysicsIslandId(RootObject);
 					FUObjectMigrationContext MigrationContext {
 						.ObjectId = HighestPriorityRequest->ObjectId,
 						.RemoteServerId = HighestPriorityRequest->DestinationServerId,
 						.OwnerServerId = OwnerServerId,
-						.PhysicsServerId = PhysicsIslandId.PhysicsServerId,
-						.PhysicsLocalIslandId = PhysicsIslandId.PhysicsLocalIslandId,
 						.MigrationSide = EObjectMigrationSide::Send
 					};
 					FScopedObjectMigrationContext ScopedMigrationContext(MigrationContext);
@@ -550,7 +547,6 @@ public:
 				.ObjectId = ObjectId,
 				.RemoteServerId = OwnerServerId,
 				.OwnerServerId = OwnerServerId,
-				.PhysicsServerId = OwnerServerId,
 				.MigrationSide = EObjectMigrationSide::Send,
 				.MultiServerCommitRequestId = ActiveRequest->RequestId,
 			};
@@ -611,15 +607,12 @@ public:
 		for (UObject* SentObject : MultiServerCommitSentObjects)
 		{
 			FRemoteServerId OwnerServerId = Handle::GetOwnerServerId(SentObject);
-			UE::RemoteObject::Handle::FPhysicsIslandId PhysicsIslandId = Handle::GetPhysicsIslandId(SentObject);
 			FRemoteObjectId SentObjectId = UE::CoreUObject::Private::FObjectHandleUtils::GetRemoteId(SentObject);
 
 			FUObjectMigrationContext MigrationContext {
 				.ObjectId = SentObjectId,
 				.RemoteServerId = OwnerServerId,
 				.OwnerServerId = OwnerServerId,
-				.PhysicsServerId = PhysicsIslandId.PhysicsServerId,
-				.PhysicsLocalIslandId = PhysicsIslandId.PhysicsLocalIslandId,
 				.MigrationSide = EObjectMigrationSide::Send
 			};
 			FScopedObjectMigrationContext ScopedMigrationContext(MigrationContext);
@@ -718,7 +711,6 @@ public:
 				.ObjectId = ObjectId,
 				.RemoteServerId = DatabaseId,
 				.OwnerServerId = DatabaseId,
-				.PhysicsServerId = DatabaseId,
 				.MigrationSide = EObjectMigrationSide::Receive
 			};
 			FScopedObjectMigrationContext ScopedMigrationContext(MigrationContext);
@@ -931,7 +923,6 @@ UObject* MigrateObjectFromObjectData(FRemoteObjectData& ObjectData, const FUObje
 	for (UObject* ReceivedObject : ReceivedObjects)
 	{
 		UE::RemoteObject::Handle::ChangeOwnerServerId(ReceivedObject, MigrationContext.OwnerServerId);
-		UE::RemoteObject::Handle::ChangePhysicsIslandId(ReceivedObject, UE::RemoteObject::Handle::FPhysicsIslandId(MigrationContext.PhysicsServerId, MigrationContext.PhysicsLocalIslandId));
 #if UE_WITH_REMOTE_OBJECT_HANDLE
 		ensureAlwaysMsgf(ReceivedObject->HasAnyInternalFlags(EInternalObjectFlags::Borrowed) == !UE::RemoteObject::Handle::IsOwned(ReceivedObject),
 			TEXT("Object %s (%s) is %s owned but is %s borrowed"),
@@ -1089,13 +1080,10 @@ void TransferObjectOwnershipToRemoteServer(UObject* Object, FRemoteServerId Dest
 	const FRemoteObjectId CanonicalRootObjectId = UE::CoreUObject::Private::FObjectHandleUtils::GetRemoteId(RootObject);
 
 	// Figure out the ownership semantics
-	UE::RemoteObject::Handle::FPhysicsIslandId PhysicsIslandId = Handle::GetPhysicsIslandId(RootObject);
 	FUObjectMigrationContext MigrationContext {
 		.ObjectId = ObjectId,
 		.RemoteServerId = DestinationServerId,
 		.OwnerServerId = DestinationServerId,
-		.PhysicsServerId = PhysicsIslandId.PhysicsServerId,
-		.PhysicsLocalIslandId = PhysicsIslandId.PhysicsLocalIslandId,
 		.MigrationSide = EObjectMigrationSide::Send
 	};
     FScopedObjectMigrationContext ScopedMigrationContext(MigrationContext);
@@ -1188,7 +1176,7 @@ void MigrateObjectToRemoteServerWithExplicitPriority(FRemoteWorkPriority Request
 	GTransferQueue.RequestObjectMigration(RequestPriority, ObjectId, DestinationServerId);
 }
 
-void OnObjectDataReceived(FRemoteServerId OwnerServerId, FRemoteServerId PhysicsId, uint32 PhysicsLocalIslandId, FRemoteObjectId ObjectId, FRemoteServerId RemoteServerId, FRemoteObjectData& Data)
+void OnObjectDataReceived(FRemoteServerId OwnerServerId, FRemoteObjectId ObjectId, FRemoteServerId RemoteServerId, FRemoteObjectData& Data)
 {
 	check(!AutoRTFM::IsClosed());
 
@@ -1196,8 +1184,6 @@ void OnObjectDataReceived(FRemoteServerId OwnerServerId, FRemoteServerId Physics
 		.ObjectId = ObjectId,
 		.RemoteServerId = RemoteServerId,
 		.OwnerServerId = OwnerServerId,
-		.PhysicsServerId = PhysicsId,
-		.PhysicsLocalIslandId = PhysicsLocalIslandId,
 		.MigrationSide = EObjectMigrationSide::Receive
 	};
 	FScopedObjectMigrationContext ScopedMigrationContext(MigrationContext);
@@ -1285,7 +1271,6 @@ void StoreObjectToDatabase(UObject* Object, FRemoteObjectId ObjectId)
 				.ObjectId = RootObjectId,
 				.RemoteServerId = DatabaseId,
 				.OwnerServerId = DatabaseId,
-				.PhysicsServerId = DatabaseId,
 				.MigrationSide = EObjectMigrationSide::Send
 			};
 			FScopedObjectMigrationContext ScopedMigrationContext(StoreToDatabaseContext);
diff --git a/Engine/Source/Runtime/CoreUObject/Public/UObject/RemoteObject.h b/Engine/Source/Runtime/CoreUObject/Public/UObject/RemoteObject.h
--- a/Engine/Source/Runtime/CoreUObject/Public/UObject/RemoteObject.h
+++ b/Engine/Source/Runtime/CoreUObject/Public/UObject/RemoteObject.h
@@ -31,64 +31,6 @@ namespace UE::RemoteObject
 
 namespace UE::RemoteObject::Handle
 {
-	struct FPhysicsIslandId
-	{
-		friend struct FRemoteObjectId;
-
-		FRemoteServerId PhysicsServerId;
-
-		//IMPORTANT: LocalIslandIndex = 0 means it's invalid. 
-		uint32 PhysicsLocalIslandId = 0;
-
-		void Reset();
-
-		FPhysicsIslandId() = default;
-
-		FPhysicsIslandId(const FRemoteServerId& PhysicsServerIdIn, const uint32 PhysicsLocalIslandIdIn)
-			: PhysicsServerId(PhysicsServerIdIn), PhysicsLocalIslandId(PhysicsLocalIslandIdIn)
-		{
-		}
-
-		bool operator==(const FPhysicsIslandId& Other) const
-		{
-			return PhysicsServerId == Other.PhysicsServerId && PhysicsLocalIslandId == Other.PhysicsLocalIslandId;
-		}
-
-		bool IsValid() const
-		{
-			return PhysicsServerId.IsValid() && PhysicsLocalIslandId != 0;
-		}
-
-		bool operator<(const FPhysicsIslandId& Other) const
-		{
-			return PhysicsServerId < Other.PhysicsServerId ||
-				(PhysicsServerId == Other.PhysicsServerId && PhysicsLocalIslandId < Other.PhysicsLocalIslandId);
-		}
-
-		bool operator<=(const FPhysicsIslandId& Other) const
-		{
-			return (*this) < Other || (*this) == Other;
-		}
-
-		friend FArchive& operator<<(FArchive& Ar, FPhysicsIslandId& ID)
-		{
-			Ar << ID.PhysicsServerId;
-			Ar << ID.PhysicsLocalIslandId;
-			return Ar;
-		}
-
-		FString ToString() const
-		{
-			return FString::Printf(TEXT("[%s, %d]"), *(PhysicsServerId.ToString()), PhysicsLocalIslandId);
-		}
-	};
-
-	UE_FORCEINLINE_HINT uint32 GetTypeHash(const FPhysicsIslandId& Key)
-	{
-		// Combine hashes of the members
-		return HashCombine(::GetTypeHash(Key.PhysicsServerId.GetIdNumber()), ::GetTypeHash(Key.PhysicsLocalIslandId));
-	}
-
 	/**
 	* Structure that holds a pointer to a native class or to a pathname of a class on disk (Blueprint etc)
 	*/
@@ -132,8 +74,6 @@ namespace UE::RemoteObject::Handle
 			(note: only valid if the object is Local) */
 		FRemoteServerId OwningServerId;
 
-		FPhysicsIslandId PhysicsIslandId;
-
 		FRemoteObjectStub() = default;
 		UE_DEPRECATED(5.8, "Stubs should only be initialized with UObject properties when they are being added to the internal stub map.")
 		COREUOBJECT_API explicit FRemoteObjectStub(UObject* Object);
@@ -246,25 +186,4 @@ namespace UE::RemoteObject::Handle
 	* Sets the owner server id for a locally resident object
 	*/
 	COREUOBJECT_API void ChangeOwnerServerId(const UObject* Object, FRemoteServerId NewOwnerServerId);
-
-	/**
-	* Get PhysicsIslandId of the input object
-	*/
-	COREUOBJECT_API FPhysicsIslandId GetPhysicsIslandId(const UObject* Object);
-
-	COREUOBJECT_API FPhysicsIslandId GetPhysicsIslandId(const FRemoteObjectId ObjectId);
-
-	COREUOBJECT_API void ChangePhysicsIslandId(const UObject* Object, FPhysicsIslandId NewPhysicsIslandId);
-
-	COREUOBJECT_API void ClearAllPhysicsServerId();
-
-	/**
-	* Updates FPhysicsIslandId::PhysicsServerId using the input map. Updating PhysicsServerId only because it is the only data needed for physics migration. (SeeDstmPhysicsMigration.cpp)
-	*/
-	COREUOBJECT_API void UpdateAllPhysicsServerId(const TMap<FPhysicsIslandId, FPhysicsIslandId>& PhysicsServerMergingMap);
-
-	/**
-	* Updates PhysicsLocalIslandId of FPhysicsIslandId using the input map. Usage in DstmPhysicsMigration.cpp for local island computation.
-	*/
-	COREUOBJECT_API void UpdateAllPhysicsLocalIslandId(const TMap<uint32, uint32>& PhysicsIslandMergingMap);
 }
diff --git a/Engine/Source/Runtime/CoreUObject/Public/UObject/RemoteObjectTransfer.h b/Engine/Source/Runtime/CoreUObject/Public/UObject/RemoteObjectTransfer.h
--- a/Engine/Source/Runtime/CoreUObject/Public/UObject/RemoteObjectTransfer.h
+++ b/Engine/Source/Runtime/CoreUObject/Public/UObject/RemoteObjectTransfer.h
@@ -108,7 +108,7 @@ namespace UE::RemoteObject::Transfer
 	* @param RemoteServerId Server that sent object data
 	* @param Data Remote object data. Data ownership is transferred to the remote object queue
 	*/
-	COREUOBJECT_API void OnObjectDataReceived(FRemoteServerId OwnerServerId, FRemoteServerId PhysicsId, uint32 PhysicsLocalIslandId, FRemoteObjectId ObjectId, FRemoteServerId RemoteServerId, FRemoteObjectData& Data);
+	COREUOBJECT_API void OnObjectDataReceived(FRemoteServerId OwnerServerId, FRemoteObjectId ObjectId, FRemoteServerId RemoteServerId, FRemoteObjectData& Data);
 
 	/**
 	* Called when a remote object request was denied by a remote server
