Fix for initialization tracking on IDO CDOs.
The fix propagates the initialization state on a type's super inheritance chain.

#rnx
#Synced-CL 48962394

diff --git a/Engine/Source/Runtime/CoreUObject/Private/UObject/InstanceDataObjectUtils.cpp b/Engine/Source/Runtime/CoreUObject/Private/UObject/InstanceDataObjectUtils.cpp
--- a/Engine/Source/Runtime/CoreUObject/Private/UObject/InstanceDataObjectUtils.cpp
+++ b/Engine/Source/Runtime/CoreUObject/Private/UObject/InstanceDataObjectUtils.cpp
@@ -1542,6 +1542,15 @@ namespace UE
 	{
 		TRACE_CPUPROFILER_EVENT_SCOPE(TrackDefaultInitializedPropertiesInCDO)
 
+		const UClass* OwnerSuperClass = OwnerClass->GetSuperClass();
+		if (OwnerSuperClass)
+		{
+			if (const UVerseClass* VerseSuperClass = Cast<UVerseClass>(OwnerSuperClass))
+			{
+				TrackDefaultInitializedPropertiesInCDO(CDO, Class, VerseSuperClass);
+			}
+		}
+
 		// If a property has been initialized by the Verse class' CDO then we consider everything
 		// under it (i.e., any sub-objects) as initialized.
 		for (const TFieldPath<FProperty>& OwnerFieldPath : OwnerClass->PropertiesWrittenByInitCDO)
@@ -1555,47 +1564,47 @@ namespace UE
 					UE::FInitializedPropertyValueState(CDO).Set(Property);
 				}
 
-				// Keep track of visited property-owner pairs to avoid referencing cycles.
-				TSet<TTuple<const FProperty*, void*>> VisitedPropOwners;
+				// Keep track of visited property-parent pairs to avoid referencing cycles.
+				TSet<TTuple<const FProperty*, void*>> VisitedPropParents;
 
 				FPropertyVisitorPath Path;
 				FPropertyVisitorData VisitorData(CDO, /*ParentStructData*/CDO);
 				FPropertyVisitorContext VisitorContext(Path, VisitorData, FPropertyVisitorContext::EScope::All);
-				PropertyVisitorHelpers::VisitProperty(Class, Property, VisitorContext, [CDO, &VisitedPropOwners](const FPropertyVisitorContext& Context) -> EPropertyVisitorControlFlow
+				PropertyVisitorHelpers::VisitProperty(Class, Property, VisitorContext, [CDO, &VisitedPropParents](const FPropertyVisitorContext& Context) -> EPropertyVisitorControlFlow
 				{
 					const FPropertyVisitorPath& PropertyPath = Context.Path;
 					const FPropertyVisitorData& Data = Context.Data;
 					const FProperty* Property = PropertyPath.Top().Property;
-					void* Owner = Data.ParentStructData;
-					TTuple<const FProperty*, void*> PropOwner(Property, Owner);
+					void* Parent = Data.ParentStructData;
+					TTuple<const FProperty*, void*> PropParent(Property, Parent);
 
-					if (!Property || VisitedPropOwners.Contains(PropOwner))
+					if (!Property || VisitedPropParents.Contains(PropParent))
 					{
 						return EPropertyVisitorControlFlow::StepOver;
 					}
 
 					bool bIsInCDO = true;
 
-					const UStruct* OwnerType = PropertyPath.Top().ParentStructType;
+					const UStruct* ParentType = PropertyPath.Top().ParentStructType;
 
-					if (OwnerType && OwnerType->IsChildOf<UObject>())
+					if (ParentType && ParentType->IsChildOf<UObject>())
 					{
-						const UObject* OwnerObject = (const UObject*)Owner;
+						const UObject* ParentObject = (const UObject*)Parent;
 
-						if (OwnerObject && (OwnerObject != CDO))
+						if (ParentObject && (ParentObject != CDO))
 						{
-							bIsInCDO = OwnerObject->IsInOuter(CDO);
+							bIsInCDO = ParentObject->IsInOuter(CDO);
 						}
 					}
 
-					// It is possible for the property and owner types to differ during re-instancing
+					// It is possible for the property and parent types to differ during re-instancing
 					// when a new CDO is created. Skip tracking in this case.
-					if (bIsInCDO && OwnerType && Property->HasAnyPropertyFlags(CPF_RequiredParm) && OwnerType->IsChildOf(Property->GetOwnerStruct()))
+					if (bIsInCDO && ParentType && Property->HasAnyPropertyFlags(CPF_RequiredParm) && ParentType->IsChildOf(Property->GetOwnerStruct()))
 					{
-						UE::FInitializedPropertyValueState(OwnerType, Owner).Set(Property);
+						UE::FInitializedPropertyValueState(ParentType, Parent).Set(Property);
 					}
 
-					VisitedPropOwners.Add(PropOwner);
+					VisitedPropParents.Add(PropParent);
 
 					if (bIsInCDO)
 					{
