Entity streaming: Dynamically add imports for External Entities through package linker-invoked static function
- Add new object flag RF_HasDynamicImports to mark objects having dynamic imports
- New Static function "InjectDynamicImportsFor" on all UObjects, invoked when the object has the RF_HasDynamicImports flag, which allows adding new items in the linker tables
- Enable Type Hierarchy table for all cases, as having superclasses for exports is necessary to find existing InjectDynamicImportsFor implementations
- Free up RF_AllocatedInSharedPage and PKG_DynamicImports which were no longer used

#tests Tested cndo, on VersePrototype internalized and externalized the level_entity children, opened maps and re-saved them, ran with -asyncloadingthread -dpcvars=s.AllowMultithreadedLoading=1 to test MT loading as well

#Synced-CL 49088904

diff --git a/Engine/Source/Developer/AssetTools/Private/AssetHeaderPatcher.cpp b/Engine/Source/Developer/AssetTools/Private/AssetHeaderPatcher.cpp
--- a/Engine/Source/Developer/AssetTools/Private/AssetHeaderPatcher.cpp
+++ b/Engine/Source/Developer/AssetTools/Private/AssetHeaderPatcher.cpp
@@ -1550,18 +1550,11 @@ FAssetHeaderPatcher::EResult FAssetHeaderPatcherInner::PatchHeader_Deserialize()
 
 	if (Summary.ImportTypeHierarchiesOffset)
 	{
-		if (UE::Serialization::Private::FImportTypeHierarchy::IsSerializationEnabled())
-		{
-			MemAr.Seek(Summary.ImportTypeHierarchiesOffset);
+		MemAr.Seek(Summary.ImportTypeHierarchiesOffset);
 
-			MemAr << ImportTypeHierarchyMap;
+		MemAr << ImportTypeHierarchyMap;
 
-			HeaderInformation.ImportTypeHierarchiesMapSize = MemAr.Tell() - Summary.ImportTypeHierarchiesOffset;
-		}
-		else
-		{
-			HeaderInformation.ImportTypeHierarchiesMapSize = 0;
-		}
+		HeaderInformation.ImportTypeHierarchiesMapSize = MemAr.Tell() - Summary.ImportTypeHierarchiesOffset;
 	}
 
 	if (Summary.ThumbnailTableOffset)
@@ -3164,19 +3157,17 @@ FAssetHeaderPatcher::EResult FAssetHeaderPatcherInner::PatchHeader_WriteDestinat
 
 		case EPatchedSection::ImportTypeHierarchyMap:
 		{
-			if (UE::Serialization::Private::FImportTypeHierarchy::IsSerializationEnabled())
-			{
-				const int64 MapStartOffset = Writer.Tell();
-				Writer << ImportTypeHierarchyMap;
+			const int64 MapStartOffset = Writer.Tell();
+			Writer << ImportTypeHierarchyMap;
 
-				checkf(!Writer.IsCriticalError(), TEXT("Issue writing %s"), *Writer.GetErrorMessage());
+			checkf(!Writer.IsCriticalError(), TEXT("Issue writing %s"), *Writer.GetErrorMessage());
+
+			const int64 MapSize = Writer.Tell() - MapStartOffset;
+			const int64 Delta = MapSize - SourceSection.Size;
+			checkf(Delta == 0, TEXT("Delta should be Zero. is %d"), (int)Delta);
+			check(Summary.ImportTypeHierarchiesCount == ImportTypeHierarchyMap.Num());
+			check(Summary.ImportTypeHierarchiesOffset == MapStartOffset);
 
-				const int64 MapSize = Writer.Tell() - MapStartOffset;
-				const int64 Delta = MapSize - SourceSection.Size;
-				checkf(Delta == 0, TEXT("Delta should be Zero. is %d"), (int)Delta);
-				check(Summary.ImportTypeHierarchiesCount == ImportTypeHierarchyMap.Num());
-				check(Summary.ImportTypeHierarchiesOffset == MapStartOffset);
-			}
 			break;
 		}
 
diff --git a/Engine/Source/Programs/AutoRTFMTests/Private/UObjectTests.cpp b/Engine/Source/Programs/AutoRTFMTests/Private/UObjectTests.cpp
--- a/Engine/Source/Programs/AutoRTFMTests/Private/UObjectTests.cpp
+++ b/Engine/Source/Programs/AutoRTFMTests/Private/UObjectTests.cpp
@@ -670,7 +670,7 @@ TEST_CASE("UObject.AtomicallySetFlags")
 	UObject* const Object = NewObject<UMyAutoRTFMTestObject>();
 
 	constexpr EObjectFlags OldFlags = EObjectFlags::RF_Public | EObjectFlags::RF_Transient;
-	constexpr EObjectFlags FlagsToAdd = EObjectFlags::RF_Transient | EObjectFlags::RF_AllocatedInSharedPage;
+	constexpr EObjectFlags FlagsToAdd = EObjectFlags::RF_Transient | EObjectFlags::RF_NonPIEDuplicateTransient;
 
 	// We need to ensure we cover the case where we are adding a flag that is already there
 	// and thus cannot just wipe that out if we abort!
@@ -700,7 +700,7 @@ TEST_CASE("UObject.AtomicallyClearFlags")
 	UObject* const Object = NewObject<UMyAutoRTFMTestObject>();
 
 	constexpr EObjectFlags OldFlags = EObjectFlags::RF_Public | EObjectFlags::RF_Transient;
-	constexpr EObjectFlags FlagsToClear = EObjectFlags::RF_Transient | EObjectFlags::RF_AllocatedInSharedPage;
+	constexpr EObjectFlags FlagsToClear = EObjectFlags::RF_Transient | EObjectFlags::RF_NonPIEDuplicateTransient;
 
 	// We need to ensure we cover the case where we are adding a flag that is already there
 	// and thus cannot just wipe that out if we abort!
diff --git a/Engine/Source/Runtime/CoreUObject/Private/Serialization/ImportTypeHierarchy.cpp b/Engine/Source/Runtime/CoreUObject/Private/Serialization/ImportTypeHierarchy.cpp
--- a/Engine/Source/Runtime/CoreUObject/Private/Serialization/ImportTypeHierarchy.cpp
+++ b/Engine/Source/Runtime/CoreUObject/Private/Serialization/ImportTypeHierarchy.cpp
@@ -6,20 +6,6 @@
 namespace UE::Serialization::Private
 {
 
-static int32 GImportTypeHierarchyEnabled = 1;
-static FAutoConsoleVariableRef CVarImportTypeHierarchyEnabled(
-	TEXT("s.ImportTypeHierarchyEnabled"),
-	GImportTypeHierarchyEnabled,
-	TEXT("Controls whether Import Type Hierarchy entries will be loaded/saved."),
-	ECVF_Default
-);
-
-bool FImportTypeHierarchy::IsSerializationEnabled()
-{
-	bool Result = (GImportTypeHierarchyEnabled > 0);
-	return Result;
-}
-
 FArchive& operator<<(FArchive& Ar, FTypeResource& TypeResource)
 {
 	FStructuredArchiveFromArchive(Ar).GetSlot() << TypeResource;
@@ -44,11 +30,6 @@ FArchive& operator<<(FArchive& Ar, FImportTypeHierarchy& I)
 
 void operator<<(FStructuredArchive::FSlot Slot, FImportTypeHierarchy& I)
 {
-	if (!FImportTypeHierarchy::IsSerializationEnabled())
-	{
-		return;
-	}
-
 	FStructuredArchive::FRecord Record = Slot.EnterRecord();
 
 	int32 SuperTypeCount = I.SuperTypes.Num();
diff --git a/Engine/Source/Runtime/CoreUObject/Private/UObject/LinkerLoad.cpp b/Engine/Source/Runtime/CoreUObject/Private/UObject/LinkerLoad.cpp
--- a/Engine/Source/Runtime/CoreUObject/Private/UObject/LinkerLoad.cpp
+++ b/Engine/Source/Runtime/CoreUObject/Private/UObject/LinkerLoad.cpp
@@ -16,6 +16,7 @@
 #include "Misc/ScopedSlowTask.h"
 #include "Misc/ObjectThumbnail.h"
 #include "Misc/App.h"
+#include "UObject/CoreRedirects.h"
 #include "UObject/InstanceDataObjectUtils.h"
 #include "UObject/MetaData.h"
 #include "UObject/LinkerLoadImportBehavior.h"
@@ -26,6 +27,7 @@
 #include "UObject/PackageTrailer.h"
 #include "UObject/PropertyBagRepository.h"
 #include "UObject/SoftObjectPath.h"
+#include "UObject/UObjectGlobals.h"
 #include "UObject/UObjectHash.h"
 #include "UObject/UObjectSerializeContext.h"
 #include "Misc/PackageName.h"
@@ -1048,7 +1050,14 @@ FLinkerLoad::ELinkerStatus FLinkerLoad::ProcessPackageSummary(TMap<TPair<FName,
 		SCOPED_LOADTIMER(LinkerLoad_SerializeImportTypeHierarchies);
 		Status = SerializeImportTypeHierarchies();
 	}
-#endif
+
+	// Add dynamic imports when needed.
+	if (Status == LINKER_Loaded)
+	{
+		SCOPED_LOADTIMER(LinkerLoad_AddDynamicImports);
+		Status = AddDynamicImports();
+	}
+#endif // WITH_EDITORONLY_DATA
 
 	TRACE_LOADTIME_END_PROCESS_SUMMARY;
 
@@ -1186,6 +1195,7 @@ FLinkerLoad::FLinkerLoad(UPackage* InParent, const FPackagePath& InPackagePath,
 , bUseFullTimeLimit(false)
 , bLoaderNeedsEngineVersionChecks(true)
 , bHasSerializedImportTypeHierarchies(false)
+, bHasAddedDynamicImports(false)
 #if WITH_EDITOR
 , bExportsDuplicatesFixed(false)
 , bIsPackageRelocated(false)
@@ -2731,11 +2741,6 @@ FLinkerLoad::ELinkerStatus FLinkerLoad::SerializeImportTypeHierarchies()
 {
 	DECLARE_SCOPE_CYCLE_COUNTER(TEXT("FLinkerLoad::SerializeImportTypeHierarchies"), STAT_LinkerLoad_SerializeImportTypeHierarchies, STATGROUP_LinkerLoad);
 
-	if (!UE::Serialization::Private::FImportTypeHierarchy::IsSerializationEnabled())
-	{
-		return LINKER_Loaded;
-	}
-
 	// Skip serializing import types if the package does not include any or when dealing with a text format.
 	if (bHasSerializedImportTypeHierarchies || (Summary.ImportTypeHierarchiesCount == 0) || (Summary.ImportTypeHierarchiesOffset == 0) || IsTextFormat())
 	{
@@ -2784,7 +2789,158 @@ FLinkerLoad::ELinkerStatus FLinkerLoad::SerializeImportTypeHierarchies()
 	return LINKER_Loaded;
 }
 
-#endif // WITH_EDITORONLY_DATA
+FLinkerLoad::ELinkerStatus FLinkerLoad::AddDynamicImports()
+{
+	DECLARE_SCOPE_CYCLE_COUNTER(TEXT("FLinkerLoad::AddDynamicImports"), STAT_LinkerLoad_AddDynamicImports, STATGROUP_LinkerLoad);
+
+	if (bHasAddedDynamicImports)
+	{
+		return LINKER_Loaded;
+	}
+
+	/* Dynamic imports are imports that were not present in a package, and get added during package loading, before imports (dependencies) are loaded
+	 * To use the feature, implement InjectDynamicImportsFor for the class of an object that must pull on additional imports, and set RF_HasDynamicImports on that object, then save the object's package.
+	 * On its next load, InjectDynamicImportsFor will be invoked, letting the implementation add new entries into the linker tables.
+	 * Be aware that the types of objects saved in a package might not represent the types active in the engine; classes can have been re-parented, deleted, etc. since the last time the package was saved.
+	 * As such when changing the class hierarchy of objects that use Dynamic Imports, mitigation steps might be taken.
+	 * 1 - A class that implemented InjectDynamicImportsFor was deleted; In this case the method will not be invoked anymore when loading a package that still has the type.
+	 * 2 - A class that implements InjectDynamicImportsFor was renamed through class redirectors; the loop below checks into core redirects to try to remap classes, the method should still be invoked on the new type.
+	 * 3 - A new class that implements InjectDynamicImportsFor was added in the class hierarchy of a type; in that case this new type is still not known to this package and the method will not be invoked until the package is re-saved.
+	 * 4 - A new class that implements InjectDynamicImportsFor was added in the class hierarchy of a type, but there already was a more basic class implementing InjectDynamicImportsFor;
+	 *     in that case, the old class will be invoked until the package is re-saved. If re-saving is not possible because of the invocation, you could mitigate with versioning by looking at the branch
+	 *     version with CustomVer(xyz) with the Linker argument, like you would do in Serialize, or disabling it temporarily to re-save, or any other method.
+	 * 5 - A class that implements InjectDynamicImportsFor was removed from the hierarchy, but the class still exists; in that case InjectDynamicImportsFor would still be invoked until the package is re-saved.
+	 *     You would also need mitigation steps like in case number 4. */
+
+	TArray<FDynamicImport> DynamicImports;
+
+	// Check if we have exports that depend on external packages
+	for (int32 ExportIndex = 0; ExportIndex < ExportMap.Num(); ++ExportIndex)
+	{
+		const FObjectExport& Export = ExportMap[ExportIndex];
+		if ((Export.ObjectFlags & RF_HasDynamicImports) == 0)
+		{
+			continue;
+		}
+
+		// Attempt to find a native class corresponding to this export
+		UClass* FoundClass = nullptr;
+		FString ClassName;
+ 
+		if (Export.ClassIndex.IsImport())
+		{
+			ClassName = GetImportPathName(Export.ClassIndex);
+		}
+		else
+		{
+			// Since the class is in this package, we won't find it as it is not loaded yet and we will rely on the ImportTypeHierarchies table below.
+			// But we might find a redirector for it, so we still get the class name.
+			ClassName = GetExportPathName(Export.ClassIndex);
+		}
+ 
+		// Check for class redirectors
+		FCoreRedirectObjectName RedirectedName = FCoreRedirects::GetRedirectedName(ECoreRedirectFlags::Type_Class, ClassName);
+		if (RedirectedName.IsValid())
+		{
+			if (UPackage* TypePackage = FindObjectFast<UPackage>(nullptr, RedirectedName.PackageName, EFindObjectFlags::ExactClass))
+			{
+				FoundClass = FindObjectFast<UClass>(TypePackage, RedirectedName.ObjectName, EFindObjectFlags::None);
+			}
+		}
+		else
+		{
+			FoundClass = Cast<UClass>(StaticFindObject(UClass::StaticClass(), nullptr, ClassName));
+		}
+
+		// If we still cannot find it, go through its saved type hierarchy
+		if (FoundClass == nullptr)
+		{
+			const UE::Serialization::Private::FImportTypeHierarchy* TypeHierarchy = ImportTypeHierarchies.Find(Export.ClassIndex);
+			if (TypeHierarchy != nullptr && TypeHierarchy->GetSuperTypes().Num() > 0)
+			{
+				// Find the first native type, which should always be loaded
+				for (const UE::Serialization::Private::FTypeResource& TypeRes : TypeHierarchy->GetSuperTypes())
+				{
+					// Only retain native classes
+					if (TypeRes.ClassName == NAME_Class || TypeRes.ClassName == NAME_VerseClass)
+					{
+						// Check for class redirectors
+						RedirectedName = FCoreRedirects::GetRedirectedName(ECoreRedirectFlags::Type_Class, { TypeRes.TypeName, FName{}, TypeRes.PackageName } );
+						if (RedirectedName.IsValid())
+						{
+							if (UPackage* TypePackage = FindObjectFast<UPackage>(nullptr, RedirectedName.PackageName, EFindObjectFlags::ExactClass))
+							{
+								FoundClass = FindObjectFast<UClass>(TypePackage, RedirectedName.ObjectName, EFindObjectFlags::None);
+							}
+						}
+						else
+						{
+							if (UPackage* TypePackage = FindObjectFast<UPackage>(nullptr, TypeRes.PackageName, EFindObjectFlags::ExactClass))
+							{
+								FoundClass = FindObjectFast<UClass>(TypePackage, TypeRes.TypeName, EFindObjectFlags::None);
+							}
+						}
+
+						if (FoundClass != nullptr)
+						{
+							break;
+						}
+					}
+				}
+			}
+		}
+
+		if (FoundClass != nullptr)
+		{
+			DynamicImports.SetNum(0);
+			FoundClass->CallInjectDynamicImportsFor(*this, Export, DynamicImports);
+
+			if (DynamicImports.Num() > 0)
+			{
+				ImportMap.Reserve(ImportMap.Num() + DynamicImports.Num());
+
+				for (const FDynamicImport& DynamicImport : DynamicImports)
+				{
+					if (DynamicImport.PackageName == NAME_None)
+					{
+						continue;
+					}
+
+					FObjectImport& Import = ImportMap.AddDefaulted_GetRef();
+					Import.ObjectName = DynamicImport.PackageName;
+					Import.ClassPackage = GLongCoreUObjectPackageName;
+					Import.ClassName = NAME_Package;
+					Import.PackageName = NAME_None;
+					Import.SourceIndex = -1;
+					Import.bImportOptional = false;
+					Import.bImportPackageHandled = false;
+					Import.bImportSearchedFor = false;
+					Import.bImportFailed = false;
+					Import.XObject = nullptr;
+					Import.SourceLinker = nullptr;
+
+					if (DynamicImport.InstanceName != NAME_None && DynamicImport.InstanceName != DynamicImport.PackageName)
+					{
+						InstancingContext.AddPackageMapping(DynamicImport.PackageName, DynamicImport.InstanceName);
+					}
+				}
+
+				Summary.ImportCount = ImportMap.Num();
+			}
+		}
+		else
+		{
+			UE_ASSET_LOG(LogLinker, Error, PackagePath, TEXT("AddDynamicImports failed to Find Class \"%s\" and its SuperClasses, some dynamic imports will not load correctly"),
+				*ClassName);
+		}
+	}
+
+	bHasAddedDynamicImports = true;
+
+	return LINKER_Loaded;
+}
+
+#endif // WITH_EDITOR
 
 /**
  * Serializes thumbnails
@@ -7806,4 +7962,4 @@ void FLinkerLoad::SerializeBulkMeta(UE::BulkData::Private::FBulkMetaData& Meta,
 #endif // WITH_EDITOR
 }
 
-#undef LOCTEXT_NAMESPACE
\ No newline at end of file
+#undef LOCTEXT_NAMESPACE
diff --git a/Engine/Source/Runtime/CoreUObject/Private/UObject/Obj.cpp b/Engine/Source/Runtime/CoreUObject/Private/UObject/Obj.cpp
--- a/Engine/Source/Runtime/CoreUObject/Private/UObject/Obj.cpp
+++ b/Engine/Source/Runtime/CoreUObject/Private/UObject/Obj.cpp
@@ -1934,6 +1934,10 @@ void UObject::AppendToClassSchema(FAppendToClassSchemaContext& Context)
 void UObject::DeclareConstructClasses(TArray<FTopLevelAssetPath>& OutConstructClasses, const UClass* SpecificSubclass) 
 {
 }
+
+void UObject::InjectDynamicImportsFor(const class FLinkerLoad& Linker, const struct FObjectExport& Export, TArray<struct FDynamicImport>& OutDynamicImports)
+{
+}
 #endif
 
 void UObject::SerializeScriptProperties(FArchive& Ar) const
diff --git a/Engine/Source/Runtime/CoreUObject/Private/UObject/SavePackage/PackageHarvester.cpp b/Engine/Source/Runtime/CoreUObject/Private/UObject/SavePackage/PackageHarvester.cpp
--- a/Engine/Source/Runtime/CoreUObject/Private/UObject/SavePackage/PackageHarvester.cpp
+++ b/Engine/Source/Runtime/CoreUObject/Private/UObject/SavePackage/PackageHarvester.cpp
@@ -1406,9 +1406,11 @@ void FPackageHarvester::HarvestImportTypeHierarchyNames(const UObject* Import)
 	const UStruct* Struct = Cast<UStruct>(Import);
 	if (Struct)
 	{
-		// Only harvest names for types that can be used to create placeholders with.
+		// Harvest names for types that can be used to create placeholders with, or types from an export with RF_HasDynamicImports
 		// See comment in SavePackage2.cpp.
-		if (UE::CanCreatePropertyBagPlaceholderForType(Struct->GetClass(), Struct))
+		bool bRequiresNames = UE::CanCreatePropertyBagPlaceholderForType(Struct->GetClass(), Struct) ||
+			(CurrentExportDependencies.CurrentExport && CurrentExportDependencies.CurrentExport->HasAnyFlags(RF_HasDynamicImports));
+		if (bRequiresNames)
 		{
 			const UStruct* SuperStruct = Struct->GetSuperStruct();
 
@@ -1608,4 +1610,4 @@ ESaveableStatus FPackageHarvester::GetSaveableStatusForRealm(UObject* Obj, ESave
 	}
 
 	return ESaveableStatus::Success;
-}
\ No newline at end of file
+}
diff --git a/Engine/Source/Runtime/CoreUObject/Private/UObject/SavePackage2.cpp b/Engine/Source/Runtime/CoreUObject/Private/UObject/SavePackage2.cpp
--- a/Engine/Source/Runtime/CoreUObject/Private/UObject/SavePackage2.cpp
+++ b/Engine/Source/Runtime/CoreUObject/Private/UObject/SavePackage2.cpp
@@ -1602,6 +1602,7 @@ ESavePackageResult BuildLinker(FSaveContext& SaveContext)
 	// Map Export Indices
 	{
 		SCOPED_SAVETIMER(UPackage_Save_MapExportIndices);
+
 		for (FObjectExport& Export : Linker->ExportMap)
 		{
 			// Set class index.
@@ -1684,7 +1685,7 @@ ESavePackageResult BuildLinker(FSaveContext& SaveContext)
 
 #if WITH_EDITORONLY_DATA
 		// Build ImportTypeHierarchies
-		if (!SaveContext.IsCooking() && UE::Serialization::Private::FImportTypeHierarchy::IsSerializationEnabled())
+		if (!SaveContext.IsCooking())
 		{
 			SCOPED_SAVETIMER(UPackage_Save_BuildImportTypeHierarchies);
 
@@ -1707,8 +1708,9 @@ ESavePackageResult BuildLinker(FSaveContext& SaveContext)
 								if (!Linker->ImportTypeHierarchies.Contains(PackageIndex))
 								{
 									// To avoid bloating the ImportTypeHierarchy table only create entries for types 
-									// that support PropertyBag placeholders.
-									if (UE::CanCreatePropertyBagPlaceholderForType(Struct->GetClass(), Struct))
+									// that support PropertyBag placeholders, or when we will have dynamic imports
+									const bool bExportHasDynamicImports = ((Export.ObjectFlags & RF_HasDynamicImports) != 0);
+									if (bExportHasDynamicImports || UE::CanCreatePropertyBagPlaceholderForType(Struct->GetClass(), Struct))
 									{
 										UE::Serialization::Private::FImportTypeHierarchy& ImportTypeHierarchy = Linker->ImportTypeHierarchies.Add(PackageIndex);
 
@@ -2518,7 +2520,7 @@ AUTORTFM_DISABLE ESavePackageResult WritePackageHeader(FStructuredArchive::FReco
 	Linker->Summary.ImportTypeHierarchiesCount = 0;
 	Linker->Summary.ImportTypeHierarchiesOffset = 0;
 
-	if (!SaveContext.IsCooking() && UE::Serialization::Private::FImportTypeHierarchy::IsSerializationEnabled())
+	if (!SaveContext.IsCooking())
 	{
 		SCOPED_SAVETIMER(UPackage_Save_SaveImportTypeHierarchies);
 
diff --git a/Engine/Source/Runtime/CoreUObject/Public/Serialization/ImportTypeHierarchy.h b/Engine/Source/Runtime/CoreUObject/Public/Serialization/ImportTypeHierarchy.h
--- a/Engine/Source/Runtime/CoreUObject/Public/Serialization/ImportTypeHierarchy.h
+++ b/Engine/Source/Runtime/CoreUObject/Public/Serialization/ImportTypeHierarchy.h
@@ -45,8 +45,6 @@ public:
 	UE_INTERNAL void AddSuperType(const FTypeResource& SuperType) { SuperTypes.Add(SuperType); }
 	UE_INTERNAL void AddSuperType(FTypeResource&& SuperType) { SuperTypes.Add(SuperType); }
 
-	UE_INTERNAL COREUOBJECT_API static bool IsSerializationEnabled();
-
 	/** I/O functions */
 	friend COREUOBJECT_API FArchive& operator<<(FArchive& Ar, FImportTypeHierarchy& I);
 	friend COREUOBJECT_API void operator<<(FStructuredArchive::FSlot Slot, FImportTypeHierarchy& I);
diff --git a/Engine/Source/Runtime/CoreUObject/Public/UObject/Class.h b/Engine/Source/Runtime/CoreUObject/Public/UObject/Class.h
--- a/Engine/Source/Runtime/CoreUObject/Public/UObject/Class.h
+++ b/Engine/Source/Runtime/CoreUObject/Public/UObject/Class.h
@@ -3950,6 +3950,13 @@ public:
 		check(CppClassStaticFunctions.GetDeclareConstructClasses());
 		CppClassStaticFunctions.GetDeclareConstructClasses()(OutConstructClasses, this);
 	}
+
+	/** Calls the c++ class's InjectDynamicImportsFor static function (from the nearest native parent if this is not native) */
+	void CallInjectDynamicImportsFor(const class FLinkerLoad& Linker, const struct FObjectExport& Export, TArray<struct FDynamicImport>& OutDynamicImports) const
+	{
+		check(CppClassStaticFunctions.GetInjectDynamicImportsFor());
+		CppClassStaticFunctions.GetInjectDynamicImportsFor()(Linker, Export, OutDynamicImports);
+	}
 #endif
 
 	/** The class default object; used for delta serialization and object initialization */
diff --git a/Engine/Source/Runtime/CoreUObject/Public/UObject/LinkerLoad.h b/Engine/Source/Runtime/CoreUObject/Public/UObject/LinkerLoad.h
--- a/Engine/Source/Runtime/CoreUObject/Public/UObject/LinkerLoad.h
+++ b/Engine/Source/Runtime/CoreUObject/Public/UObject/LinkerLoad.h
@@ -471,6 +471,8 @@ private:
 	bool					bLoaderNeedsEngineVersionChecks : 1;
 	/** Whether we already serialized import type hierarchies.																*/
 	bool					bHasSerializedImportTypeHierarchies : 1;
+	/** Whether we already added dynamic imports.																			*/
+	bool					bHasAddedDynamicImports : 1;
 
 #if WITH_EDITOR
 	/** Check to avoid multiple export duplicate fixups in case we don't save asset. */
@@ -815,9 +817,9 @@ public:
 	COREUOBJECT_API static bool GetPreloadingEnabled();
 	COREUOBJECT_API static void SetPreloadingEnabled(bool bEnabled);
 	COREUOBJECT_API static bool TryGetPreloadedLoader(const FPackagePath& InPackagePath, FOpenPackageResult& OutResult);
-	private:
+private:
 		static bool bPreloadingEnabled;
-	public:
+public:
 #endif
 
 	/** 
@@ -1355,11 +1357,16 @@ private:
 	 */
 	ELinkerStatus SerializeDataResourceMap();
 
+#if WITH_EDITORONLY_DATA
 	/**
 	 * Serializes the import type hierarchy map.
 	 */
-#if WITH_EDITORONLY_DATA
 	FLinkerLoad::ELinkerStatus SerializeImportTypeHierarchies();
+
+	/**
+	 * Handle class callbacks for each export.
+	 */
+	ELinkerStatus AddDynamicImports();
 #endif
 
 	/** Sets the basic linker archive info */
diff --git a/Engine/Source/Runtime/CoreUObject/Public/UObject/Object.h b/Engine/Source/Runtime/CoreUObject/Public/UObject/Object.h
--- a/Engine/Source/Runtime/CoreUObject/Public/UObject/Object.h
+++ b/Engine/Source/Runtime/CoreUObject/Public/UObject/Object.h
@@ -418,9 +418,17 @@ public:
 	 *        class because each class calls Super::DeclareConstructClasses.
 	 */
 	static COREUOBJECT_API void DeclareConstructClasses(TArray<FTopLevelAssetPath>& OutConstructClasses, const UClass* SpecificSubclass);
+	/**
+	 * Allows class to inject items into the linker tables of the package that is loading an export of this class, for example
+	 * adding dynamically determined imports. Only objects with RF_HasDynamicImports currently invoke this.
+	 * @param Linker  A reference to the linker used for loading the export
+	 * @param Export  The export being processed
+	 * @param OutDynamicImports  The dynamic imports to inject
+	 */
+	static COREUOBJECT_API void InjectDynamicImportsFor(const class FLinkerLoad& Linker, const struct FObjectExport& Export, TArray<struct FDynamicImport>& OutDynamicImports);
 #endif
 
-	/** After a critical error, perform any mission-critical cleanup, such as restoring the video mode orreleasing hardware resources. */
+	/** After a critical error, perform any mission-critical cleanup, such as restoring the video mode or releasing hardware resources. */
 	virtual void ShutdownAfterError() {}
 
 	UE_DEPRECATED(5.3, "Matinee has been removed in 5.0, nothing is calling PostInterpChange anymore")
diff --git a/Engine/Source/Runtime/CoreUObject/Public/UObject/ObjectMacros.h b/Engine/Source/Runtime/CoreUObject/Public/UObject/ObjectMacros.h
--- a/Engine/Source/Runtime/CoreUObject/Public/UObject/ObjectMacros.h
+++ b/Engine/Source/Runtime/CoreUObject/Public/UObject/ObjectMacros.h
@@ -157,7 +157,7 @@ enum EPackageFlags
 //	PKG_Unused						= 0x02000000,	
 //	PKG_Unused						= 0x04000000,
 	PKG_CookGenerated				= 0x08000000,	///< This package was generated by the cooker and does not exist in the WorkspaceDomain
-	PKG_DynamicImports				= 0x10000000,	///< This package should resolve dynamic imports from its export at runtime.
+	PKG_DynamicImports UE_DEPRECATED(5.8, "This flag is now obsolete.") = 0x10000000,	///< Obsolete
 	PKG_RuntimeGenerated			= 0x20000000,	///< This package contains elements that are runtime generated, and may not follow standard loading order rules
 	PKG_ReloadingForCooker			= 0x40000000,   ///< This package is reloading in the cooker, try to avoid getting data we will never need. We won't save this package.
 	PKG_FilterEditorOnly			= 0x80000000,	///< Package has editor-only data filtered out
@@ -625,7 +625,8 @@ enum EObjectFlags
 	// RF_MirroredGarbage is mirrored in EInternalObjectFlags::Garbage because checking the internal flags is much faster for the Garbage Collector
 	// while checking the object flags is much faster outside of it where the Object pointer is already available and most likely cached.
 	RF_MirroredGarbage			=0x40000000,	///< Garbage from logical point of view and should not be referenced. This flag is mirrored in EInternalObjectFlags as Garbage for performance
-	RF_AllocatedInSharedPage	=0x80000000,	///< Allocated from a ref-counted page shared with other UObjects
+	RF_HasDynamicImports		=0x80000000,	///< This object has dynamic imports
+	RF_AllocatedInSharedPage UE_DEPRECATED(5.8, "This flag is now obsolete as the functionality is unnecessary.") =0x80000000,	///< Allocated from a ref-counted page shared with other UObjects
 
 	// !!!!!!!!!!!!!!!!!!!!!!!!!!!!!!!!!!!!!!!!!!!!!!!!!!!!!!!!!!!!!!!!!!!!!!!!!!!!!!!!!!!!!!!!!!!!!!!!!!!!!!!!!!!!!!!!!!!!!!!!!!!
 	// DO NOT ADD composite flags to EObjectFlags. Debugger visualisations have trouble displaying composite flag values as text.
@@ -645,7 +646,7 @@ inline constexpr EObjectFlags RF_HasPlaceholderType = EObjectFlags(0x20000000);
 #define RF_AllFlags				(EObjectFlags)0xffffffff	///< All flags, used mainly for error checking
 
 /** Flags to load from unreal asset files */
-#define RF_Load						((EObjectFlags)(RF_Public | RF_Standalone | RF_Transactional | RF_ClassDefaultObject | RF_ArchetypeObject | RF_DefaultSubObject | RF_TextExportTransient | RF_InheritableComponentTemplate | RF_DuplicateTransient | RF_NonPIEDuplicateTransient | RF_MigratingAsset)) 
+#define RF_Load						((EObjectFlags)(RF_Public | RF_Standalone | RF_Transactional | RF_ClassDefaultObject | RF_ArchetypeObject | RF_DefaultSubObject | RF_TextExportTransient | RF_InheritableComponentTemplate | RF_DuplicateTransient | RF_NonPIEDuplicateTransient | RF_MigratingAsset | RF_HasDynamicImports)) 
 
 /** Sub-objects will inherit these flags from their SuperObject */
 #define RF_PropagateToSubObjects	((EObjectFlags)(RF_Public | RF_ArchetypeObject | RF_Transactional | RF_Transient))
@@ -2078,7 +2079,8 @@ public: \
 	#define UOBJECT_CPPCLASS_STATICFUNCTIONS_WITHEDITORONLYDATA(TClass) \
 		, FUObjectCppClassStaticFunctions::DeclareCustomVersionsType(&TClass::DeclareCustomVersions) \
 		, FUObjectCppClassStaticFunctions::AppendToClassSchemaType(&TClass::AppendToClassSchema) \
-		, FUObjectCppClassStaticFunctions::DeclareConstructClassesType(&TClass::DeclareConstructClasses)
+		, FUObjectCppClassStaticFunctions::DeclareConstructClassesType(&TClass::DeclareConstructClasses) \
+		, FUObjectCppClassStaticFunctions::InjectDynamicImportsForType(&TClass::InjectDynamicImportsFor)
 		/* UObjectCppClassStaticFunctions: Extend this macro with the address of your new static function, if it is editor-only. */
 		/* Order must match the order in the FUObjectCppClassStaticFunctions constructor. */
 
@@ -2122,6 +2124,10 @@ struct TUObjectCppClassStaticThunks
 	{
 		T::DeclareConstructClasses(OutConstructClasses, SpecificSubclass);
 	}
+	static void InjectDynamicImportsFor(const class FLinkerLoad& Linker, const struct FObjectExport& Export, TArray<struct FDynamicImport>& OutDynamicImports)
+	{
+		T::InjectDynamicImportsFor(Linker, Export, OutDynamicImports);
+	}
 #endif // WITH_EDITORONLY_DATA
 #endif // UE_WITH_CONSTINIT_UOBJECT && !IS_MONOLITHIC
 };
@@ -2137,9 +2143,10 @@ struct FUObjectCppClassStaticFunctions
 public:
 	typedef void (*AddReferencedObjectsType)	(UObject* ThisObject, FReferenceCollector& Ar);
 #if WITH_EDITORONLY_DATA
-	typedef void (*DeclareCustomVersionsType)   (FArchive& Ar, const UClass* SpecificSubclass);
-	typedef void (*AppendToClassSchemaType)   (FAppendToClassSchemaContext& Context);
-	typedef void (*DeclareConstructClassesType)   (TArray<FTopLevelAssetPath>& OutConstructClasses, const UClass* SpecificSubclass);
+	typedef void (*DeclareCustomVersionsType)	(FArchive& Ar, const UClass* SpecificSubclass);
+	typedef void (*AppendToClassSchemaType)		(FAppendToClassSchemaContext& Context);
+	typedef void (*DeclareConstructClassesType)	(TArray<FTopLevelAssetPath>& OutConstructClasses, const UClass* SpecificSubclass);
+	typedef void (*InjectDynamicImportsForType)	(const class FLinkerLoad& Linker, const struct FObjectExport& Export, TArray<struct FDynamicImport>& OutDynamicImports);
 #endif
 	// UObjectCppClassStaticFunctions: Extend this list of types with the type of your new static function.
 
@@ -2155,7 +2162,7 @@ public:
 		return FUObjectCppClassStaticFunctions(
 			&DispatchType::AddReferencedObjects
 	#if WITH_EDITORONLY_DATA
-			, &DispatchType::DeclareCustomVersions, &DispatchType::AppendToClassSchema, &DispatchType::DeclareConstructClasses
+			, &DispatchType::DeclareCustomVersions, &DispatchType::AppendToClassSchema, &DispatchType::DeclareConstructClasses, &DispatchType::InjectDynamicImportsFor
 	#endif
 		);
 	}
@@ -2166,6 +2173,7 @@ public:
 		, DeclareCustomVersionsType InDeclareCustomVersions
 		, AppendToClassSchemaType InAppendToClassSchema
 		, DeclareConstructClassesType InDeclareConstructClasses
+		, InjectDynamicImportsForType InInjectDynamicImportsFor
 #endif
 	)
 		: AddReferencedObjects(InAddReferencedObjects)
@@ -2173,6 +2181,7 @@ public:
 		, DeclareCustomVersions(InDeclareCustomVersions)
 		, AppendToClassSchema(InAppendToClassSchema)
 		, DeclareConstructClasses(InDeclareConstructClasses)
+		, InjectDynamicImportsFor(InInjectDynamicImportsFor)
 #endif
 	{
 		// Null elements are not valid in this constructor
@@ -2181,6 +2190,7 @@ public:
 		check(InDeclareCustomVersions);
 		check(InAppendToClassSchema);
 		check(InDeclareConstructClasses);
+		check(InInjectDynamicImportsFor);
 #endif
 		// UObjectCppClassStaticFunctions: Extend the constructor with initializers for your new static function member.
 		// Order must match the order in UOBJECT_CPPCLASS_STATICFUNCTIONS_FORCLASS.
@@ -2238,6 +2248,15 @@ public:
 		check(InDeclareConstructClasses != nullptr); // It is not valid to clear single elements (see IsInitialized). Call Reset to clear all elements.
 		DeclareConstructClasses = InDeclareConstructClasses;
 	}
+	InjectDynamicImportsForType GetInjectDynamicImportsFor() const
+	{
+		return InjectDynamicImportsFor;
+	}
+	void SetInjectDynamicImportsFor(InjectDynamicImportsForType InInjectDynamicImportsFor)
+	{
+		check(InInjectDynamicImportsFor != nullptr); // It is not valid to clear single elements (see IsInitialized). Call Reset to clear all elements.
+		InjectDynamicImportsFor = InInjectDynamicImportsFor;
+	}
 #endif
 	// UObjectCppClassStaticFunctions: Extend the list of accessors for your new static function.
 
@@ -2247,6 +2266,7 @@ private:
 	DeclareCustomVersionsType DeclareCustomVersions = nullptr;
 	AppendToClassSchemaType AppendToClassSchema = nullptr;
 	DeclareConstructClassesType DeclareConstructClasses = nullptr;
+	InjectDynamicImportsForType InjectDynamicImportsFor = nullptr;
 #endif
 	// UObjectCppClassStaticFunctions: Extend this list of members with the member for your new static function.
 
diff --git a/Engine/Source/Runtime/CoreUObject/Public/UObject/ObjectResource.h b/Engine/Source/Runtime/CoreUObject/Public/UObject/ObjectResource.h
--- a/Engine/Source/Runtime/CoreUObject/Public/UObject/ObjectResource.h
+++ b/Engine/Source/Runtime/CoreUObject/Public/UObject/ObjectResource.h
@@ -630,3 +630,17 @@ struct FCellImport : FCellResource
 	friend COREUOBJECT_API FArchive& operator<<(FArchive& Ar, FCellImport& I);
 	friend COREUOBJECT_API void operator<<(FStructuredArchive::FSlot Slot, FCellImport& I);
 };
+
+/** Represents a dynamically injected import. Eventually turned into a FObjectImport. */
+struct FDynamicImport
+{
+	/**
+	 * Package Name to import.
+	 */
+	FName PackageName;
+
+	/**
+	 * Optional instance name to rename the package.
+	 */
+	FName InstanceName;
+};
diff --git a/Engine/Source/Runtime/Engine/Private/Level.cpp b/Engine/Source/Runtime/Engine/Private/Level.cpp
--- a/Engine/Source/Runtime/Engine/Private/Level.cpp
+++ b/Engine/Source/Runtime/Engine/Private/Level.cpp
@@ -2686,13 +2686,6 @@ void ULevel::PreEditUndo()
 	if (!IsUsingExternalActors())
 	{
 		Super::PreEditUndo();
-		// Since package don't record their package flag in transaction, sync the level package dynamic import flag
-		GetPackage()->ClearPackageFlags(PKG_DynamicImports);
-	}
-	else
-	{
-		// Since package don't record their package flag in transaction, sync the level package dynamic import flag
-		GetPackage()->SetPackageFlags(PKG_DynamicImports);
 	}
 
 	// Detach existing model components.  These are left in the array, so they are saved for undoing the undo.
@@ -4210,15 +4203,6 @@ bool ULevel::IsUsingExternalActors() const
 void ULevel::SetUseExternalActors(bool bEnable)
 {
 	bUseExternalActors = bEnable;
-	UPackage* LevelPackage = GetPackage();
-	if (bEnable)
-	{
-		LevelPackage->SetPackageFlags(PKG_DynamicImports);
-	}
-	else
-	{
-		LevelPackage->ClearPackageFlags(PKG_DynamicImports);
-	}
 	CreateOrUpdateActorFolders();
 }
 
