Serialize packed remote object headers along with remote object data for access to remote object name and class without the need to deserialize its data

#Synced-CL 49104830

diff --git a/Engine/Source/Runtime/CoreUObject/Private/UObject/Class.cpp b/Engine/Source/Runtime/CoreUObject/Private/UObject/Class.cpp
--- a/Engine/Source/Runtime/CoreUObject/Private/UObject/Class.cpp
+++ b/Engine/Source/Runtime/CoreUObject/Private/UObject/Class.cpp
@@ -8127,6 +8127,12 @@ UScriptStruct* TBaseStructure<FRemoteObjectBytes>::Get()
 	return ScriptStruct;
 }
 
+UScriptStruct* TBaseStructure<FPackedRemoteObjectHeader>::Get()
+{
+	static UScriptStruct* ScriptStruct = StaticGetBaseStructureInternal(TEXT("PackedRemoteObjectHeader"));
+	return ScriptStruct;
+}
+
 UScriptStruct* TBaseStructure<FRemoteObjectData>::Get()
 {
 	static UScriptStruct* ScriptStruct = StaticGetBaseStructureInternal(TEXT("RemoteObjectData"));
diff --git a/Engine/Source/Runtime/CoreUObject/Private/UObject/RemoteObjectSerialization.cpp b/Engine/Source/Runtime/CoreUObject/Private/UObject/RemoteObjectSerialization.cpp
--- a/Engine/Source/Runtime/CoreUObject/Private/UObject/RemoteObjectSerialization.cpp
+++ b/Engine/Source/Runtime/CoreUObject/Private/UObject/RemoteObjectSerialization.cpp
@@ -42,7 +42,7 @@ static FAutoConsoleVariableRef CVarUseImmutableArchetypes(
 
 const UObject* FindImmutableArchetype(const UObject* InObj);
 
-using FNameIndexType = FPackedRemoteObjectPathName::FNameIndexType;
+using FNameIndexType = FRemoteObjectTables::FNameIndexType;
 
 DEFINE_LOG_CATEGORY_STATIC(LogRemoteSerialization, Log, All);
 
@@ -52,6 +52,15 @@ FArchive& operator<<(FArchive& Ar, FRemoteObjectBytes& Chunk)
 	return Ar;
 }
 
+FArchive& operator<<(FArchive& Ar, FPackedRemoteObjectHeader& Header)
+{
+	Ar << Header.NameIndex;
+	Ar << Header.IdIndex;
+	Ar << Header.ClassIndex;
+
+	return Ar;
+}
+
 namespace UE::RemoteObject::Serialization::Disk
 {
 
@@ -66,7 +75,7 @@ void SerializeNameTables(FArchive& Ar, const FRemoteObjectData& InObjectData)
 	}
 	Ar << const_cast<TArray<FRemoteObjectId>&>(InObjectData.Tables.RemoteIds);
 	Ar << const_cast<TArray<FPackedRemoteObjectPathName>&>(InObjectData.PathNames);
-	Ar << const_cast<TArray<FPackedRemoteObjectPathName::FNameIndexType>&>(InObjectData.SerializedObjectIds);
+	Ar << const_cast<TArray<FPackedRemoteObjectHeader>&>(InObjectData.SerializedObjectHeaders);
 }
 
 void DeserializeNameTables(FArchive& Ar, FRemoteObjectData& OutObjectData)
@@ -82,7 +91,7 @@ void DeserializeNameTables(FArchive& Ar, FRemoteObjectData& OutObjectData)
 	}
 	Ar << OutObjectData.Tables.RemoteIds;
 	Ar << OutObjectData.PathNames;
-	Ar << OutObjectData.SerializedObjectIds;
+	Ar << OutObjectData.SerializedObjectHeaders;
 }
 
 FString GenerateRemoteObjectFilename(FRemoteObjectId ObjectId, FRemoteServerId OwnerServerId)
@@ -435,7 +444,7 @@ public:
 		using namespace UE::RemoteObject::Private;
 		using namespace UE::CoreUObject::Private;
 
-		OutHeader.Name = Object->GetFName();
+		OutHeader.Name = AddNameToNameMap(Object->GetFName());
 		OutHeader.RemoteId = AddRemoteIdToIdMap(FObjectHandleUtils::GetRemoteId(Object));
 		OutHeader.Class = Object->GetClass();
 		UObject* Outer = FObjectHandleUtils::GetNonAccessTrackedOuterNoResolve(Object);
@@ -443,7 +452,10 @@ public:
 		OutHeader.Archetype = Object->GetArchetype();
 		OutHeader.InternalFlags = (int32)(Object->GetInternalFlags() & EInternalObjectFlags::Garbage);
 
-		ObjectData.SerializedObjectIds.Add(OutHeader.RemoteId);
+		FPackedRemoteObjectHeader& PackedHeader = ObjectData.SerializedObjectHeaders.Emplace_GetRef();
+		PackedHeader.NameIndex = OutHeader.Name;
+		PackedHeader.IdIndex = OutHeader.RemoteId;
+		PackedHeader.ClassIndex = AddPathNameToNameMap(OutHeader.Class);
 
 		return true;
 	}
@@ -514,7 +526,7 @@ public:
 	FArchiveRemoteObjectWriter(UObject* InRootObject, FRemoteObjectData& OutObjectData, const FUObjectMigrationContext* InMigrationContext, const TCHAR* ArchiveName = nullptr, TSet<UObject*>* OutReferencedObjectsSet = nullptr)
 		: TArchiveRemoteObjectBase<FMemoryWriter>(OutObjectData, InMigrationContext, ArchiveName ? ArchiveName : TEXT("RemoteObjectWriter"))
 	{
-		static_assert(sizeof(FNameIndexType) == sizeof(FPackedRemoteObjectPathName::FNameIndexType), "Name Index type must match the type used for storing remote object path name indices");
+		static_assert(sizeof(FNameIndexType) == sizeof(FRemoteObjectTables::FNameIndexType), "Name Index type must match the type used for storing remote object path name indices");
 
 		SetRootObject(InRootObject);
 		ObjectsToSerialize.Add(InRootObject);
@@ -630,8 +642,7 @@ public:
 	}
 
 	FORCENOINLINE void SerializeRemoteObject(UObject* Object, const FRemoteObjectHeader& Header)
-	{
-		AddNameToNameMap(Header.Name);
+	{		
 		Object->Serialize(*this);
 	}
 };
@@ -978,16 +989,16 @@ FRemoteObjectConstructionOverrides::FRemoteObjectConstructionOverrides(const FRe
 	for (const FRemoteObjectHeader& Header : InObjectHeaders)
 	{
 		FRemoteObjectConstructionParams& Params = Overrides.Emplace_GetRef();
-		Params.Name = Header.Name;
+		Params.Name = ObjectData.GetName(Header.Name);
 #if UE_WITH_REMOTE_OBJECT_HANDLE
 		Params.OuterId = ::GetRemoteObjectId(Header.Outer.GetHandle());
 #endif
-		Params.RemoteId = ObjectData.Tables.RemoteIds[Header.RemoteId];
+		Params.RemoteId = ObjectData.GetRemoteObjectId(Header.RemoteId);
 		Params.SerialNumber = Header.SerialNumber;
 	}
 }
 
-UObject* ConstructRemoteObject(const FRemoteObjectId SerializedObjectId, const FRemoteObjectHeader& Header, ERemoteObjectSerializationFlags DeserializeFlags)
+UObject* ConstructRemoteObject(const FRemoteObjectData& ObjectData, const FRemoteObjectHeader& Header, ERemoteObjectSerializationFlags DeserializeFlags)
 {
 	using namespace UE::RemoteObject::Handle;
 	using namespace UE::RemoteObject::Private;
@@ -1005,7 +1016,8 @@ UObject* ConstructRemoteObject(const FRemoteObjectId SerializedObjectId, const F
 		UE_CLOG(!Outer, LogRemoteSerialization, Fatal, TEXT("Failed to resolve an Outer when constructing remote object"))
 	}
 #endif
-	FName Name = Header.Name;
+	FName Name = ObjectData.GetName(Header.Name);
+	FRemoteObjectId SerializedObjectId = ObjectData.GetRemoteObjectId(Header.RemoteId);
 
 	// The object may already exist in memory (it could be a default subobject of an object we've just created)
 	UObject* Object = StaticFindObjectFast(Class, Outer, Name);
@@ -1294,7 +1306,7 @@ int32 DeserializeObjectData(FRemoteObjectData& ObjectData,
 				if (!Stub->Name.IsNone())
 				{
 					// Remote object could've been renamed when it was migrated so always make sure that it has the same name locally
-					ObjectHeader.Name = Stub->Name;
+					ObjectHeader.Name = ObjectData.Tables.AddUniqueName(Stub->Name);
 				}
 			}
 		}
@@ -1320,10 +1332,10 @@ int32 DeserializeObjectData(FRemoteObjectData& ObjectData,
 			{
 				FRemoteObjectId SerializedObjectId = ObjectData.GetRemoteObjectId(ObjectHeader.RemoteId);
 				OutObjectRemoteIds.Add(SerializedObjectId);
-				OutObjects.Add(ConstructRemoteObject(SerializedObjectId, ObjectHeader, DeserializeFlags));
+				OutObjects.Add(ConstructRemoteObject(ObjectData, ObjectHeader, DeserializeFlags));
 			}
 
-			checkf(OutObjectRemoteIds.Num() == ObjectData.SerializedObjectIds.Num(), TEXT("The number of constructed object ids (%d) does not match the number of serialized object ids (%d)"), OutObjectRemoteIds.Num(), ObjectData.SerializedObjectIds.Num());
+			checkf(OutObjectRemoteIds.Num() == ObjectData.SerializedObjectHeaders.Num(), TEXT("The number of constructed object ids (%d) does not match the number of serialized object headers (%d)"), OutObjectRemoteIds.Num(), ObjectData.SerializedObjectHeaders.Num());
 #if DO_CHECK
 			for (int32 ObjectIndex = 0; ObjectIndex < OutObjectRemoteIds.Num(); ++ObjectIndex)
 			{
diff --git a/Engine/Source/Runtime/CoreUObject/Private/UObject/RemoteObjectSerialization.h b/Engine/Source/Runtime/CoreUObject/Private/UObject/RemoteObjectSerialization.h
--- a/Engine/Source/Runtime/CoreUObject/Private/UObject/RemoteObjectSerialization.h
+++ b/Engine/Source/Runtime/CoreUObject/Private/UObject/RemoteObjectSerialization.h
@@ -26,8 +26,8 @@ namespace UE::RemoteObject::Serialization
 	*/
 	struct FRemoteObjectHeader
 	{
-		FName Name;
-		FPackedRemoteObjectPathName::FNameIndexType RemoteId;
+		FRemoteObjectTables::FNameIndexType Name;
+		FRemoteObjectTables::FNameIndexType RemoteId;
 		TObjectPtr<UClass> Class;
 		TObjectPtr<UObject> Outer;
 		TObjectPtr<UObject> Archetype;
diff --git a/Engine/Source/Runtime/CoreUObject/Public/UObject/Class.h b/Engine/Source/Runtime/CoreUObject/Public/UObject/Class.h
--- a/Engine/Source/Runtime/CoreUObject/Public/UObject/Class.h
+++ b/Engine/Source/Runtime/CoreUObject/Public/UObject/Class.h
@@ -5608,6 +5608,12 @@ template<> struct TBaseStructure<FRemoteObjectBytes>
 	static COREUOBJECT_API UScriptStruct* Get();
 };
 
+struct FPackedRemoteObjectHeader;
+template<> struct TBaseStructure<FPackedRemoteObjectHeader>
+{
+	static COREUOBJECT_API UScriptStruct* Get();
+};
+
 struct FRemoteObjectData;
 template<> struct TBaseStructure<FRemoteObjectData>
 {
diff --git a/Engine/Source/Runtime/CoreUObject/Public/UObject/NoExportTypes.h b/Engine/Source/Runtime/CoreUObject/Public/UObject/NoExportTypes.h
--- a/Engine/Source/Runtime/CoreUObject/Public/UObject/NoExportTypes.h
+++ b/Engine/Source/Runtime/CoreUObject/Public/UObject/NoExportTypes.h
@@ -2823,6 +2823,24 @@ struct FRemoteObjectPathName : public FRemoteObjectTables
 	GENERATED_BODY()
 };
 
+/**
+ * @note The full C++ class is located here: Engine\Source\Runtime\CoreUObject\Public\UObject\RemoteObjectTransfer.h
+ */
+USTRUCT(noexport)
+struct FPackedRemoteObjectHeader
+{
+	GENERATED_BODY()
+
+	UPROPERTY()
+	uint16 NameIndex;
+
+	UPROPERTY()
+	uint16 IdIndex;
+
+	UPROPERTY()
+	uint16 ClassIndex;
+};
+
 /** 
  * @note The full C++ class is located here: Engine\Source\Runtime\CoreUObject\Public\UObject\RemoteObjectTransfer.h
  */
@@ -2838,7 +2856,7 @@ struct FRemoteObjectData
 	TArray<FPackedRemoteObjectPathName> PathNames;
 
 	UPROPERTY()
-	TArray<uint16> SerializedObjectIds;
+	TArray<FPackedRemoteObjectHeader> SerializedObjectHeaders;
 
 	UPROPERTY()
 	TArray<FRemoteObjectBytes> Bytes;
diff --git a/Engine/Source/Runtime/CoreUObject/Public/UObject/RemoteObjectPathName.h b/Engine/Source/Runtime/CoreUObject/Public/UObject/RemoteObjectPathName.h
--- a/Engine/Source/Runtime/CoreUObject/Public/UObject/RemoteObjectPathName.h
+++ b/Engine/Source/Runtime/CoreUObject/Public/UObject/RemoteObjectPathName.h
@@ -9,6 +9,8 @@
 */
 struct FRemoteObjectTables
 {
+	using FNameIndexType = uint16;
+
 	TArray<FName> Names;
 	TArray<FRemoteObjectId> RemoteIds;
 
@@ -17,6 +19,16 @@ struct FRemoteObjectTables
 		return Names.Num();
 	}
 
+	FNameIndexType AddUniqueName(const FName& Name)
+	{
+		int32 NameIndex = Names.IndexOfByKey(Name);
+		if (NameIndex == INDEX_NONE)
+		{
+			NameIndex = Names.Add(Name);
+		}
+		return IntCastChecked<FNameIndexType>(NameIndex);
+	}
+
 	COREUOBJECT_API void Serialize(FArchive& Ar);
 };
 
@@ -34,7 +46,7 @@ inline FArchive& operator<<(FArchive& Ar, FRemoteObjectTables& Tables)
 */
 struct FPackedRemoteObjectPathName
 {
-	using FNameIndexType = uint16;
+	using FNameIndexType = FRemoteObjectTables::FNameIndexType;
 	TArray<FNameIndexType> RemoteIds;
 	TArray<FNameIndexType> Names;
 
diff --git a/Engine/Source/Runtime/CoreUObject/Public/UObject/RemoteObjectTransfer.h b/Engine/Source/Runtime/CoreUObject/Public/UObject/RemoteObjectTransfer.h
--- a/Engine/Source/Runtime/CoreUObject/Public/UObject/RemoteObjectTransfer.h
+++ b/Engine/Source/Runtime/CoreUObject/Public/UObject/RemoteObjectTransfer.h
@@ -23,6 +23,16 @@ struct FRemoteObjectBytes
 	TArray<uint8> Bytes;
 };
 
+/**
+* Structure that holds packed serialized remote object header (noexport type)
+*/
+struct FPackedRemoteObjectHeader
+{
+	FRemoteObjectTables::FNameIndexType NameIndex = 0;
+	FRemoteObjectTables::FNameIndexType IdIndex = 0;
+	FRemoteObjectTables::FNameIndexType ClassIndex = 0;
+};
+
 /**
 * Structure that holds remote object memory (noexport type)
 */
@@ -32,8 +42,8 @@ struct FRemoteObjectData
 	FRemoteObjectTables Tables;
 	/** List of unique pathnames (stored as a list of indices of unique FNames) */
 	TArray<FPackedRemoteObjectPathName> PathNames;
-	/** List of remote id indices of the objects serialized in this ObjectData struct (in the same order they were serialized so it matches the order of object headers after deserialization) */
-	TArray<FPackedRemoteObjectPathName::FNameIndexType> SerializedObjectIds;
+	/** List of packed remote object headers of the objects serialized in this ObjectData struct (in the same order they were serialized so it matches the order of object headers after deserialization) */
+	TArray<FPackedRemoteObjectHeader> SerializedObjectHeaders;
 	/** Serialized object headers and data (properties) */
 	TArray<FRemoteObjectBytes> Bytes;
 
@@ -53,7 +63,7 @@ struct FRemoteObjectData
 	/**
 	* Returns a unique remote object id based on its index the id tables
 	*/
-	inline FRemoteObjectId GetRemoteObjectId(FPackedRemoteObjectPathName::FNameIndexType IdIndex) const
+	inline FRemoteObjectId GetRemoteObjectId(FRemoteObjectTables::FNameIndexType IdIndex) const
 	{
 		return Tables.RemoteIds[IdIndex];
 	}
@@ -65,7 +75,7 @@ struct FRemoteObjectData
 	*/
 	inline FRemoteObjectId GetSerializedObjectId(int32 ObjectIndex) const
 	{
-		return GetRemoteObjectId(SerializedObjectIds[ObjectIndex]);
+		return GetRemoteObjectId(SerializedObjectHeaders[ObjectIndex].IdIndex);
 	}
 
 	/**
@@ -74,7 +84,7 @@ struct FRemoteObjectData
 	*/
 	inline FRemoteObjectId GetRootSerializedObjectId() const
 	{
-		return SerializedObjectIds.Num() ? GetSerializedObjectId(0) : FRemoteObjectId();
+		return SerializedObjectHeaders.Num() ? GetSerializedObjectId(0) : FRemoteObjectId();
 	}
 
 	/**
@@ -83,11 +93,19 @@ struct FRemoteObjectData
 	*/
 	inline void ForEachSerializedObjectId(TFunctionRef<void(FRemoteObjectId)> Op) const
 	{
-		for (FPackedRemoteObjectPathName::FNameIndexType Index : SerializedObjectIds)
+		for (const FPackedRemoteObjectHeader& Header : SerializedObjectHeaders)
 		{
-			Op(Tables.RemoteIds[Index]);
+			Op(Tables.RemoteIds[Header.IdIndex]);
 		}
 	}
+
+	/**
+	* Returns a unique FName based on its index the id tables
+	*/
+	inline FName GetName(FPackedRemoteObjectPathName::FNameIndexType NameIndex) const
+	{
+		return Tables.Names[NameIndex];
+	}
 };
 
 namespace UE::RemoteObject::Transfer
