Adding FSerializedRemoteObjectIterator to allow range-for iteration over remote objects serialized in FRemoteObjectData

#Synced-CL 49104830

diff --git a/Engine/Source/Runtime/CoreUObject/Private/UObject/RemoteObjectSerialization.cpp b/Engine/Source/Runtime/CoreUObject/Private/UObject/RemoteObjectSerialization.cpp
--- a/Engine/Source/Runtime/CoreUObject/Private/UObject/RemoteObjectSerialization.cpp
+++ b/Engine/Source/Runtime/CoreUObject/Private/UObject/RemoteObjectSerialization.cpp
@@ -1343,6 +1343,12 @@ int32 DeserializeObjectData(FRemoteObjectData& ObjectData,
 				FRemoteObjectId SerializedObjectId = ObjectData.GetSerializedObjectId(ObjectIndex);
 				checkf(ConstructedObjectId == SerializedObjectId, TEXT("Constructed object id[%d] (%s) does not match the serialized object id (%s)"), ObjectIndex, *ConstructedObjectId.ToString(), *SerializedObjectId.ToString());
 			}
+			for (FSerializedRemoteObjectIterator It(ObjectData); It; ++It)
+			{
+				UObject* SerializedObject = OutObjects[It.GetIndex()];
+				UClass* HeaderClass = It.GetClass();
+				checkf(SerializedObject->GetClass() == HeaderClass, TEXT("Constructed object's [%d] (%s) class does not match the serialized object header class (%s)"), It.GetIndex(), *SerializedObject->GetFullName(), *HeaderClass->GetPathName());
+			}
 #endif
 		}
 
diff --git a/Engine/Source/Runtime/CoreUObject/Public/UObject/RemoteObjectTransfer.h b/Engine/Source/Runtime/CoreUObject/Public/UObject/RemoteObjectTransfer.h
--- a/Engine/Source/Runtime/CoreUObject/Public/UObject/RemoteObjectTransfer.h
+++ b/Engine/Source/Runtime/CoreUObject/Public/UObject/RemoteObjectTransfer.h
@@ -7,6 +7,7 @@
 #include "UObject/RemoteObjectPathName.h"
 #include "Async/Async.h"
 #include "UObject/RemoteExecutor.h"
+#include "Misc/UEOps.h"
 
 namespace UE::Net
 {
@@ -14,6 +15,7 @@ namespace UE::Net
 }
 
 struct FUObjectMigrationContext;
+struct FRemoteObjectData;
 
 /**
 * Structure that holds serialized remote object data chunk (< 64KB of data) (noexport type)
@@ -33,6 +35,46 @@ struct FPackedRemoteObjectHeader
 	FRemoteObjectTables::FNameIndexType ClassIndex = 0;
 };
 
+/**
+* Iterator that maps packed serialized remote object id indices to the actual unique remote object ids
+*/
+class FSerializedRemoteObjectIterator
+{
+	const FRemoteObjectData& Data;
+	int32 Index = 0;
+
+public:
+
+	FSerializedRemoteObjectIterator(const FRemoteObjectData& InData, int32 InIndex = 0)
+		: Data(InData)
+		, Index(InIndex)
+	{
+	}
+	inline FSerializedRemoteObjectIterator& operator++()
+	{
+		++Index;
+		return *this;
+	}
+	// UEOpEquals is required by UEOps.h to support range-for iteration
+	[[nodiscard]] FORCEINLINE bool UEOpEquals(const FSerializedRemoteObjectIterator& Other) const
+	{
+		return (Index == Other.Index);
+	}
+	inline int32 GetIndex() const
+	{
+		return Index;
+	}
+	// operator bool is for the old UE-style iterator iteration support
+	inline operator bool() const;
+	inline FRemoteObjectId GetId() const;
+	inline FName GetName() const;
+	inline UClass* GetClass() const;
+	inline FRemoteObjectId operator*() const
+	{
+		return GetId();
+	}
+};
+
 /**
 * Structure that holds remote object memory (noexport type)
 */
@@ -88,26 +130,45 @@ struct FRemoteObjectData
 	}
 
 	/**
-	* Calls the specified function for each serialized object remote id
-	* @param Op Function to call (can be a lambda)
+	* Returns a unique FName based on its index the id tables
 	*/
-	inline void ForEachSerializedObjectId(TFunctionRef<void(FRemoteObjectId)> Op) const
+	inline FName GetName(FRemoteObjectTables::FNameIndexType IdIndex) const
 	{
-		for (const FPackedRemoteObjectHeader& Header : SerializedObjectHeaders)
-		{
-			Op(Tables.RemoteIds[Header.IdIndex]);
-		}
+		return Tables.Names[IdIndex];
 	}
 
-	/**
-	* Returns a unique FName based on its index the id tables
-	*/
-	inline FName GetName(FPackedRemoteObjectPathName::FNameIndexType NameIndex) const
+	// Support for iterating over serialized object ids:
+	// 
+	// FRemoteObjectData Data;
+	// for (FRemoteObjectId SerializedObjectId : Data) { }
+	
+	FSerializedRemoteObjectIterator begin() const
+	{
+		return FSerializedRemoteObjectIterator(*this, 0);
+	}
+	FSerializedRemoteObjectIterator end() const
 	{
-		return Tables.Names[NameIndex];
+		return FSerializedRemoteObjectIterator(*this, SerializedObjectHeaders.Num());
 	}
 };
 
+inline FSerializedRemoteObjectIterator::operator bool() const
+{
+	return Index < Data.SerializedObjectHeaders.Num();
+}
+inline FRemoteObjectId FSerializedRemoteObjectIterator::GetId() const
+{
+	return Data.Tables.RemoteIds[Data.SerializedObjectHeaders[Index].IdIndex];
+}
+inline FName FSerializedRemoteObjectIterator::GetName() const
+{
+	return Data.Tables.Names[Data.SerializedObjectHeaders[Index].NameIndex];
+}
+inline UClass* FSerializedRemoteObjectIterator::GetClass() const
+{
+	return Cast<UClass>(Data.PathNames[Data.SerializedObjectHeaders[Index].ClassIndex].Resolve(Data.Tables));
+}
+
 namespace UE::RemoteObject::Transfer
 {
 	/** Information for performing a migration (send) an object to a remote server */
