Use `OnComplete` to fix the transactionally safe async loading code.

#Synced-CL 49108000

diff --git a/Engine/Source/Programs/AutoRTFMTests/Private/UPackageTests.cpp b/Engine/Source/Programs/AutoRTFMTests/Private/UPackageTests.cpp
--- a/Engine/Source/Programs/AutoRTFMTests/Private/UPackageTests.cpp
+++ b/Engine/Source/Programs/AutoRTFMTests/Private/UPackageTests.cpp
@@ -738,4 +738,21 @@ TEST_CASE("UPackage.AsyncLoading")
 			REQUIRE(Package);
 		});
 	}
+
+	SECTION("Flush In Child Transaction")
+	{
+		FString Name(FString::Printf(TEXT("/Game/%dAutoRTFMTestPackage%d"), FPlatformProcess::GetCurrentProcessId(), __LINE__));
+		FPackageScopedMaker _(Name);
+
+		AutoRTFM::Testing::Commit([&]
+		{
+			FLoadPackageAsyncDelegate Delegate;
+			int32 RequestId = LoadPackageAsync(*Name, Delegate);
+
+			AutoRTFM::Testing::Commit([&]
+			{
+				FlushAsyncLoading(RequestId);
+			});
+		});
+	}
 }
diff --git a/Engine/Source/Runtime/CoreUObject/Private/Serialization/TransactionallySafeAsyncLoading.cpp b/Engine/Source/Runtime/CoreUObject/Private/Serialization/TransactionallySafeAsyncLoading.cpp
--- a/Engine/Source/Runtime/CoreUObject/Private/Serialization/TransactionallySafeAsyncLoading.cpp
+++ b/Engine/Source/Runtime/CoreUObject/Private/Serialization/TransactionallySafeAsyncLoading.cpp
@@ -187,7 +187,18 @@ public:
 						const AutoRTFM::EContextStatus Status = AutoRTFM::Close([&]
 						{
 							AutoRTFM::OnCommit([this] { DrainTransactionallyDeferredLoadPackages(); });
-							AutoRTFM::PushOnAbortHandler(this, [this] { TransactionallyDeferredLoadPackages.Empty(); });
+							AutoRTFM::OnComplete([this]
+							{
+								if (!AutoRTFM::IsRetrying() && !TransactionallyDeferredLoadPackages.IsEmpty())
+								{
+									for (const FTransactionallyDeferredLoadPackagePayload& Payload : TransactionallyDeferredLoadPackages)
+									{
+										UE_LOG(LogStreaming, Display, TEXT("Ignoring aborted load for package '%s'"), *Payload.PackagePath.GetDebugNameWithExtension());
+									}
+									
+									TransactionallyDeferredLoadPackages.Empty();
+								}
+							});
 						});
 						ensure(Status == AutoRTFM::EContextStatus::OnTrack);
 					}
@@ -337,17 +348,8 @@ public:
 						RequestIdsClone.Add(RequestId);
 					}
 
-					const bool bShouldPopOnAbortHandler = !TransactionallyDeferredLoadPackages.IsEmpty();
-
 					const AutoRTFM::EContextStatus Status = AutoRTFM::Close([&]
 					{
-						if (bShouldPopOnAbortHandler)
-						{
-							// Pop the on-abort handler so that we don't purge the list of transactionally
-							// deferred packages to load, which we need to do in the retry below.
-							AutoRTFM::PopOnAbortHandler(this);
-						}
-
 						AutoRTFM::OnComplete([this, RequestIdsClone = ::MoveTemp(RequestIdsClone)]() mutable
 						{
 							check(AutoRTFM::IsRetrying());
