[ZenLoader]
- Early out now looks for RF_LoadCompleted instead of RF_NeedLoad when outside the mutex otherwise it could cause a thread to start using an object before it was serialized
- UObject->HasAnyFlags along with other variants looking at flags provides an acquire barrier to match the release barrier when setting the flag.

#Synced-CL 49011649

diff --git a/Engine/Source/Runtime/CoreUObject/Private/Serialization/AsyncLoading2.cpp b/Engine/Source/Runtime/CoreUObject/Private/Serialization/AsyncLoading2.cpp
--- a/Engine/Source/Runtime/CoreUObject/Private/Serialization/AsyncLoading2.cpp
+++ b/Engine/Source/Runtime/CoreUObject/Private/Serialization/AsyncLoading2.cpp
@@ -8430,7 +8430,10 @@ bool FAsyncPackage2::CreateRuntimeExports(FAsyncLoadingThreadState2& ThreadState
 						EventDrivenCreateExport(*LocalHeaderData, BundleEntry.LocalExportIndex);
 					}
 				}
-				else if (Export.Object && Export.Object->HasAnyFlags(RF_NeedLoad) && !CanDeferSerialization(Export))
+				// We cannot test RF_NeedLoad since it is removed before we start serializing.
+				// Since under multithreaded mode, we might race, we always enter the serialization
+				// race unless we are positive the serialization is completed.
+				else if (Export.Object && !Export.Object->HasAnyFlags(RF_LoadCompleted) && !CanDeferSerialization(Export))
 				{
 					check(BundleEntry.CommandType == FExportBundleEntry::ExportCommandType_Serialize);
 
@@ -8534,7 +8537,10 @@ bool FAsyncPackage2::PreloadRuntimeExports(FAsyncLoadingThreadState2& ThreadStat
 
 				if (BundleEntry.CommandType == FExportBundleEntry::ExportCommandType_Serialize)
 				{
-					if (Export.Object && Export.Object->HasAnyFlags(RF_NeedLoad))
+					// We cannot test RF_NeedLoad since it is removed before we start serializing.
+					// Since under multithreaded mode, we might race, we always enter the serialization
+					// race unless we are positive the serialization is completed.
+					if (Export.Object && !Export.Object->HasAnyFlags(RF_LoadCompleted))
 					{
 						EventDrivenSerializeExport(*LocalHeaderData, BundleEntry.LocalExportIndex, &Ar);
 					}
@@ -8748,7 +8754,10 @@ void FAsyncPackage2::ConditionalSerializeExport(const FAsyncPackageHeaderData& H
 		return;
 	}
 
-	if (Export.Object->HasAllFlags(RF_NeedLoad))
+	// We cannot test RF_NeedLoad since it is removed before we start serializing.
+	// Since under multithreaded mode, we might race, we always enter the serialization
+	// race unless we are positive the serialization is completed.
+	if (!Export.Object->HasAllFlags(RF_LoadCompleted))
 	{
 		FAsyncPackageScope2 Scope(this);
 		EventDrivenSerializeExport(Header, LocalExportIndex, nullptr);
@@ -8792,7 +8801,10 @@ void FAsyncPackage2::ConditionalSerializeImport(const FAsyncPackageHeaderData& H
 
 	if (UObject* FromImportStore = ImportStore.GetImportObject(Header, ObjectIndex))
 	{
-		if (!FromImportStore->HasAllFlags(RF_NeedLoad))
+		// We cannot test RF_NeedLoad since it is removed before we start serializing.
+		// Since under multithreaded mode, we might race, we always enter the serialization
+		// race unless we are positive the serialization is completed.
+		if (FromImportStore->HasAllFlags(RF_LoadCompleted))
 		{
 			return;
 		}
@@ -9365,7 +9377,7 @@ void FAsyncPackage2::EventDrivenSerializeExport(const FAsyncPackageHeaderData& H
 
 	TRACE_LOADTIME_SERIALIZE_EXPORT_SCOPE(Object, Export.CookedSerialSize);
 
-	if ((ExportObject.bFiltered || ExportObject.bExportLoadFailed) || !(Object && Object->HasAnyFlags(RF_NeedLoad)))
+	if (ExportObject.bFiltered || ExportObject.bExportLoadFailed || Object == nullptr)
 	{
 		if (ExportObject.bExportLoadFailed)
 		{
@@ -9380,11 +9392,21 @@ void FAsyncPackage2::EventDrivenSerializeExport(const FAsyncPackageHeaderData& H
 		else
 		{
 			UE_ASYNC_PACKAGE_LOG_VERBOSE(VeryVerbose, Desc, TEXT("SerializeExport"),
-				TEXT("Skipped already serialized export %s"), *Header.NameMap.GetName(Export.ObjectName).ToString());
+				TEXT("Skipped nullptr object for export %s"), *Header.NameMap.GetName(Export.ObjectName).ToString());
 		}
 		return;
 	}
 
+	// We cannot test RF_NeedLoad since it is removed before we start serializing.
+	// Since under multithreaded mode, we might race, we always enter the serialization
+	// race unless we are positive the serialization is completed.
+	if (Object->HasAnyFlags(RF_LoadCompleted))
+	{
+		UE_ASYNC_PACKAGE_LOG_VERBOSE(VeryVerbose, Desc, TEXT("SerializeExport"),
+			TEXT("Skipped already serialized export %s"), *Header.NameMap.GetName(Export.ObjectName).ToString());
+		return;
+	}
+
 	ProcessExportDependencies(Header, LocalExportIndex, FExportBundleEntry::ExportCommandType_Serialize);
 
 	FSerializationMutex ScopeLock(SerializationMutex);
diff --git a/Engine/Source/Runtime/CoreUObject/Public/UObject/UObjectBase.h b/Engine/Source/Runtime/CoreUObject/Public/UObject/UObjectBase.h
--- a/Engine/Source/Runtime/CoreUObject/Public/UObject/UObjectBase.h
+++ b/Engine/Source/Runtime/CoreUObject/Public/UObject/UObjectBase.h
@@ -371,7 +371,10 @@ private:
 
 		UE_AUTORTFM_OPEN
 		{
-			Result = FPlatformAtomics::AtomicRead_Relaxed((int32*)&ObjectFlags);
+			// Make this load acquire by default so that we can rely on looking at flags
+			// as a barrier between the thread setting a flag and a thread taking a decision
+			// depending on that flag.
+			Result = ((std::atomic<int32>*)&ObjectFlags)->load(std::memory_order_acquire);
 		};
 
 		return Result;
