Low level optimizations to object hash and archetype functions:
Change UPackage::StaticClass checks to Cast<UPackage>as it is faster.
Avoid redundant calculation of internal exclusion flags.
Avoid unnecessary lookups of global object table.

#Synced-CL 49126341

diff --git a/Engine/Source/Runtime/CoreUObject/Private/UObject/UObjectArchetype.cpp b/Engine/Source/Runtime/CoreUObject/Private/UObject/UObjectArchetype.cpp
--- a/Engine/Source/Runtime/CoreUObject/Private/UObject/UObjectArchetype.cpp
+++ b/Engine/Source/Runtime/CoreUObject/Private/UObject/UObjectArchetype.cpp
@@ -93,8 +93,7 @@ FORCENOINLINE UObject* FindArchetypeFromRequiredInfoImpl(const UClass* Class, co
 	}
 	else
 	{
-		if (Outer
-			&& Outer->GetClass() != UPackage::StaticClass()) // packages cannot have subobjects
+		if (Outer && Cast<UPackage>(Outer) == nullptr) // packages cannot have subobjects
 		{
 			// Get a lock on the UObject hash tables for the duration of the GetArchetype operation
 			FScopedUObjectHashTablesLock HashTablesLock;
@@ -119,9 +118,8 @@ FORCENOINLINE UObject* FindArchetypeFromRequiredInfoImpl(const UClass* Class, co
 				{
 					ArchetypeToSearch = FEditorCacheArchetypeManager::Get().GetCachedArchetype(Outer);
 				}
-#endif
-
 				if (!ArchetypeToSearch)
+#endif
 				{
 					ArchetypeToSearch = FindArchetypeFromRequiredInfoImpl<bUseImmutableArchetype, bUseArchetypeCache>(Outer->GetClass(), Outer->GetOuter(), Outer->GetFName(), Outer->GetFlags(), bUseUpToDateClass, Policy);
 				}
@@ -317,4 +315,4 @@ const UObject* FindImmutableArchetype(const UObject* InObj)
 	checkf(ImmutableArchetype->HasAnyFlags(RF_ClassDefaultObject) || ImmutableArchetype->GetOuter()->GetClass() == Archetype->GetOuter()->GetClass(), TEXT("Immutable archetype outer class mismatch: expected: %s, got: %s"), *Archetype->GetOuter()->GetClass()->GetPathName(), *ImmutableArchetype->GetOuter()->GetClass()->GetPathName());
 #endif // UE_WITH_IMMUTABLEARCHETYPE_DEBUGGING
 	return ImmutableArchetype;
-}
\ No newline at end of file
+}
diff --git a/Engine/Source/Runtime/CoreUObject/Private/UObject/UObjectBaseUtility.cpp b/Engine/Source/Runtime/CoreUObject/Private/UObject/UObjectBaseUtility.cpp
--- a/Engine/Source/Runtime/CoreUObject/Private/UObject/UObjectBaseUtility.cpp
+++ b/Engine/Source/Runtime/CoreUObject/Private/UObject/UObjectBaseUtility.cpp
@@ -176,14 +176,14 @@ void UObjectBaseUtility::ReattachExternalPackage()
 UObject* UObjectBaseUtility::GetOutermostObject() const
 {
 	UObject* Top = (UObject*)(this);
-	if (Top->IsA<UPackage>())
+	if (Cast<UPackage>(Top) != nullptr)
 	{
 		return nullptr;
 	}
 	for (;;)
 	{
 		UObject* CurrentOuter = Top->GetOuter();
-		if (CurrentOuter->IsA<UPackage>())
+		if (Cast<UPackage>(CurrentOuter) != nullptr)
 		{
 			return Top;
 		}
@@ -343,9 +343,9 @@ UObjectBaseUtility* UObjectBaseUtility::GetImplementingOuterObject(const UClass*
  */
 bool UObjectBaseUtility::IsIn( const UObject* SomeOuter ) const
 {
-	if (SomeOuter->IsA<UPackage>())
+	if (const UPackage* OuterPackage = Cast<UPackage>(SomeOuter))
 	{
-		return IsInPackage(static_cast<const UPackage*>(SomeOuter));
+		return IsInPackage(OuterPackage);
 	}
 	return IsInOuter(SomeOuter);
 }
diff --git a/Engine/Source/Runtime/CoreUObject/Private/UObject/UObjectHash.cpp b/Engine/Source/Runtime/CoreUObject/Private/UObject/UObjectHash.cpp
--- a/Engine/Source/Runtime/CoreUObject/Private/UObject/UObjectHash.cpp
+++ b/Engine/Source/Runtime/CoreUObject/Private/UObject/UObjectHash.cpp
@@ -1043,21 +1043,21 @@ UObject* StaticFindObjectInPackageInternal(FUObjectHashTables& ThreadHash, const
 	return Result;
 }
 
-UObject* StaticFindObjectFastInternalThreadSafe(FUObjectHashTables& ThreadHash, const UClass* ObjectClass, const UObject* ObjectPackage, FName ObjectName, EFindObjectFlags Flags, EObjectFlags ExcludeFlags, EInternalObjectFlags ExclusiveInternalFlags)
+UObject* StaticFindObjectFastInternalThreadSafe(FUObjectHashTables& ThreadHash, const UClass* ObjectClass, const UObject* ObjectOuter, FName ObjectName, EFindObjectFlags Flags, EObjectFlags ExcludeFlags, EInternalObjectFlags ExclusiveInternalFlags)
 {
 	ExclusiveInternalFlags |= GetMandatoryInternalExclusionFlags();
 
 	// If they specified an outer use that during the hashing
 	UObject* Result = nullptr;
-	if (ObjectPackage != nullptr)
+	if (ObjectOuter != nullptr)
 	{
-		int32 Hash = GetObjectOuterHash(ObjectName, (PTRINT)ObjectPackage);
+		int32 Hash = GetObjectOuterHash(ObjectName, (PTRINT)ObjectOuter);
 		FHashTableLock HashLock(ThreadHash);
 		bool bExactClass = !!(Flags & EFindObjectFlags::ExactClass);
 		for (TMultiMap<int32, uint32>::TConstKeyIterator HashIt(ThreadHash.HashOuter, Hash); HashIt; ++HashIt)
 		{
-			uint32 InternalIndex = HashIt.Value();
-			UObject* Object = static_cast<UObject*>(GUObjectArray.IndexToObject(InternalIndex)->GetObject());
+			FUObjectItem* ObjectEntry = GUObjectArray.IndexToObject(HashIt.Value());
+			UObject* Object = static_cast<UObject*>(ObjectEntry->GetObject());
 			if
 				/* check that the name matches the name we're searching for */
 				((Object->GetFName() == ObjectName)
@@ -1066,15 +1066,15 @@ UObject* StaticFindObjectFastInternalThreadSafe(FUObjectHashTables& ThreadHash,
 				&& !Object->HasAnyFlags(ExcludeFlags)
 
 				/* check that the object has the correct Outer */
-				&& Object->GetOuter() == ObjectPackage
+				&& Object->GetOuter() == ObjectOuter
 
 				/** If a class was specified, check that the object is of the correct class */
 				&& (ObjectClass == nullptr || (bExactClass ? Object->GetClass() == ObjectClass : Object->IsA(ObjectClass)))
 				
 				/** Include (or not) pending kill objects */
-				&& !Object->HasAnyInternalFlags(ExclusiveInternalFlags))
+				&& !ObjectEntry->HasAnyFlags(ExclusiveInternalFlags))
 			{
-				checkf(!Object->IsUnreachable(), TEXT("%s"), *Object->GetFullName());
+				checkf(!ObjectEntry->IsUnreachable(), TEXT("%s"), *Object->GetFullName());
 				if (Result)
 				{
 					UE_LOG(LogUObjectHash, Warning, TEXT("Ambiguous search, could be %s or %s"), *GetFullNameSafe(Result), *GetFullNameSafe(Object));
@@ -1090,10 +1090,13 @@ UObject* StaticFindObjectFastInternalThreadSafe(FUObjectHashTables& ThreadHash,
 		}
 
 #if WITH_EDITOR
-		// if the search fail and the OuterPackage is a UPackage, lookup potential external package
-		if (Result == nullptr && ObjectPackage->IsA(UPackage::StaticClass()))
+		// if the search fail and the ObjectOuter is a UPackage, lookup potential external package
+		if (Result == nullptr)
 		{
-			Result = StaticFindObjectInPackageInternal(ThreadHash, ObjectClass, static_cast<const UPackage*>(ObjectPackage), ObjectName, Flags, ExcludeFlags, ExclusiveInternalFlags);
+			if (const UPackage* ObjectPackage = Cast<UPackage>(ObjectOuter))
+			{
+				Result = StaticFindObjectInPackageInternal(ThreadHash, ObjectClass, ObjectPackage, ObjectName, Flags, ExcludeFlags, ExclusiveInternalFlags);
+			}
 		}
 #endif
 	}
@@ -1559,10 +1562,12 @@ void GetObjectsWithOuter(const class UObjectBase* Outer, TArray<UObject *>& Resu
 	checkf(Outer != nullptr, TEXT("Getting objects with a null outer is no longer supported. If you want to get all packages you might consider using GetObjectsOfClass instead."));
 
 #if WITH_EDITOR
-	// uncomment ensure to more easily find place where GetObjectsWithOuter should be replaced with GetObjectsWithPackage
-	if (!/*ensure*/(!Outer->GetClass()->IsChildOf(UPackage::StaticClass())))
+	if (const UPackage* OuterPackage = Cast<UPackage>(Outer))
 	{
-		GetObjectsWithPackage((UPackage*)Outer, Results, bIncludeNestedObjects, ExclusionFlags, ExclusionInternalFlags);
+		// uncomment ensure to more easily find place where GetObjectsWithOuter should be replaced with GetObjectsWithPackage
+		// ensure(!OuterPackage);
+
+		GetObjectsWithPackage(OuterPackage, Results, bIncludeNestedObjects, ExclusionFlags, ExclusionInternalFlags);
 		return;
 	}
 #endif
@@ -1623,11 +1628,12 @@ void ForEachObjectWithOuterBreakable(const class UObjectBase* Outer, TFunctionRe
 	checkf(Outer != nullptr, TEXT("Getting objects with a null outer is no longer supported. If you want to get all packages you might consider using GetObjectsOfClass instead."));
 	
 #if WITH_EDITOR
-	// uncomment ensure to more easily find place where ForEachObjectWithOuter should be replaced with ForEachObjectWithPackage
-	if (!/*ensure*/(!Outer->GetClass()->IsChildOf(UPackage::StaticClass())))
+	if (const UPackage* OuterPackage = Cast<UPackage>(Outer))
 	{
-		ForEachObjectWithPackage((UPackage*)Outer, [Operation](UObject* InObject) { Operation(InObject); return true; }, bIncludeNestedObjects, ExclusionFlags, ExclusionInternalFlags);
-		return;
+		// uncomment ensure to more easily find place where GetObjectsWithOuter should be replaced with GetObjectsWithPackage
+		// ensure(!OuterPackage)
+
+		ForEachObjectWithPackage(OuterPackage, [Operation](UObject* InObject) { Operation(InObject); return true; }, bIncludeNestedObjects, ExclusionFlags, ExclusionInternalFlags);
 	}
 #endif
 
@@ -1676,15 +1682,17 @@ UObjectBase* FindObjectWithOuter(const class UObjectBase* Outer, const class UCl
 {
 	UObject* Result = nullptr;
 	check( Outer );
-	// We don't want to return any objects that are currently being background loaded unless we're using the object iterator during async loading.
-	EInternalObjectFlags ExclusionInternalFlags = GetMandatoryInternalExclusionFlags();
-
-	if( NameToLookFor != NAME_None )
+	
+	if(!NameToLookFor.IsNone())
 	{
-		Result = StaticFindObjectFastInternal(ClassToLookFor, static_cast<const UObject*>(Outer), NameToLookFor, EFindObjectFlags::None, RF_NoFlags, ExclusionInternalFlags);
+		// Don't calculate mandatory exclusion flags here as they are handled inside
+		Result = StaticFindObjectFastInternal(ClassToLookFor, static_cast<const UObject*>(Outer), NameToLookFor);
 	}
 	else
 	{
+		// We don't want to return any objects that are currently being background loaded unless we're using the object iterator during async loading.
+		EInternalObjectFlags ExclusionInternalFlags = GetMandatoryInternalExclusionFlags();
+
 		FUObjectHashTables& ThreadHash = FUObjectHashTables::Get();
 		FHashTableLock HashLock(ThreadHash);
 		FHashBucket* Inners = ThreadHash.ObjectOuterMap.Find(Outer);
@@ -1866,8 +1874,6 @@ void ForEachObjectOfClass(const UClass* ClassToLookFor, TFunctionRef<void(UObjec
 {
 	TRACE_CPUPROFILER_EVENT_SCOPE(ForEachObjectOfClass);
 
-	ExclusionInternalFlags |= GetMandatoryInternalExclusionFlags();
-
 	// Most classes searched for have around 10 subclasses, some have hundreds
 	TArray<const UClass*, TInlineAllocator<16>> ClassesToSearch;
 	ClassesToSearch.Add(ClassToLookFor);
@@ -1877,7 +1883,8 @@ void ForEachObjectOfClass(const UClass* ClassToLookFor, TFunctionRef<void(UObjec
 
 	if (bIncludeDerivedClasses)
 	{
-		RecursivelyPopulateDerivedClasses(ThreadHash, ClassToLookFor, ClassesToSearch, ExclusionInternalFlags);
+		// RecursivelyPopulate expects the full exclusion flags, while the implementation function adds them
+		RecursivelyPopulateDerivedClasses(ThreadHash, ClassToLookFor, ClassesToSearch, ExclusionInternalFlags | GetMandatoryInternalExclusionFlags());
 	}
 
 	ForEachObjectOfClasses_Implementation(ThreadHash, ClassesToSearch, Operation, ExclusionFlags, ExclusionInternalFlags);
@@ -2000,7 +2007,7 @@ UE_AUTORTFM_NOAUTORTFM  // This is not safe to be called from a closed transacti
 void HashObject(UObjectBase* Object)
 {
 	FName Name = Object->GetFName();
-	if (Name != NAME_None)
+	if (!Name.IsNone())
 	{
 		LLM_SCOPE_BYTAG(UObjectHash);
 #if LLM_ALLOW_ASSETS_TAGS
@@ -2057,7 +2064,7 @@ UE_AUTORTFM_NOAUTORTFM  // This is not safe to be called from a closed transacti
 void UnhashObject(UObjectBase* Object)
 {
 	FName Name = Object->GetFName();
-	if (Name != NAME_None)
+	if (!Name.IsNone())
 	{
 		LLM_SCOPE_BYTAG(UObjectHash);
 #if LLM_ALLOW_ASSETS_TAGS
