[UE] Apply FObjectInitializer instancing exclusions only to the object being constructed/initialized.

Notes:
- The instancing graph is constructed at the top level and propagates through subobject instancing.
- Exclusions on FObjectInitializer (e.g. set by the class) should only apply to the top-level object being constructed.
- This fixes a regression introduced on the Verse side at CL 46787943 (and further exposed by CL 47523491) that causes FObjectInitializer exclusions at each level to be unintentionally aggregated into the instancing graph when constructing a Verse CDO hierarchy.

#Synced-CL 49069706

diff --git a/Engine/Source/Runtime/CoreUObject/Private/UObject/Class.cpp b/Engine/Source/Runtime/CoreUObject/Private/UObject/Class.cpp
--- a/Engine/Source/Runtime/CoreUObject/Private/UObject/Class.cpp
+++ b/Engine/Source/Runtime/CoreUObject/Private/UObject/Class.cpp
@@ -2742,6 +2742,8 @@ FTopLevelAssetPath UStruct::GetStructPathName() const
 
 void UStruct::InstanceSubobjectTemplates(TNotNull<void*> Data, const void* DefaultData, const UStruct* DefaultStruct, TNotNull<UObject*> Owner, FObjectInstancingGraph* InstanceGraph)
 {
+	const FObjectInitializer* ObjInitializer = FUObjectThreadContext::Get().TopInitializer();
+
 	// Dynamic instancing allows the instance graph to infer when to instance subobjects instead of requiring the
 	// explicit 'CPF_InstancedReference' flag. This also means the code below intentionally visits all properties
 	// in the reflink chain as a result, unless they are otherwise explicitly excluded through the instance graph.
@@ -2755,8 +2757,8 @@ void UStruct::InstanceSubobjectTemplates(TNotNull<void*> Data, const void* Defau
 		;
 
 		const bool bPropertyInExclusionList =
-			InstanceGraph &&
-			InstanceGraph->IsPropertyInSubobjectExclusionList(Property)
+			(InstanceGraph && InstanceGraph->IsPropertyInSubobjectExclusionList(Property)) ||
+			(ObjInitializer && ObjInitializer->IsPropertyInSubobjectExclusionList(Property))
 		;
 
 		if (bPropertySupportsInstancing && !bPropertyInExclusionList)
diff --git a/Engine/Source/Runtime/CoreUObject/Private/UObject/UObjectGlobals.cpp b/Engine/Source/Runtime/CoreUObject/Private/UObject/UObjectGlobals.cpp
--- a/Engine/Source/Runtime/CoreUObject/Private/UObject/UObjectGlobals.cpp
+++ b/Engine/Source/Runtime/CoreUObject/Private/UObject/UObjectGlobals.cpp
@@ -4407,11 +4407,6 @@ void FObjectInitializer::InstanceSubobjects(UClass* Class, bool bNeedInstancing,
 
 	FObjectInstancingGraph TempInstancingGraph;
 	FObjectInstancingGraph* UseInstancingGraph = InstanceGraph ? InstanceGraph : &TempInstancingGraph;
-	// Transfer any excluded properties from FObjectInitializer into FObjectInstancingGraph
-	for (const FProperty* Property : SubobjectInstantiationExclusionList)
-	{
-		UseInstancingGraph->AddPropertyToSubobjectExclusionList(Property);
-	}
 	{
 		UseInstancingGraph->AddNewObject(Obj, ObjectArchetype);
 	}
diff --git a/Engine/Source/Runtime/CoreUObject/Public/UObject/UObjectGlobals.h b/Engine/Source/Runtime/CoreUObject/Public/UObject/UObjectGlobals.h
--- a/Engine/Source/Runtime/CoreUObject/Public/UObject/UObjectGlobals.h
+++ b/Engine/Source/Runtime/CoreUObject/Public/UObject/UObjectGlobals.h
@@ -1532,6 +1532,22 @@ public:
 		return *this;
 	}
 
+	/**
+	 * Adds a member variable property that should not be instanced for this object. Not intended for use outside of engine code.
+	 */
+	UE_INTERNAL void AddPropertyToSubobjectExclusionList(const FProperty* Property)
+	{
+		SubobjectInstantiationExclusionList.Add(Property);
+	}
+
+	/**
+	 * Checks if a member variable property should not be instanced for this object. Not intended for use outside of engine code.
+	 */
+	UE_INTERNAL bool IsPropertyInSubobjectExclusionList(const FProperty* Property) const
+	{
+		return SubobjectInstantiationExclusionList.Contains(Property);
+	}
+
 	/**
 	 * Asserts with the specified message if code is executed inside UObject constructor
 	 **/
@@ -1701,12 +1717,6 @@ private:
 	/** Asserts if SetDefaultSubobjectClass or DoNotCreateOptionalDefaultSuobject are called inside of the constructor body */
 	COREUOBJECT_API void AssertIfSubobjectSetupIsNotAllowed(TArrayView<const FName> SubobjectNames) const;
 
-	friend class UVerseClass;
-	void AddPropertyToSubobjectExclusionList(const FProperty* Property)
-	{
-		SubobjectInstantiationExclusionList.Add(Property);
-	}
-
 	/**  object to initialize, from static allocate object, after construction **/
 	UObject* Obj;
 	/**  object to copy properties from **/
