Added comment to explain the intent of the order of comparisons in FMulticastDelegateProperty::Identical.

diff --git a/Engine/Source/Runtime/CoreUObject/Private/UObject/PropertyMulticastDelegate.cpp b/Engine/Source/Runtime/CoreUObject/Private/UObject/PropertyMulticastDelegate.cpp
--- a/Engine/Source/Runtime/CoreUObject/Private/UObject/PropertyMulticastDelegate.cpp
+++ b/Engine/Source/Runtime/CoreUObject/Private/UObject/PropertyMulticastDelegate.cpp
@@ -159,11 +159,14 @@ bool FMulticastDelegateProperty::Identical( const void* A, const void* B, uint32
 		}
 		else
 #endif
+
 		if (BindingA.GetFunctionName() != BindingB.GetFunctionName())
 		{
 			return false;
 		}
 
+		// Check UObject pointer comparison last.  This is to avoid unnecessary resolution of the
+		// delegates' weak bindings back to pointers, particularly when objects are remote.
 		if (BindingA.GetUObject() != BindingB.GetUObject())
 		{
 			return false;
