[ZenLoader/UObject]
- Reduce the barrier cost by only using acquire when needed

#Synced-CL 49181768

diff --git a/Engine/Source/Runtime/CoreUObject/Private/Serialization/AsyncLoading2.cpp b/Engine/Source/Runtime/CoreUObject/Private/Serialization/AsyncLoading2.cpp
--- a/Engine/Source/Runtime/CoreUObject/Private/Serialization/AsyncLoading2.cpp
+++ b/Engine/Source/Runtime/CoreUObject/Private/Serialization/AsyncLoading2.cpp
@@ -8433,7 +8433,7 @@ bool FAsyncPackage2::CreateRuntimeExports(FAsyncLoadingThreadState2& ThreadState
 				// We cannot test RF_NeedLoad since it is removed before we start serializing.
 				// Since under multithreaded mode, we might race, we always enter the serialization
 				// race unless we are positive the serialization is completed.
-				else if (Export.Object && !Export.Object->HasAnyFlags(RF_LoadCompleted) && !CanDeferSerialization(Export))
+				else if (Export.Object && !Export.Object->HasAnyFlags(RF_LoadCompleted, std::memory_order_acquire) && !CanDeferSerialization(Export))
 				{
 					check(BundleEntry.CommandType == FExportBundleEntry::ExportCommandType_Serialize);
 
@@ -8540,7 +8540,7 @@ bool FAsyncPackage2::PreloadRuntimeExports(FAsyncLoadingThreadState2& ThreadStat
 					// We cannot test RF_NeedLoad since it is removed before we start serializing.
 					// Since under multithreaded mode, we might race, we always enter the serialization
 					// race unless we are positive the serialization is completed.
-					if (Export.Object && !Export.Object->HasAnyFlags(RF_LoadCompleted))
+					if (Export.Object && !Export.Object->HasAnyFlags(RF_LoadCompleted, std::memory_order_acquire))
 					{
 						EventDrivenSerializeExport(*LocalHeaderData, BundleEntry.LocalExportIndex, &Ar);
 					}
@@ -8757,7 +8757,7 @@ void FAsyncPackage2::ConditionalSerializeExport(const FAsyncPackageHeaderData& H
 	// We cannot test RF_NeedLoad since it is removed before we start serializing.
 	// Since under multithreaded mode, we might race, we always enter the serialization
 	// race unless we are positive the serialization is completed.
-	if (!Export.Object->HasAllFlags(RF_LoadCompleted))
+	if (!Export.Object->HasAllFlags(RF_LoadCompleted, std::memory_order_acquire))
 	{
 		FAsyncPackageScope2 Scope(this);
 		EventDrivenSerializeExport(Header, LocalExportIndex, nullptr);
@@ -8804,7 +8804,7 @@ void FAsyncPackage2::ConditionalSerializeImport(const FAsyncPackageHeaderData& H
 		// We cannot test RF_NeedLoad since it is removed before we start serializing.
 		// Since under multithreaded mode, we might race, we always enter the serialization
 		// race unless we are positive the serialization is completed.
-		if (FromImportStore->HasAllFlags(RF_LoadCompleted))
+		if (FromImportStore->HasAllFlags(RF_LoadCompleted, std::memory_order_acquire))
 		{
 			return;
 		}
@@ -9400,7 +9400,7 @@ void FAsyncPackage2::EventDrivenSerializeExport(const FAsyncPackageHeaderData& H
 	// We cannot test RF_NeedLoad since it is removed before we start serializing.
 	// Since under multithreaded mode, we might race, we always enter the serialization
 	// race unless we are positive the serialization is completed.
-	if (Object->HasAnyFlags(RF_LoadCompleted))
+	if (Object->HasAnyFlags(RF_LoadCompleted, std::memory_order_acquire))
 	{
 		UE_ASYNC_PACKAGE_LOG_VERBOSE(VeryVerbose, Desc, TEXT("SerializeExport"),
 			TEXT("Skipped already serialized export %s"), *Header.NameMap.GetName(Export.ObjectName).ToString());
diff --git a/Engine/Source/Runtime/CoreUObject/Public/UObject/UObjectBase.h b/Engine/Source/Runtime/CoreUObject/Public/UObject/UObjectBase.h
--- a/Engine/Source/Runtime/CoreUObject/Public/UObject/UObjectBase.h
+++ b/Engine/Source/Runtime/CoreUObject/Public/UObject/UObjectBase.h
@@ -295,6 +295,19 @@ public:
 		return Flags;
 	}
 
+	/**
+	 * Retrieve the object flags directly
+	 *
+	 * @param InMemoryOrder	Memory ordering to apply when reading the flags from memory
+	 * @return Flags for this object
+	 **/
+	FORCEINLINE EObjectFlags GetFlags(std::memory_order InMemoryOrder) const
+	{
+		EObjectFlags Flags = (EObjectFlags)GetFlagsInternal(InMemoryOrder);
+		checkfSlow((Flags & ~RF_AllFlags) == 0, TEXT("%s flagged as RF_AllFlags"), *GetFName().ToString());
+		return Flags;
+	}
+
 	/**
 	 *	Atomically adds the specified flags.
 	 */
@@ -371,10 +384,24 @@ private:
 
 		UE_AUTORTFM_OPEN
 		{
-			// Make this load acquire by default so that we can rely on looking at flags
+			Result = ((std::atomic<int32>*)&ObjectFlags)->load(std::memory_order_relaxed);
+		};
+
+		return Result;
+	}
+
+	FORCEINLINE int32 GetFlagsInternal(std::memory_order InMemoryOrder) const
+	{
+		static_assert(sizeof(int32) == sizeof(ObjectFlags), "Flags must be 32-bit for atomics.");
+
+		int32 Result = 0;
+
+		UE_AUTORTFM_OPEN
+		{
+			// Make this load use memory ordering so that we can rely on looking at flags
 			// as a barrier between the thread setting a flag and a thread taking a decision
 			// depending on that flag.
-			Result = ((std::atomic<int32>*)&ObjectFlags)->load(std::memory_order_acquire);
+			Result = ((std::atomic<int32>*)&ObjectFlags)->load(InMemoryOrder);
 		};
 
 		return Result;
diff --git a/Engine/Source/Runtime/CoreUObject/Public/UObject/UObjectBaseUtility.h b/Engine/Source/Runtime/CoreUObject/Public/UObject/UObjectBaseUtility.h
--- a/Engine/Source/Runtime/CoreUObject/Public/UObject/UObjectBaseUtility.h
+++ b/Engine/Source/Runtime/CoreUObject/Public/UObject/UObjectBaseUtility.h
@@ -96,6 +96,19 @@ public:
 		return (GetFlags() & FlagsToCheck) != 0;
 	}
 
+	/**
+	 * Used to safely check whether any of the passed in flags are set. 
+	 *
+	 * @param FlagsToCheck	Object flags to check for.
+	 * @param MemoryOrder	Memory ordering to apply when reading the flags from memory
+	 * @return true if any of the passed in flags are set, false otherwise  (including no flags passed in).
+	 */
+	inline bool HasAnyFlags( EObjectFlags FlagsToCheck, std::memory_order InMemoryOrder) const
+	{
+		checkSlow(!(FlagsToCheck & (RF_MarkAsNative | RF_MarkAsRootSet)) || FlagsToCheck == RF_AllFlags); // These flags can't be used outside of constructors / internal code
+		return (GetFlags(InMemoryOrder) & FlagsToCheck) != 0;
+	}
+
 	/**
 	 * Used to safely check whether all of the passed in flags are set. 
 	 *
@@ -108,6 +121,19 @@ public:
 		return ((GetFlags() & FlagsToCheck) == FlagsToCheck);
 	}
 
+	/**
+	 * Used to safely check whether all of the passed in flags are set. 
+	 *
+	 * @param FlagsToCheck	Object flags to check for
+	 * @param MemoryOrder	Memory ordering to apply when reading the flags from memory
+	 * @return true if all of the passed in flags are set (including no flags passed in), false otherwise
+	 */
+	inline bool HasAllFlags( EObjectFlags FlagsToCheck, std::memory_order InMemoryOrder) const
+	{
+		checkSlow(!(FlagsToCheck & (RF_MarkAsNative | RF_MarkAsRootSet)) || FlagsToCheck == RF_AllFlags); // These flags can't be used outside of constructors / internal code
+		return ((GetFlags(InMemoryOrder) & FlagsToCheck) == FlagsToCheck);
+	}
+
 	/**
 	 * Returns object flags that are both in the mask and set on the object.
 	 *
