[UE] Add CLASS_Transient use cases to low-level subobject instancing tests/ground truth.

#rnx
#Synced-CL 49096161

diff --git a/Engine/Source/Runtime/CoreUObject/Tests/UObject/InstancedPropertiesTest.cpp b/Engine/Source/Runtime/CoreUObject/Tests/UObject/InstancedPropertiesTest.cpp
--- a/Engine/Source/Runtime/CoreUObject/Tests/UObject/InstancedPropertiesTest.cpp
+++ b/Engine/Source/Runtime/CoreUObject/Tests/UObject/InstancedPropertiesTest.cpp
@@ -79,6 +79,31 @@ TEST_CASE_NAMED(FInstancedPropertiesTest, "CoreUObject::InstancedProperties", "[
 #endif
 	}
 
+	SECTION("Transient type instanced object property")
+	{
+		const FName TransientTypeInstancedObjectPropertyName = GET_MEMBER_NAME_CHECKED(USubobjectInstancingTestOuterObject, TransientInnerObjectFromType);
+		const FProperty* TransientTypeInstancedObjectProperty = USubobjectInstancingTestOuterObject::StaticClass()->FindPropertyByName(TransientTypeInstancedObjectPropertyName);
+
+		REQUIRE(TransientTypeInstancedObjectProperty != nullptr);
+		CHECK(TransientTypeInstancedObjectProperty->IsA<FObjectProperty>());
+		
+		const FObjectProperty* ObjectProperty = CastFieldChecked<FObjectProperty>(TransientTypeInstancedObjectProperty);
+		{
+			REQUIRE(ObjectProperty->PropertyClass != nullptr);
+			CHECK(ObjectProperty->PropertyClass->HasAllClassFlags(CLASS_Transient));
+		}
+
+		CHECK(TransientTypeInstancedObjectProperty->ContainsInstancedObjectProperty());
+		// Note: CLASS_Transient imposes RF_Transient on all non-archetype allocations (see StaticAllocateObject), which means the referenced object will not be exported on save.
+		// It does NOT impose CPF_Transient - the property value is instead serialized as NULL on save if the referenced object is not also exported.
+		CHECK_FALSE(TransientTypeInstancedObjectProperty->HasAllPropertyFlags(CPF_Transient));
+		CHECK(TransientTypeInstancedObjectProperty->HasAllPropertyFlags(CPF_InstancedReference));
+		CHECK(TransientTypeInstancedObjectProperty->HasAllPropertyFlags(CPF_PersistentInstance));
+#if WITH_METADATA
+		CHECK(TransientTypeInstancedObjectProperty->GetBoolMetaData(TEXT("EditInline")));
+#endif
+	}
+
 	SECTION("Instanced object array property")
 	{
 		const FName InstancedObjectArrayPropertyName = GET_MEMBER_NAME_CHECKED(USubobjectInstancingTestOuterObject, InnerObjectArray);
diff --git a/Engine/Source/Runtime/CoreUObject/Tests/UObject/SubobjectInstancingTest.cpp b/Engine/Source/Runtime/CoreUObject/Tests/UObject/SubobjectInstancingTest.cpp
--- a/Engine/Source/Runtime/CoreUObject/Tests/UObject/SubobjectInstancingTest.cpp
+++ b/Engine/Source/Runtime/CoreUObject/Tests/UObject/SubobjectInstancingTest.cpp
@@ -57,6 +57,9 @@ USubobjectInstancingTestOuterObject::USubobjectInstancingTestOuterObject(const F
 	// Property flags will exclude RF_DefaultSubObject, and runtime behavior is similar to LocalOnlyInnerObject above (i.e. no inheritance)
 	InnerObjectUsingNew = NewObject<USubobjectInstancingTestObject>(this, TEXT("InnerObjectUsingNew"));
 
+	// Property with transient semantics inferred from the class type (Transient); all non-archetype instances will be RF_Transient
+	TransientInnerObjectFromType = CreateDefaultSubobject<USubobjectInstancingTransientTestObject>(TEXT("TransientInnerObjectFromType"));
+
 	// Transient reference to an instanced subobject that is constructed via NewObject() instead of CreateDefaultSubobject() (not serialized)
 	TransientInnerObjectUsingNew = NewObject<USubobjectInstancingTestObject>(this, TEXT("TransientInnerObjectUsingNew"));
 
@@ -676,6 +679,57 @@ namespace SubobjectInstancingTest::Private
 			USubobjectInstancingTestOuterObject* NewOuterObject = NewObject<USubobjectInstancingTestOuterObject>(GetTransientPackage(), InstancingTestClass, NAME_None, RF_NoFlags, Template);
 			CHECK(OuterObject->TransientInnerObjectUsingNew != NewOuterObject->TransientInnerObjectUsingNew);
 		}
+
+		SECTION("Native property initialized to a transient type instantiation from CDO")
+		{
+			USubobjectInstancingTestOuterObject* OuterObject = NewObject<USubobjectInstancingTestOuterObject>(GetTransientPackage(), InstancingTestClass);
+
+			REQUIRE(OuterObject->TransientInnerObjectFromType != nullptr);
+			CHECK(OuterObject->TransientInnerObjectFromType->IsInOuter(OuterObject));
+
+			CHECK(OuterObject->TransientInnerObjectFromType->TestValue == 100);
+
+			// Transient class types impose RF_Transient on all non-archetype instantiations (see StaticAllocateObject).
+			CHECK(OuterObject->TransientInnerObjectFromType->HasAllFlags(RF_Transient));
+
+			CHECK(OuterObject->TransientInnerObjectFromType->HasAllFlags(RF_DefaultSubObject));
+			CHECK(OuterObject->TransientInnerObjectFromType->HasAllFlags(OuterObject->GetMaskedFlags(RF_PropagateToSubObjects)));
+
+			CHECK_FALSE(OuterObject->TransientInnerObjectFromType->IsTemplate());
+			CHECK(OuterObject->TransientInnerObjectFromType->IsDefaultSubobject());
+			CHECK_FALSE(OuterObject->TransientInnerObjectFromType->IsTemplateForSubobjects());
+
+			USubobjectInstancingTestOuterObject* NewOuterObject = NewObject<USubobjectInstancingTestOuterObject>(GetTransientPackage(), InstancingTestClass);
+			CHECK(OuterObject->TransientInnerObjectFromType != NewOuterObject->TransientInnerObjectFromType);
+		}
+
+		SECTION("Native property initialized to a transient type instantiation from template")
+		{
+			USubobjectInstancingTestOuterObject* Template = NewObject<USubobjectInstancingTestOuterObject>(GetTransientPackage(), InstancingTestClass, NAME_None, RF_ArchetypeObject);
+			Template->TransientInnerObjectFromType->TestValue = 200;
+
+			USubobjectInstancingTestOuterObject* OuterObject = NewObject<USubobjectInstancingTestOuterObject>(GetTransientPackage(), InstancingTestClass, NAME_None, RF_NoFlags, Template);
+
+			REQUIRE(OuterObject->TransientInnerObjectFromType != nullptr);
+			CHECK(OuterObject->TransientInnerObjectFromType->IsInOuter(OuterObject));
+
+			// The property itself is not transient, so the template is used as the basis for property initialization.
+			CHECK(OuterObject->TransientInnerObjectFromType->TestValue == 200);
+
+			// Transient class types impose RF_Transient on all non-archetype instantiations (see StaticAllocateObject).
+			CHECK(OuterObject->TransientInnerObjectFromType->HasAllFlags(RF_Transient));
+			CHECK_FALSE(Template->TransientInnerObjectFromType->HasAnyFlags(RF_Transient));
+
+			CHECK(OuterObject->TransientInnerObjectFromType->HasAllFlags(RF_DefaultSubObject));
+			CHECK(OuterObject->TransientInnerObjectFromType->HasAllFlags(OuterObject->GetMaskedFlags(RF_PropagateToSubObjects)));
+
+			CHECK_FALSE(OuterObject->TransientInnerObjectFromType->IsTemplate());
+			CHECK(OuterObject->TransientInnerObjectFromType->IsDefaultSubobject());
+			CHECK_FALSE(OuterObject->TransientInnerObjectFromType->IsTemplateForSubobjects());
+
+			USubobjectInstancingTestOuterObject* NewOuterObject = NewObject<USubobjectInstancingTestOuterObject>(GetTransientPackage(), InstancingTestClass, NAME_None, RF_NoFlags, Template);
+			CHECK(OuterObject->TransientInnerObjectFromType != NewOuterObject->TransientInnerObjectFromType);
+		}
 	}
 
 	void RunSubobjectInstancingTests_Array(UClass* InstancingTestClass)
diff --git a/Engine/Source/Runtime/CoreUObject/Tests/UObject/SubobjectInstancingTest.h b/Engine/Source/Runtime/CoreUObject/Tests/UObject/SubobjectInstancingTest.h
--- a/Engine/Source/Runtime/CoreUObject/Tests/UObject/SubobjectInstancingTest.h
+++ b/Engine/Source/Runtime/CoreUObject/Tests/UObject/SubobjectInstancingTest.h
@@ -70,6 +70,12 @@ public:
 	TObjectPtr<USubobjectInstancingTestDirectlyNestedObject> InnerObject;
 };
 
+UCLASS(Transient)
+class USubobjectInstancingTransientTestObject : public USubobjectInstancingTestObject
+{
+	GENERATED_BODY()
+};
+
 UCLASS(DefaultToInstanced)
 class USubobjectInstancingDefaultToInstancedTestObject : public USubobjectInstancingTestObject
 {
@@ -137,6 +143,9 @@ public:
 	UPROPERTY()
 	TObjectPtr<USubobjectInstancingDefaultToInstancedTestObject> InnerObjectFromType;
 
+	UPROPERTY(Instanced)
+	TObjectPtr<USubobjectInstancingTransientTestObject> TransientInnerObjectFromType;
+
 	UPROPERTY(Instanced)
 	TObjectPtr<USubobjectInstancingTestObject> InnerObjectUsingNew;
 
