Low-level function to convert an EAssetAccessSpecifier to FText

#rnx

#Synced-CL 49071325

diff --git a/Engine/Source/Editor/ContentBrowser/Private/AssetViewTypes.cpp b/Engine/Source/Editor/ContentBrowser/Private/AssetViewTypes.cpp
--- a/Engine/Source/Editor/ContentBrowser/Private/AssetViewTypes.cpp
+++ b/Engine/Source/Editor/ContentBrowser/Private/AssetViewTypes.cpp
@@ -307,23 +307,7 @@ FText FAssetViewItem::GetItemAssetAccessSpecifierText() const
 	EAssetAccessSpecifier AssetAccessSpecifier;
 	if (GetItemAssetAccessSpecifier(AssetAccessSpecifier))
 	{
-		switch (AssetAccessSpecifier)
-		{
-			case EAssetAccessSpecifier::Public:
-			{
-				return LOCTEXT("PublicAssetState", "Public");
-			}
-
-			case EAssetAccessSpecifier::EpicInternal:
-			{
-				return LOCTEXT("PrivateAssetState", "Epic Internal");
-			}
-
-			case EAssetAccessSpecifier::Private:
-			{
-				return LOCTEXT("EpicInternalAssetState", "Private");
-			}
-		}
+		return AssetAccessSpecifierAsText(AssetAccessSpecifier);
 	}
 
 	// This would only happen if there's something wrong with the underlying FContentBrowserItem.
diff --git a/Engine/Source/Runtime/CoreUObject/Private/UObject/Package.cpp b/Engine/Source/Runtime/CoreUObject/Private/UObject/Package.cpp
--- a/Engine/Source/Runtime/CoreUObject/Private/UObject/Package.cpp
+++ b/Engine/Source/Runtime/CoreUObject/Private/UObject/Package.cpp
@@ -21,6 +21,32 @@
 
 #include UE_INLINE_GENERATED_CPP_BY_NAME(Package)
 
+FText AssetAccessSpecifierAsText(EAssetAccessSpecifier InAccess)
+{
+#define LOCTEXT_NAMESPACE "AssetAccessSpecifier"
+	switch (InAccess)
+	{
+	case EAssetAccessSpecifier::Public:
+	{
+		return LOCTEXT("Public", "Public");
+	}
+
+	case EAssetAccessSpecifier::EpicInternal:
+	{
+		return LOCTEXT("EpicInternal", "Epic Internal");
+	}
+
+	case EAssetAccessSpecifier::Private:
+	{
+		return LOCTEXT("Private", "Private");
+	}
+	default:
+		ensureMsgf(false, TEXT("Unknwon EAssetAccessSpecifier value %d"), (int64)InAccess);
+		return FText::GetEmpty();
+	}
+#undef LOCTEXT_NAMESPACE
+}
+
 /*-----------------------------------------------------------------------------
 	UPackage.
 -----------------------------------------------------------------------------*/
diff --git a/Engine/Source/Runtime/CoreUObject/Public/UObject/Package.h b/Engine/Source/Runtime/CoreUObject/Public/UObject/Package.h
--- a/Engine/Source/Runtime/CoreUObject/Public/UObject/Package.h
+++ b/Engine/Source/Runtime/CoreUObject/Public/UObject/Package.h
@@ -207,6 +207,8 @@ enum class EAssetAccessSpecifier : uint8
 	EpicInternal,
 };
 
+COREUOBJECT_API FText AssetAccessSpecifierAsText(EAssetAccessSpecifier);
+
 /**
 * A package.
 */
