[StructUtils] Sanitize the name and make sure the ID are valid before hashing the properties.

#Synced-CL 49284918

diff --git a/Engine/Source/Runtime/CoreUObject/Private/StructUtils/PropertyBag.cpp b/Engine/Source/Runtime/CoreUObject/Private/StructUtils/PropertyBag.cpp
--- a/Engine/Source/Runtime/CoreUObject/Private/StructUtils/PropertyBag.cpp
+++ b/Engine/Source/Runtime/CoreUObject/Private/StructUtils/PropertyBag.cpp
@@ -4284,8 +4284,9 @@ TArray<FPropertyBagPropertyDesc> UPropertyBag::SanitizePropertyDescs(const TArra
 	return Result;
 }
 
-const UPropertyBag* UPropertyBag::GetOrCreateFromDescs(const TConstArrayView<FPropertyBagPropertyDesc> PropertyDescs, const TCHAR* PrefixName)
+const UPropertyBag* UPropertyBag::GetOrCreateFromDescs(const TConstArrayView<FPropertyBagPropertyDesc> InPropertyDescs, const TCHAR* PrefixName)
 {
+	TArray<FPropertyBagPropertyDesc> PropertyDescs = SanitizePropertyDescs(InPropertyDescs);
 	const uint64 BagHash = UE::StructUtils::Private::CalcPropertyDescArrayHash(PropertyDescs);
 	const FString ScriptStructName = PrefixName == nullptr
 		? FString::Printf(TEXT("PropertyBag_%llx"), BagHash)
@@ -4308,7 +4309,7 @@ const UPropertyBag* UPropertyBag::GetOrCreateFromDescs(const TConstArrayView<FPr
 
 	UPropertyBag* NewBag = NewObject<UPropertyBag>(GetTransientPackage(), *ScriptStructName, RF_Standalone | RF_Transient);
 
-	NewBag->PropertyDescs = SanitizePropertyDescs(PropertyDescs);
+	NewBag->PropertyDescs = MoveTemp(PropertyDescs);
 
 	// Add properties (AddCppProperty() adds them backwards in the linked list)
 	for (int32 DescIndex = NewBag->PropertyDescs.Num() - 1; DescIndex >= 0; DescIndex--)
