Fixing a bug in package reloading that miscalculated the batch size and called GC too frequently which could result in very poor performance.

#Synced-CL 49123435

diff --git a/Engine/Source/Runtime/CoreUObject/Private/UObject/PackageReload.cpp b/Engine/Source/Runtime/CoreUObject/Private/UObject/PackageReload.cpp
--- a/Engine/Source/Runtime/CoreUObject/Private/UObject/PackageReload.cpp
+++ b/Engine/Source/Runtime/CoreUObject/Private/UObject/PackageReload.cpp
@@ -615,13 +615,6 @@ void ReloadPackages(const TArrayView<FReloadPackageData>& InPackagesToReload, TA
 {
 	FPackageReferencersHelper PackageReferencersHelper(InPackagesToReload);
 
-	// Interdependencies between packages (in particular Blueprints) make it unsafe to run this logic in batches. 
-	// There are a number of edge cases that would have to be addressed if the batching logic were to be re-enabled, but 
-	// most likely the blueprint reparenting step would have to take in a TMap<UObject*, UObject*> that it could update
-	// when it decided that it needed to replace an instance due to hierarchy changes (e.g. class layout changing
-	// due to SuperStruct changes). For now, just process assets in a batch size of 1:
-	InNumPackagesPerBatch = 1;
-
 	FString Msg;
 	Msg.Append(FString::Printf(TEXT("Reloading %d Package(s):"), InPackagesToReload.Num()));
 	const int32 MAX_PACKAGES_TO_LOG = 10;
@@ -697,7 +690,12 @@ void ReloadPackages(const TArrayView<FReloadPackageData>& InPackagesToReload, TA
 				UPackage* NewPackage = NewPackages.Refs[PackageIndex].Package;
 				NewPackages.Refs[PackageIndex].EventData = PackageReloadInternal::GeneratePackageReloadEvent(ExistingPackage, NewPackage);
 
-				const bool bEndBatch = (PackageIndex == (BatchStartIndex + InNumPackagesPerBatch)) || (ExistingPackage && ExistingPackage->ContainsMap());
+				// Interdependencies between packages (in particular Blueprints) make it unsafe to run this logic in batches. 
+				// There are a number of edge cases that would have to be addressed if the batching logic were to be re-enabled, but 
+				// most likely the blueprint reparenting step would have to take in a TMap<UObject*, UObject*> that it could update
+				// when it decided that it needed to replace an instance due to hierarchy changes (e.g. class layout changing
+				// due to SuperStruct changes). For now, just process assets in a batch size of 1:
+				const bool bEndBatch = true; //(PackageIndex == (BatchStartIndex + InNumPackagesPerBatch - 1)) || (ExistingPackage && ExistingPackage->ContainsMap());
 				if (bEndBatch)
 				{
 					++PackageIndex; // We still need to move on-to the next package for the next batch
@@ -817,7 +815,11 @@ void ReloadPackages(const TArrayView<FReloadPackageData>& InPackagesToReload, TA
 					NewPackages.Refs[BatchPackageIndex].EventData.Reset();
 				}
 			}
-			CollectGarbage(GARBAGE_COLLECTION_KEEPFLAGS);
+
+			if (PackageIndex > 0 && (PackageIndex % InNumPackagesPerBatch) == 0)
+			{
+				CollectGarbage(GARBAGE_COLLECTION_KEEPFLAGS);
+			}
 
 			FCoreUObjectDelegates::OnPackageReloaded.Broadcast(EPackageReloadPhase::PostBatchPostGC, nullptr);
 		}
@@ -832,6 +834,8 @@ void ReloadPackages(const TArrayView<FReloadPackageData>& InPackagesToReload, TA
 			Package->SetDirtyFlag(false);
 		}
 	}, false);
+	
+	CollectGarbage(GARBAGE_COLLECTION_KEEPFLAGS);
 
 	// Finalization and error reporting
 	OutReloadedPackages.Reserve(ExistingPackages.Refs.Num());
