Linker Instancing Context - Avoid adding identical mappings (original -> instanced) to the instancing map
* Could lead to dummy entries that prevented further valid mappings from being added

#Synced-CL 49180728

diff --git a/Engine/Source/Runtime/CoreUObject/Private/UObject/LinkerInstancingContext.cpp b/Engine/Source/Runtime/CoreUObject/Private/UObject/LinkerInstancingContext.cpp
--- a/Engine/Source/Runtime/CoreUObject/Private/UObject/LinkerInstancingContext.cpp
+++ b/Engine/Source/Runtime/CoreUObject/Private/UObject/LinkerInstancingContext.cpp
@@ -201,7 +201,7 @@ public:
 			if (!bWasFound && InstancedPackageMapFunc)
 			{
 				RemappedPackageName = InstancedPackageMapFunc(PackageName);
-				bAddPackageMapping = true; 
+				bAddPackageMapping = RemappedPackageName != PackageName; 
 			}
 		}
 
@@ -751,4 +751,4 @@ bool FLinkerInstancingContextTests::RunTest(const FString& Parameters)
 	}
 
 	return true;
-}
\ No newline at end of file
+}
