Linker Instancing Context - Misc cleanup and optimizations

#Synced-CL 49180728

diff --git a/Engine/Source/Runtime/CoreUObject/Private/UObject/LinkerInstancingContext.cpp b/Engine/Source/Runtime/CoreUObject/Private/UObject/LinkerInstancingContext.cpp
--- a/Engine/Source/Runtime/CoreUObject/Private/UObject/LinkerInstancingContext.cpp
+++ b/Engine/Source/Runtime/CoreUObject/Private/UObject/LinkerInstancingContext.cpp
@@ -300,11 +300,13 @@ void FLinkerInstancingContext::BuildPackageMapping(FName Original, FName Instanc
 {
 	SharedData->BuildPackageMapping(Original, Instanced, GetSoftObjectPathRemappingEnabled());
 
+#if DO_ENSURE
 	if (ImmutableData.IsValid())
 	{
 		// check we didn't already built a package mapping on it 
 		ensureAlwaysMsgf(SharedData->ValidatePackageMappingAgainst(*ImmutableData), TEXT("Validate against ImmutableData failed after building PackageMapping for Original: %s, Instanced: %s"), *Original.ToString(), *Instanced.ToString());
 	}
+#endif
 }
 
 bool FLinkerInstancingContext::FindPackageMapping(FName Original, FName& Instanced) const
@@ -356,6 +358,7 @@ FSoftObjectPath FLinkerInstancingContext::RemapPath(const FSoftObjectPath& Path)
 /** Add a mapping from a package name to a new package name. There should be no separators (. or :) in these strings. */
 void FLinkerInstancingContext::AddPackageMapping(FName Original, FName Instanced)
 {
+#if DO_ENSURE
 	if (ImmutableData.IsValid())
 	{
 		// validate immutable package mapping doesn't contain this mapping or maps to the same package
@@ -365,6 +368,7 @@ void FLinkerInstancingContext::AddPackageMapping(FName Original, FName Instanced
 			ensureAlwaysMsgf(Instanced == AlreadyMapped, TEXT("Trying to add duplicate package mapping, this add is masked by the previous one (present in the Immutable linked InstancingContext) Original: %s, Instanced: %s, InMapping: %s"), *Original.ToString(), *Instanced.ToString(), *AlreadyMapped.ToString());
 		}
 	}
+#endif
 
 	SharedData->AddPackageMapping(Original, Instanced);
 }
@@ -379,6 +383,7 @@ void FLinkerInstancingContext::AddPackageMappingFunc(TFunction<FName(FName)> InI
 /** Add a mapping from a top level asset path (/Path/To/Package.AssetName) to another. */
 void FLinkerInstancingContext::AddPathMapping(FSoftObjectPath Original, FSoftObjectPath Instanced)
 {
+#if DO_ENSURE
 	ensureAlwaysMsgf(Original.GetSubPathString().IsEmpty(),
 		TEXT("Linker instance remap paths should be top-level assets only: %s->"), *Original.ToString());
 	ensureAlwaysMsgf(Instanced.GetSubPathString().IsEmpty(),
@@ -393,6 +398,7 @@ void FLinkerInstancingContext::AddPathMapping(FSoftObjectPath Original, FSoftObj
 			ensureAlwaysMsgf(Instanced == AlreadyMapped, TEXT("Trying to add duplicate path mapping, this add is masked by the previous one (present in the Immutable linked InstancingContext) Original: %s, Instanced: %s, InMapping: %s"), *Original.ToString(), *Instanced.ToString(), *AlreadyMapped.ToString());
 		}	
 	}
+#endif
 
 	SharedData->AddPathMapping(Original, Instanced);
 }
@@ -444,30 +450,17 @@ void FLinkerInstancingContext::FixupSoftObjectPath(FSoftObjectPath& InOutSoftObj
 
 void FLinkerInstancedPackageMap::AddPackageMapping(FName Original, FName Instanced)
 {
-	if (InstanceMappingDirection == EInstanceMappingDirection::OriginalToInstanced)
+	FName Key = InstanceMappingDirection == EInstanceMappingDirection::OriginalToInstanced ? Original : Instanced;
+	FName Value = InstanceMappingDirection == EInstanceMappingDirection::OriginalToInstanced ? Instanced : Original;
+
+	const FName& StoredValue = InstancedPackageMapping.FindOrAdd(Key, Value);
+	if (StoredValue == Value)
 	{
-		if (!InstancedPackageMapping.Contains(Original))
-		{
-			InstancedPackageMapping.Add(Original, Instanced);
-			bIsInstanced |= !Instanced.IsNone();
-		}
-		else
-		{
-			ensureAlwaysMsgf(InstancedPackageMapping.FindChecked(Original) == Instanced, TEXT("Trying to add duplicate package mapping, this add is masked by the previous one. Original: %s, Instanced: %s, InMapping: %s"), *Original.ToString(), *Instanced.ToString(), *InstancedPackageMapping.FindChecked(Original).ToString());
-		}
+		bIsInstanced |= !Value.IsNone();
 	}
 	else
 	{
-		check(InstanceMappingDirection == EInstanceMappingDirection::InstancedToOriginal);
-		if (!InstancedPackageMapping.Contains(Instanced))
-		{
-			InstancedPackageMapping.Add(Instanced, Original);
-			bIsInstanced |= !Original.IsNone();
-		}
-		else
-		{
-			ensureAlwaysMsgf(InstancedPackageMapping.FindChecked(Instanced) == Original, TEXT("Trying to add duplicate package mapping, this add is masked by the previous one. Original: %s, Instanced: %s, InMapping: %s"), *Original.ToString(), *Instanced.ToString(), *InstancedPackageMapping.FindChecked(Instanced).ToString());
-		}
+		ensureAlwaysMsgf(StoredValue == Value, TEXT("Trying to add duplicate package mapping, this add is masked by the previous one. Original: %s, Instanced: %s, InMapping: %s"), *Original.ToString(), *Instanced.ToString(), *StoredValue.ToString());
 	}
 }
 
