Add a PostLoad callback to UScriptStruct. Only valid in the editor. Use when a struct instance needs to execute deprecation but has to wait until the UObject is loaded. The struct instance that wants the callback has to register with the FStructPostLoadContext. When the object that owns the struct is loaded, it will find the struct instance and call PostLoad.
#Synced-CL 49284918

diff --git a/Engine/Source/Runtime/CoreUObject/Private/UObject/Class.cpp b/Engine/Source/Runtime/CoreUObject/Private/UObject/Class.cpp
--- a/Engine/Source/Runtime/CoreUObject/Private/UObject/Class.cpp
+++ b/Engine/Source/Runtime/CoreUObject/Private/UObject/Class.cpp
@@ -3335,6 +3335,11 @@ void UScriptStruct::PrepareCppStructOps()
 		UE_LOG(LogClass, Verbose, TEXT("Native struct %s has native CanEditChange."), *GetName());
 		StructFlags = EStructFlags(StructFlags | STRUCT_CanEditChange);
 	}
+	if (CppStructOps->HasPostLoad())
+	{
+		UE_LOG(LogClass, Verbose, TEXT("Native struct %s has native PostLoad."), *GetName());
+		StructFlags = EStructFlags(StructFlags | STRUCT_PostLoad);
+	}
 #endif
 
 	if (CppStructOps->HasVisitor())
diff --git a/Engine/Source/Runtime/CoreUObject/Private/UObject/Obj.cpp b/Engine/Source/Runtime/CoreUObject/Private/UObject/Obj.cpp
--- a/Engine/Source/Runtime/CoreUObject/Private/UObject/Obj.cpp
+++ b/Engine/Source/Runtime/CoreUObject/Private/UObject/Obj.cpp
@@ -41,6 +41,7 @@
 #include "UObject/PropertyPortFlags.h"
 #include "UObject/PropertyStateTracking.h"
 #include "UObject/UnrealType.h"
+#include "UObject/StructPostLoadContext.h"
 #include "UObject/ObjectInstancingGraph.h"
 #include "UObject/ObjectRedirector.h"
 #include "UObject/ObjectSaveContext.h"
@@ -477,6 +478,10 @@ void UObject::PostLoad()
 		LoadConfig();
 	}
 	CheckDefaultSubobjects();
+
+#if WITH_EDITOR
+	UE::FStructPostLoadContext::Get().OnPostLoad(this);
+#endif
 }
 
 #if WITH_EDITOR
diff --git a/Engine/Source/Runtime/CoreUObject/Private/UObject/StructPostLoadContext.cpp b/Engine/Source/Runtime/CoreUObject/Private/UObject/StructPostLoadContext.cpp
new file mode 100644
--- /dev/null
+++ b/Engine/Source/Runtime/CoreUObject/Private/UObject/StructPostLoadContext.cpp
@@ -0,0 +1,81 @@
+// Copyright Epic Games, Inc. All Rights Reserved.
+
+#include "UObject/StructPostLoadContext.h"
+
+#include "UObject/Class.h"
+#include "UObject/Object.h"
+#include "UObject/UObjectThreadContext.h"
+#include "UObject/UnrealType.h"
+
+#if WITH_EDITOR
+namespace UE
+{
+	static FStructPostLoadContext GStructPostLoadContext;
+	FStructPostLoadContext& FStructPostLoadContext::Get()
+	{
+		return GStructPostLoadContext;
+	}
+
+	bool FStructPostLoadContext::RequestPostLoad(TNotNull<const FUObjectSerializeContext*> Context)
+	{
+		if (Context->SerializedObject)
+		{
+			FAnnotation Annotation = Annotations.GetAndRemoveAnnotation(Context->SerializedObject);
+			if (Context->SerializedPropertyPath.GetSegmentCount() == 1
+				&& Context->SerializedPropertyPath.GetSegment(0).Index == INDEX_NONE)
+			{
+				Annotation.StructPaths.Add(Context->SerializedPropertyPath);
+			}
+			else
+			{
+				Annotation.bHasInvalidPropertyPathName = true;
+			}
+
+			Annotations.AddAnnotation(Context->SerializedObject, MoveTemp(Annotation));
+			return true;
+		}
+		return false;
+	}
+
+	void FStructPostLoadContext::OnPostLoad(TNotNull<UObjectBase*> Object)
+	{
+		const UE::FStructPostLoadContext::FAnnotation Annotation = Annotations.GetAndRemoveAnnotation(Object);
+		if (!Annotation.IsDefault())
+		{
+			auto CallPostLoad = [](const FStructProperty* Property, void* PropertyValue)
+			{
+				UScriptStruct::ICppStructOps* TheCppStructOps = Property->Struct->GetCppStructOps();
+				check(TheCppStructOps);
+				TheCppStructOps->PostLoad(PropertyValue);
+			};
+
+			if (!Annotation.bHasInvalidPropertyPathName)
+			{
+				for (const FPropertyPathName& PathName : Annotation.StructPaths)
+				{
+					const FProperty* Property = Object->GetClass()->FindPropertyByName(PathName.GetSegment(0).Name);
+					const FStructProperty* StructProperty = CastField<const FStructProperty>(Property);
+					if (ensure(StructProperty))
+					{
+						UObjectBase* Container = Object;
+						void* PropertyValue = StructProperty->ContainerPtrToValuePtr<void>(Container);
+						CallPostLoad(StructProperty, PropertyValue);
+					}
+				}
+			}
+			else
+			{
+				for (TPropertyValueIterator<FStructProperty> It(Object->GetClass(), Object); It; ++It)
+				{
+					const FStructProperty* Property = It.Key();
+					if (Property->Struct->StructFlags & STRUCT_PostLoad)
+					{
+						void* PropertyValue = (void*)It.Value();
+						CallPostLoad(Property, PropertyValue);
+					}
+				}
+			}
+		}
+	}
+}// namespace UE
+#endif // WITH_EDITOR
\ No newline at end of file
diff --git a/Engine/Source/Runtime/CoreUObject/Public/UObject/Class.h b/Engine/Source/Runtime/CoreUObject/Public/UObject/Class.h
--- a/Engine/Source/Runtime/CoreUObject/Public/UObject/Class.h
+++ b/Engine/Source/Runtime/CoreUObject/Public/UObject/Class.h
@@ -1125,6 +1125,9 @@ enum EStructFlags
 
 	/** If set, this struct will have Visit on it to allow custom property visiting implementation */
 	STRUCT_Visitor = 0x04000000,
+
+	/** If set, this struct will have PostLoad on it in the editor. */
+	STRUCT_PostLoad UE_INTERNAL = 0x08000000,
 	
 	/** Struct flags that are automatically inherited */
 	STRUCT_Inherit				= STRUCT_HasInstancedReference|STRUCT_Atomic,
@@ -1193,6 +1196,7 @@ namespace UE::CoreUObject::Private
 		bool ClearOnFinishDestroy : 1;
 #if WITH_EDITOR
 		bool HasCanEditChange : 1;
+		/*UE_INTERNAL*/ bool HasPostLoad : 1;
 #endif
 		bool HasVisitor : 1;
 		bool HasGetPreloadDependencies : 1;
@@ -1238,6 +1242,9 @@ namespace UE::CoreUObject::Private
 			.ClearOnFinishDestroy						= TraitsType::WithClearOnFinishDestroy,
 #if WITH_EDITOR
 			.HasCanEditChange							= TraitsType::WithCanEditChange,
+			PRAGMA_DISABLE_INTERNAL_WARNINGS
+			.HasPostLoad								= TraitsType::WithPostLoad,
+			PRAGMA_ENABLE_INTERNAL_WARNINGS
 #endif
 			.HasVisitor									= TraitsType::WithVisitor,
 			.HasGetPreloadDependencies					= TraitsType::WithGetPreloadDependencies,
@@ -1518,6 +1525,12 @@ namespace UE::CoreUObject::Private
 	{
 		return ((const CppStruct*)Data)->CanEditChange(PropertyChain);
 	}
+	template <typename CppStruct>
+		requires (!!TStructOpsTypeTraits<CppStruct>::WithPostLoad)
+	UE_INTERNAL void PostLoad(void* Data)
+	{
+		((CppStruct*)Data)->PostLoad();
+	}
 	#define UE_FOREACH_FAKE_VTABLE_FUNC_EDITOR_ONLY(...) __VA_ARGS__
 #else
 	#define UE_FOREACH_FAKE_VTABLE_FUNC_EDITOR_ONLY(...)
@@ -1569,11 +1582,12 @@ namespace UE::CoreUObject::Private
 		IndexNameFunctype(19,	InitializeIntrusiveUnsetOptionalValue,	void(void* Data)) \
 		IndexNameFunctype(20,	IsIntrusiveOptionalValueSet,			bool(const void* Data)) \
 		IndexNameFunctype(21,	ClearIntrusiveOptionalValue,			void(void* Data)) \
+		IndexNameFunctype(22,	Visit,									EPropertyVisitorControlFlow(FPropertyVisitorContext& Context, const TFunctionRef<EPropertyVisitorControlFlow(const FPropertyVisitorContext& Context)> InFunc)) \
+		IndexNameFunctype(23,	ResolveVisitedPathInfo,					void*(void* Data, const FPropertyVisitorInfo& Info)) \
 		UE_FOREACH_FAKE_VTABLE_FUNC_EDITOR_ONLY( \
-		IndexNameFunctype(22,	CanEditChange,							bool(const FEditPropertyChain& PropertyChain, const void* Data)) \
-		) \
-		IndexNameFunctype(23,	Visit,									EPropertyVisitorControlFlow(FPropertyVisitorContext& Context, const TFunctionRef<EPropertyVisitorControlFlow(const FPropertyVisitorContext& Context)> InFunc)) \
-		IndexNameFunctype(24,	ResolveVisitedPathInfo,					void*(void* Data, const FPropertyVisitorInfo& Info))
+		IndexNameFunctype(24,	CanEditChange,							bool(const FEditPropertyChain& PropertyChain, const void* Data)) \
+		IndexNameFunctype(25,	PostLoad,								void(void* Data)) \
+		)
 
 	// These represent one function per optional function in the fake vtable.
 	enum class EStructOpsFakeVTableFlags : uint32
@@ -2155,6 +2169,23 @@ public:
 		{
 			return HasCanEditChange() && FakeVPtr->GetHandlingFunc<UE::CoreUObject::Private::EStructOpsFakeVTableFlags::CanEditChange>()(PropertyChain, Data);
 		}
+
+		/** Returns true if this struct support post load. */
+		UE_INTERNAL bool HasPostLoad() const
+		{
+			return GetCapabilities().HasPostLoad;
+		}
+
+		/** PostLoad property PostLoad signature. */
+		UE_INTERNAL void PostLoad(void* Data) const
+		{
+			PRAGMA_DISABLE_INTERNAL_WARNINGS
+			if (HasPostLoad())
+			{
+				FakeVPtr->GetHandlingFunc<UE::CoreUObject::Private::EStructOpsFakeVTableFlags::PostLoad>()(Data);
+			}
+			PRAGMA_ENABLE_INTERNAL_WARNINGS
+		}
 #endif
 
 		/** Returns true if this struct wants to indicate whether it has a custom impl for visiting properties */
diff --git a/Engine/Source/Runtime/CoreUObject/Public/UObject/StructOpsTypeTraits.h b/Engine/Source/Runtime/CoreUObject/Public/UObject/StructOpsTypeTraits.h
--- a/Engine/Source/Runtime/CoreUObject/Public/UObject/StructOpsTypeTraits.h
+++ b/Engine/Source/Runtime/CoreUObject/Public/UObject/StructOpsTypeTraits.h
@@ -31,11 +31,12 @@ struct TStructOpsTypeTraitsBase2
 		WithNetSharedSerialization     = false,                         // struct has a NetSerialize function that does not require the package map to serialize its state.
 		WithGetPreloadDependencies     = false,                         // struct has a GetPreloadDependencies function to return all objects that will be Preload()ed when the struct is serialized at load time.
 		WithPureVirtual                = false,                         // struct has PURE_VIRTUAL functions and cannot be constructed when CHECK_PUREVIRTUALS is true
-		WithFindInnerPropertyInstance  = false,							// struct has a FindInnerPropertyInstance function that can provide an FProperty and data pointer when given a property FName
-		WithCanEditChange			   = false,							// struct has an editor-only CanEditChange function that can conditionally make child properties read-only in the details panel (same idea as UObject::CanEditChange)
-		WithClearOnFinishDestroy	   = false,							// struct should be cleared during owner UObject's FinishDestroy. Clearing calls destructor and initializes again to default value. This is intended for structs which may need to access UObject pointer members during destruction. Referenced objects may already have their FinishDestroy() called. Clearing should ensure that no UObject pointer members are used during the final destruction.
-		WithVisitor					   = false,							// struct has Visit function that allows to visit additional properties
-		WithIntrusiveOptionalSafeForGC = false,							// struct with an intrusive unset state for TOptional certifies that object fields it contains are nulled in the unset state
+		WithFindInnerPropertyInstance  = false,                         // struct has a FindInnerPropertyInstance function that can provide an FProperty and data pointer when given a property FName
+		WithCanEditChange              = false,                         // struct has an editor-only CanEditChange function that can conditionally make child properties read-only in the details panel (same idea as UObject::CanEditChange)
+		WithClearOnFinishDestroy       = false,                         // struct should be cleared during owner UObject's FinishDestroy. Clearing calls destructor and initializes again to default value. This is intended for structs which may need to access UObject pointer members during destruction. Referenced objects may already have their FinishDestroy() called. Clearing should ensure that no UObject pointer members are used during the final destruction.
+		WithVisitor                    = false,                         // struct has Visit function that allows to visit additional properties
+		WithIntrusiveOptionalSafeForGC = false,                         // struct with an intrusive unset state for TOptional certifies that object fields it contains are nulled in the unset state
+		WithPostLoad UE_INTERNAL       = false,                         // struct has an editor-only PostLoad function which is called in editor by the UObject::PostLoad that owns the struct. The user has to request the callback via the annotation. Use when a deprecation requires access to a UObject.
 	};
 
 	static constexpr EPropertyObjectReferenceType WithSerializerObjectReferences = EPropertyObjectReferenceType::Conservative; // struct's Serialize method(s) may serialize object references of these types - default Conservative means unknown and object reference collector archives should serialize this struct 
diff --git a/Engine/Source/Runtime/CoreUObject/Public/UObject/StructPostLoadContext.h b/Engine/Source/Runtime/CoreUObject/Public/UObject/StructPostLoadContext.h
new file mode 100644
--- /dev/null
+++ b/Engine/Source/Runtime/CoreUObject/Public/UObject/StructPostLoadContext.h
@@ -0,0 +1,56 @@
+// Copyright Epic Games, Inc. All Rights Reserved.
+#pragma once
+
+#include "Misc/NotNull.h"
+#include "UObject/PropertyPathName.h"
+#include "UObject/UObjectAnnotation.h"
+
+struct FUObjectSerializeContext;
+class UObjectBase;
+
+#if WITH_EDITOR
+namespace UE
+{
+	/**
+	 * Request a struct PostLoad during serialization.
+	 * This is only valid in the editor.
+	 * Use for structs that need to do fix up logic after all objects have been loaded.
+	 */
+	class UE_INTERNAL FStructPostLoadContext
+	{
+	public:
+		/** Struct PostLoad request annotation. */
+		struct FAnnotation
+		{
+			bool IsDefault() const
+			{
+				return !bHasInvalidPropertyPathName && StructPaths.Num() == 0;
+			}
+
+			/** The property path of the struct that request a PostLoad. **/
+			TArray<FPropertyPathName, TInlineAllocator<1>> StructPaths;
+
+			/** The serialize context doesn't contain a valid property path. */
+			bool bHasInvalidPropertyPathName = false;
+		};
+
+	public:
+		COREUOBJECT_API static FStructPostLoadContext& Get();
+
+		/** A struct request a post load. */
+		COREUOBJECT_API bool RequestPostLoad(TNotNull<const FUObjectSerializeContext*> Context);
+
+		/** Gets the annotation for the given object. */
+		FAnnotation GetAndRemoveRequest(TNotNull<const UObjectBase*> Object)
+		{
+			return Annotations.GetAndRemoveAnnotation(Object);
+		}
+
+		/** Run the PostLoad on any request for the object. */
+		COREUOBJECT_API void OnPostLoad(TNotNull<UObjectBase*> Object);
+
+	private:
+		FUObjectAnnotationSparse<FAnnotation, true/*bAutoRemove*/> Annotations;
+	};
+}
+#endif
\ No newline at end of file
