Removing SoftObjectPtr validation that is no longer safe with MT async-load.

#Synced-CL 49316295

diff --git a/Engine/Source/Runtime/CoreUObject/Private/UObject/PropertySoftObjectPtr.cpp b/Engine/Source/Runtime/CoreUObject/Private/UObject/PropertySoftObjectPtr.cpp
--- a/Engine/Source/Runtime/CoreUObject/Private/UObject/PropertySoftObjectPtr.cpp
+++ b/Engine/Source/Runtime/CoreUObject/Private/UObject/PropertySoftObjectPtr.cpp
@@ -83,19 +83,7 @@ void FSoftObjectProperty::SerializeItem( FStructuredArchive::FSlot Slot, void* V
 	// Allow persistent archives so they can keep track of string references. (e.g. FArchiveSaveTagImports)
 	if( !UnderlyingArchive.IsObjectReferenceCollector() || UnderlyingArchive.IsModifyingWeakAndStrongReferences() || UnderlyingArchive.IsPersistent() )
 	{
-		FSoftObjectPtr OldValue = *(FSoftObjectPtr*)Value;
 		Slot << *(FSoftObjectPtr*)Value;
-
-		// Check for references to instances of wrong types and null them out 
-#if !(UE_BUILD_TEST || UE_BUILD_SHIPPING) 
-		if (UnderlyingArchive.IsLoading() || UnderlyingArchive.IsModifyingWeakAndStrongReferences())
-		{
-			if (OldValue.GetUniqueID() != ((FSoftObjectPtr*)Value)->GetUniqueID())
-			{
-				CheckValidObject(Value, nullptr, Defaults); // FSoftObjectProperty is never non-nullable at this point so it's ok to pass null as the current value
-			}
-		}
-#endif
 	}
 	else
 	{
