Fixing cluster mutable references being garbage collected if that cluster is referenced only by another cluster

#Synced-CL 49321981

diff --git a/Engine/Source/Runtime/CoreUObject/Private/UObject/GarbageCollection.cpp b/Engine/Source/Runtime/CoreUObject/Private/UObject/GarbageCollection.cpp
--- a/Engine/Source/Runtime/CoreUObject/Private/UObject/GarbageCollection.cpp
+++ b/Engine/Source/Runtime/CoreUObject/Private/UObject/GarbageCollection.cpp
@@ -1201,7 +1201,11 @@ static FORCENOINLINE void MarkReferencedClustersAsReachable(int32 ClusterIndex,
 			FUObjectItem* ReferencedClusterRootObjectItem = GUObjectArray.IndexToObjectUnsafeForGC(ReferncedClusterIndex);
 			if (!ReferencedClusterRootObjectItem->HasAnyFlags(EInternalObjectFlags::Garbage))
 			{
-				FGCFlags::FastMarkAsReachableInterlocked_ForGC(ReferencedClusterRootObjectItem);
+				if (FGCFlags::MarkAsReachableInterlocked_ForGC(ReferencedClusterRootObjectItem))
+				{
+					// Mark everything in the referenced cluster along with its mutable objects as reachable
+					MarkReferencedClustersAsReachable<Options>(ReferencedClusterRootObjectItem->GetClusterIndex(), ObjectsToSerialize);
+				}
 			}
 			else
 			{
