[UE] Fix for non-native instancing LLTs running inside the editor automation context as a single pass.

Notes:
- Ensure that transient mutations on non-native CDO unit tests get cleaned up before moving on to the next test that will use the same test class instance.

#Synced-CL 49321981

diff --git a/Engine/Source/Runtime/CoreUObject/Tests/UObject/SubobjectInstancingTest.cpp b/Engine/Source/Runtime/CoreUObject/Tests/UObject/SubobjectInstancingTest.cpp
--- a/Engine/Source/Runtime/CoreUObject/Tests/UObject/SubobjectInstancingTest.cpp
+++ b/Engine/Source/Runtime/CoreUObject/Tests/UObject/SubobjectInstancingTest.cpp
@@ -1736,6 +1736,10 @@ namespace SubobjectInstancingTest::Private
 				// the native object initializer will not include properties inherited from the native super class hierarchy, because it expects they have already been initialized by the native ctor. For this
 				// reason, users cannot override a native instanced reference member's default value. This test is validating that any such override will not actually be used by subobject template instancing.
 				CHECK(OuterObject->EditTimeInnerObject == nullptr);
+
+				// Clean up the subobject on the CDO so it doesn't persist through remaining tests.
+				CDO->EditTimeInnerObject->Rename(nullptr, GetTransientPackage(), REN_DontCreateRedirectors);
+				CDO->EditTimeInnerObject = nullptr;
 			}
 		}
 
@@ -1804,6 +1808,10 @@ namespace SubobjectInstancingTest::Private
 				// the native object initializer will not include properties inherited from the native super class hierarchy, because it expects they have already been initialized by the native ctor. For this
 				// reason, users cannot override a native instanced reference member's default value. This test is validating that any such override will not actually be used by subobject template instancing.
 				CHECK(OuterObject->TransientEditTimeInnerObject == nullptr);
+
+				// Clean up the subobject on the CDO so it doesn't persist through remaining tests.
+				CDO->TransientEditTimeInnerObject->Rename(nullptr, GetTransientPackage(), REN_DontCreateRedirectors);
+				CDO->TransientEditTimeInnerObject = nullptr;
 			}
 		}
 
@@ -1881,6 +1889,10 @@ namespace SubobjectInstancingTest::Private
 
 				USubobjectInstancingTestOuterObject* NewOuterObject = NewObject<USubobjectInstancingTestOuterObject>(GetTransientPackage(), InstancingTestClass);
 				CHECK(NonNativeEditTimeInnerObject != ObjectProperty->GetObjectPropertyValue_InContainer(NewOuterObject));
+
+				// Clean up the subobject on the CDO so it doesn't persist through remaining tests.
+				ObjectProperty->SetObjectPropertyValue_InContainer(CDO, nullptr);
+				NonNativeEditTimeInnerObjectDefaultValue->Rename(nullptr, GetTransientPackage(), REN_DontCreateRedirectors);
 			}
 		}
 
