[UE] Ensure default values constructed for non-nullable reference types don't propagate template flags when the existing value is used as a template.

#Synced-CL 49321981

diff --git a/Engine/Source/Runtime/CoreUObject/Private/UObject/PropertyBaseObject.cpp b/Engine/Source/Runtime/CoreUObject/Private/UObject/PropertyBaseObject.cpp
--- a/Engine/Source/Runtime/CoreUObject/Private/UObject/PropertyBaseObject.cpp
+++ b/Engine/Source/Runtime/CoreUObject/Private/UObject/PropertyBaseObject.cpp
@@ -778,6 +778,8 @@ static bool CanPropertyHoldObject(const FObjectPropertyBase* Prop, TObjectPtr<co
 
 UObject* FObjectPropertyBase::ConstructDefaultObjectValueIfNecessary(UObject* ExistingValue, FString* OutFailureReason, const void* Defaults) const
 {
+	UObject* Template = nullptr;
+	UClass* NewClass = PropertyClass;
 	UObject* NewDefaultObjectValue = nullptr;
 
 	FUObjectSerializeContext* SerializeContext = FUObjectThreadContext::Get().GetSerializeContext();
@@ -787,10 +789,6 @@ UObject* FObjectPropertyBase::ConstructDefaultObjectValueIfNecessary(UObject* Ex
 		Outer = GetTransientPackage();
 	}
 
-	// Filter out any flags that do not belong to the RF_Load or RF_PropagateToSubObjects groups. We mostly want to exclude
-	// UObject life-time and internal GC flags. Include PostLoad/PostLoadSubobjects to preserve any PostLoad logic during loads.
-	const EObjectFlags FlagsToKeepMask = (RF_Load | RF_NeedPostLoad | RF_NeedPostLoadSubobjects | RF_PropagateToSubObjects);
-
 	if (ExistingValue)
 	{
 		UClass* ExistingValueClass = ExistingValue->GetClass();
@@ -800,10 +798,9 @@ UObject* FObjectPropertyBase::ConstructDefaultObjectValueIfNecessary(UObject* Ex
 			if (ExistingValue->IsTemplate() && 	// Existing value is a template so we can construct a new value with it as the archetype
 				ExistingValue->GetOuter() != Outer) // Unless the template's Outer is the same as the new Outer in which case the template (ExistingValue) IS the object we can reuse
 			{
-				EObjectFlags NewFlags = ExistingValue->GetFlags() & FlagsToKeepMask;
-
 				// We probably got here because an object value failed to load (missing import class) and the property is left with a template of default subobject
-				NewDefaultObjectValue = NewObject<UObject>(Outer, ExistingValue->GetClass(), NAME_None, NewFlags, ExistingValue);
+				Template = ExistingValue;
+				NewClass = ExistingValueClass;
 			}
 			else
 			{
@@ -850,17 +847,33 @@ UObject* FObjectPropertyBase::ConstructDefaultObjectValueIfNecessary(UObject* Ex
 			bool bCanReplaceObjectWithExistingName = !ObjectWithExistingName || ObjectWithExistingName->GetClass()->IsChildOf(PropertyClass);
 
 			FName NewName = (bCanReplaceObjectWithExistingName) ? ExistingName : NAME_None;
-			EObjectFlags NewFlags = ExistingValue ? (ExistingValue->GetFlags() & FlagsToKeepMask) : Outer->GetMaskedFlags(RF_NeedPostLoad | RF_NeedPostLoadSubobjects | RF_PropagateToSubObjects);
-			UObject* Template = Defaults ? GetObjectPropertyValue(Defaults) : nullptr;
+
+			// Filter out any flags that do not belong to the RF_Load or RF_PropagateToSubObjects groups. We mostly want to exclude
+			// UObject life-time and internal GC flags. Include PostLoad/PostLoadSubobjects to preserve any PostLoad logic during loads.
+			EObjectFlags NewFlags;
+			if (ExistingValue && ExistingValue != Template)
+			{
+				NewFlags = ExistingValue->GetMaskedFlags(RF_Load | RF_NeedPostLoad | RF_NeedPostLoadSubobjects | RF_PropagateToSubObjects);
+			}
+			else
+			{
+				NewFlags = Outer->GetMaskedFlags(RF_PropagateToSubObjects | RF_NeedPostLoad | RF_NeedPostLoadSubobjects);
+			}
+			
+			if (!Template && Defaults)
+			{
+				Template = GetObjectPropertyValue(Defaults);
+			}
+			
 			if (Template && !Template->IsA(PropertyClass))
 			{
 				Template = nullptr;
 			}
 
-			// Existing value did not exist or it could not be used as a template
+			// Existing value did not exist, was not type-compatible or is being used as a template
 			// Existing value may be null in case we were serializing an array of UObjects that failed to load (missing import class). Since the array is first pre-allocated with null values
 			// it will not have any existing objects to instantiate
-			NewDefaultObjectValue = NewObject<UObject>(Outer, PropertyClass, NewName, NewFlags, Template);
+			NewDefaultObjectValue = NewObject<UObject>(Outer, NewClass, NewName, NewFlags, Template);
 		}
 	}
 
