Fixed not returning correctly the contains overrides on a subobject in the details panel. Found by [REDACTED_HANDLE]

#Synced-CL 49377257

diff --git a/Engine/Source/Runtime/CoreUObject/Private/UObject/OverridableManager.cpp b/Engine/Source/Runtime/CoreUObject/Private/UObject/OverridableManager.cpp
--- a/Engine/Source/Runtime/CoreUObject/Private/UObject/OverridableManager.cpp
+++ b/Engine/Source/Runtime/CoreUObject/Private/UObject/OverridableManager.cpp
@@ -194,6 +194,10 @@ EOverriddenState FOverridableManager::GetOverriddenState(TNotNull<UObject*> Obje
 		{
 			const EOverriddenPropertyOperation Operation = OverriddenProperties->GetOverriddenPropertyOperation(Context.Path.GetRootIterator());
 			bool bIsEqual = false;
+
+			const FProperty* Property = Context.Path.Top().Property;
+			const FObjectPropertyBase* ObjectProperty = CastField<FObjectPropertyBase>(Property);
+
 			if (Context.Path.Num() <= Path.Num())
 			{
 				if (!Context.Path.Contained(Path, &bIsEqual))
@@ -201,7 +205,10 @@ EOverriddenState FOverridableManager::GetOverriddenState(TNotNull<UObject*> Obje
 					return EPropertyVisitorControlFlow::StepOver;
 				}
 
-				if (Operation == EOverriddenPropertyOperation::Modified)
+				// Recurse if we can tell that from the operation that it contains overrides,
+				// if the operation is none, still recurse if it is an object property
+				if (Operation == EOverriddenPropertyOperation::Modified ||
+					(ObjectProperty && Operation == EOverriddenPropertyOperation::None))
 				{
 					return EPropertyVisitorControlFlow::StepInto;
 				}
@@ -221,8 +228,6 @@ EOverriddenState FOverridableManager::GetOverriddenState(TNotNull<UObject*> Obje
 				return EPropertyVisitorControlFlow::Stop;
 			}
 
-			const FProperty* Property = Context.Path.Top().Property;
-
 			// Skip always overridden properties
 			if (Property->HasAnyPropertyFlags(CPF_ExperimentalAlwaysOverriden))
 			{
@@ -252,7 +257,7 @@ EOverriddenState FOverridableManager::GetOverriddenState(TNotNull<UObject*> Obje
 				}
 
 				// Only recurse on instanced subobjects
-				if (const FObjectPropertyBase* ObjectProperty = CastField<FObjectPropertyBase>(Property))
+				if (ObjectProperty)
 				{
 					if (UObject* InstancedSubObject = ObjectProperty->GetObjectPropertyValue(Context.Data.PropertyData))
 					{
