Clang/linux fixes for UnrealAssetStringify

#rnx
#Synced-CL 49566423

diff --git a/Engine/Source/Programs/UnrealAssetStringify/Private/UnrealAssetStringify.cpp b/Engine/Source/Programs/UnrealAssetStringify/Private/UnrealAssetStringify.cpp
--- a/Engine/Source/Programs/UnrealAssetStringify/Private/UnrealAssetStringify.cpp
+++ b/Engine/Source/Programs/UnrealAssetStringify/Private/UnrealAssetStringify.cpp
@@ -109,8 +109,8 @@ struct FUnrealAssetStringifyArgs
 
 bool FUnrealAssetStringifyArgs::ParseArgs(int32 ArgC, TCHAR* const ArgV[], FUnrealAssetStringifyArgs& OutArgs)
 {
-	int32 I = 1;
-	for(I; I < ArgC - 1; ++I)
+	int32 I;
+	for(I = 1; I < ArgC - 1; ++I)
 	{
 		if(FStringView(ArgV[I]) == TEXT("--save"))
 		{
diff --git a/Engine/Source/Runtime/CoreUObject/Private/UObject/PropertyTag.cpp b/Engine/Source/Runtime/CoreUObject/Private/UObject/PropertyTag.cpp
--- a/Engine/Source/Runtime/CoreUObject/Private/UObject/PropertyTag.cpp
+++ b/Engine/Source/Runtime/CoreUObject/Private/UObject/PropertyTag.cpp
@@ -608,7 +608,7 @@ void FPropertyTag::SerializeTaggedProperty(FStructuredArchive::FSlot Slot, FProp
 
 	// Ensure that we serialize what we expected to serialize.
 	const int64 EndOfProperty = UnderlyingArchive.Tell();
-	if (Size && (EndOfProperty - StartOfProperty != Size) 
+	if ((Size && (EndOfProperty - StartOfProperty != Size))
 #if UE_INTERNAL_UNIVERSAL_PLACEHOLDERS
 	|| UnderlyingArchive.IsCriticalError()
 #endif//UE_INTERNAL_UNIVERSAL_PLACEHOLDERS
diff --git a/Engine/Source/Runtime/CoreUObject/Public/UObject/ObjectPtr.h b/Engine/Source/Runtime/CoreUObject/Public/UObject/ObjectPtr.h
--- a/Engine/Source/Runtime/CoreUObject/Public/UObject/ObjectPtr.h
+++ b/Engine/Source/Runtime/CoreUObject/Public/UObject/ObjectPtr.h
@@ -890,7 +890,7 @@ namespace ObjectPtr_Private
 				// can't create a placeholder that necessarily conforms to the native code's
 				// requirements.
 				const UObject* Resolved = InObjectPtr;
-				if(Resolved && !(Resolved->IsA<T>()))
+				if(Resolved && !(InObjectPtr.template IsA<T>()))
 				{
 					InObjectPtr = nullptr;
 				}
