Pre-reserve export map space in NotifyCompiledVersePackage

#rnx
#Synced-CL 49580249

diff --git a/Engine/Source/Runtime/CoreUObject/Private/Serialization/AsyncLoading2.cpp b/Engine/Source/Runtime/CoreUObject/Private/Serialization/AsyncLoading2.cpp
--- a/Engine/Source/Runtime/CoreUObject/Private/Serialization/AsyncLoading2.cpp
+++ b/Engine/Source/Runtime/CoreUObject/Private/Serialization/AsyncLoading2.cpp
@@ -11678,62 +11678,71 @@ void FAsyncLoadingThread2::NotifyCompiledVersePackage(Verse::VPackage* VersePack
 	UPackage* Package = VersePackage->GetOrCreateUPackage(Context);
 	FName PackageName = Package->GetFName();
 	FPackageId PackageId = FPackageId::FromName(PackageName);
+
+	// Construct an export map to reserve space in the package, rather than inserting objects directly.
+	// StoreGlobalObject does not amortize export map reallocations, so this avoids a quadratic slowdown.
+	TArray<FExportMapEntry> ExportMap;
+	TArray<UObject*> Exports;
+	ForEachObjectWithOuter(Package, [this, &ExportMap, &Exports](UObject* Object)
 	{
+		if (Object->HasAllFlags(RF_Public))
 		{
-			FWriteLockedPackageRef PackageRef = GlobalImportStore.AddPackageRef(PackageId, PackageName, EPackageLoader::Unknown, EPackageExtension::Unspecified);
-
-#if !WITH_EDITOR
-			ensureMsgf(!PackageRef->HasPackage(), TEXT("Compiled Verse package %s has already been added to the loader."), *PackageName.ToString());
-			PackageRef->SetPackage(Package);
-			Package->SetCanBeImportedFlag(true);
-			Package->SetPackageId(PackageId);
-#endif
-			PackageRef->SetAllPublicExportsLoaded(true);
-			Package->SetInternalFlags(EInternalObjectFlags::LoaderImport);
-		}
+			Object->SetInternalFlags(EInternalObjectFlags::LoaderImport);
+			Object->SetFlags(RF_WasLoaded | RF_LoadCompleted);
 
-		ForEachObjectWithOuter(Package, [this, PackageId](UObject* Object)
-		{
-			if (Object->HasAllFlags(RF_Public))
+			TArray<FName, TInlineAllocator<64>> FullPath;
+			FullPath.Add(Object->GetFName());
+			UObject* Outer = Object->GetOuter();
+			while (Outer)
 			{
-				Object->SetInternalFlags(EInternalObjectFlags::LoaderImport);
-				Object->SetFlags(RF_WasLoaded | RF_LoadCompleted);
-
-				TArray<FName, TInlineAllocator<64>> FullPath;
-				FullPath.Add(Object->GetFName());
-				UObject* Outer = Object->GetOuter();
-				while (Outer)
-				{
-					FullPath.Add(Outer->GetFName());
-					Outer = Outer->GetOuter();
-				}
-				TStringBuilder<256> PackageRelativeExportPath;
-				for (int32 PathIndex = FullPath.Num() - 2; PathIndex >= 0; --PathIndex)
+				FullPath.Add(Outer->GetFName());
+				Outer = Outer->GetOuter();
+			}
+			TStringBuilder<256> PackageRelativeExportPath;
+			for (int32 PathIndex = FullPath.Num() - 2; PathIndex >= 0; --PathIndex)
+			{
+				TCHAR NameStr[FName::StringBufferSize];
+				uint32 NameLen = FullPath[PathIndex].ToString(NameStr);
+				for (uint32 I = 0; I < NameLen; ++I)
 				{
-					TCHAR NameStr[FName::StringBufferSize];
-					uint32 NameLen = FullPath[PathIndex].ToString(NameStr);
-					for (uint32 I = 0; I < NameLen; ++I)
-					{
-						NameStr[I] = TChar<TCHAR>::ToLower(NameStr[I]);
-					}
-					PackageRelativeExportPath.AppendChar('/');
-					PackageRelativeExportPath.Append(FStringView(NameStr, NameLen));
+					NameStr[I] = TChar<TCHAR>::ToLower(NameStr[I]);
 				}
-				uint64 ExportHash = CityHash64(reinterpret_cast<const char*>(PackageRelativeExportPath.GetData() + 1), (PackageRelativeExportPath.Len() - 1) * sizeof(TCHAR));
-				GlobalImportStore.StoreGlobalObject(PackageId, ExportHash, Object);
+				PackageRelativeExportPath.AppendChar('/');
+				PackageRelativeExportPath.Append(FStringView(NameStr, NameLen));
 			}
-		}, /* bIncludeNestedObjects*/ true);
-		for (int32 Index = 0; Index < VersePackage->NumDefinitions(); Index++)
+			uint64 ExportHash = CityHash64(reinterpret_cast<const char*>(PackageRelativeExportPath.GetData() + 1), (PackageRelativeExportPath.Len() - 1) * sizeof(TCHAR));
+
+			ExportMap.Add({ .PublicExportHash = ExportHash });
+			Exports.Add(Object);
+		}
+	}, /* bIncludeNestedObjects*/ true);
+
+	{
+		FWriteLockedPackageRef PackageRef = GlobalImportStore.AddPackageRef(PackageId, PackageName, EPackageLoader::Unknown, EPackageExtension::Unspecified);
+		PackageRef->PreInsertPublicExports({ ExportMap });
+
+#if !WITH_EDITOR
+		ensureMsgf(!PackageRef->HasPackage(), TEXT("Compiled Verse package %s has already been added to the loader."), *PackageName.ToString());
+		PackageRef->SetPackage(Package);
+		Package->SetCanBeImportedFlag(true);
+		Package->SetPackageId(PackageId);
+#endif
+		PackageRef->SetAllPublicExportsLoaded(true);
+		Package->SetInternalFlags(EInternalObjectFlags::LoaderImport);
+	}
+	for (int32 Index = 0; Index < Exports.Num(); ++Index)
+	{
+		GlobalImportStore.StoreGlobalObject(PackageId, ExportMap[Index].PublicExportHash, Exports[Index]);
+	}
+	for (int32 Index = 0; Index < VersePackage->NumDefinitions(); Index++)
+	{
+		if (Verse::VCell* Cell = VersePackage->GetDefinition(Index).DynamicCast<Verse::VCell>())
 		{
-			if (Verse::VCell* Cell = VersePackage->GetDefinition(Index).DynamicCast<Verse::VCell>())
-			{
-				FUtf8StringView VersePath = VersePackage->GetDefinitionName(Index).AsStringView();
-				uint64 ExportHash = CityHash64(reinterpret_cast<const char*>(VersePath.GetData()), VersePath.Len());
-				GlobalImportStore.StoreGlobalCell(PackageId, ExportHash, Cell);
-			}
+			FUtf8StringView VersePath = VersePackage->GetDefinitionName(Index).AsStringView();
+			uint64 ExportHash = CityHash64(reinterpret_cast<const char*>(VersePath.GetData()), VersePath.Len());
+			GlobalImportStore.StoreGlobalCell(PackageId, ExportHash, Cell);
 		}
 	}
-
 	GlobalImportStore.ReleasePackageRef(PackageId);
 #endif
 }
