Creating placeholder types that inherit from a known base class.

#rnx
#Synced-CL 49617180

diff --git a/Engine/Source/Runtime/CoreUObject/Internal/UObject/InstanceDataObjectUtils.h b/Engine/Source/Runtime/CoreUObject/Internal/UObject/InstanceDataObjectUtils.h
--- a/Engine/Source/Runtime/CoreUObject/Internal/UObject/InstanceDataObjectUtils.h
+++ b/Engine/Source/Runtime/CoreUObject/Internal/UObject/InstanceDataObjectUtils.h
@@ -20,6 +20,7 @@ class UStruct;
 namespace UE { class FPropertyPathNameTree; }
 namespace UE { class FUnknownEnumNames; }
 namespace UE::Serialization::Private { class FImportTypeHierarchy; }
+namespace UE::Serialization::Private { struct FTypeResource; }
 
 namespace UE
 {
@@ -54,6 +55,9 @@ bool CanCreatePropertyBagPlaceholderTypeForImportClass(const UClass* ClassKind,
 /** Query if placeholder support is enabled for a specific type. */
 bool CanCreatePropertyBagPlaceholderForType(const UClass* ClassKind, const UStruct* Type);
 
+/** Query if a given type can serve as a super-type for a placeholder type. */
+bool CanCreatePropertyBagPlaceholderSuperType(const UE::Serialization::Private::FTypeResource& SuperType);
+
 /** Helper to check if a class is an IDO class. */
 UE_API bool IsClassOfInstanceDataObjectClass(const UStruct* Class);
 
diff --git a/Engine/Source/Runtime/CoreUObject/Private/UObject/InstanceDataObjectUtils.cpp b/Engine/Source/Runtime/CoreUObject/Private/UObject/InstanceDataObjectUtils.cpp
--- a/Engine/Source/Runtime/CoreUObject/Private/UObject/InstanceDataObjectUtils.cpp
+++ b/Engine/Source/Runtime/CoreUObject/Private/UObject/InstanceDataObjectUtils.cpp
@@ -270,10 +270,12 @@ namespace UE
 	static const FName NAME_PresentAsTypeMetadata(ANSITEXTVIEW("PresentAsType"));
 	static const FName NAME_IsInstanceDataObjectStruct(ANSITEXTVIEW("IsInstanceDataObjectClass"));
 	static const FName NAME_ContainsLoosePropertiesMetadata(ANSITEXTVIEW("ContainsLooseProperties"));
+	static const FName NAME_Class(ANSITEXTVIEW("Class"));
 	static const FName NAME_VerseClass(ANSITEXTVIEW("VerseClass"));
 	static const FName NAME_BlueprintGeneratedClass(ANSITEXTVIEW("BlueprintGeneratedClass"));
 	static const FName NAME_VerseDevice(ANSITEXTVIEW("VerseDevice_C"));
 	static const FName NAME_Entity(ANSITEXTVIEW("entity"));
+	static const FName NAME_Component(ANSITEXTVIEW("component"));
 	static const FName NAME_IDOMapKey(ANSITEXTVIEW("Key"));
 	static const FName NAME_IDOMapValue(ANSITEXTVIEW("Value"));
 
@@ -833,6 +835,22 @@ namespace UE
 		return bResult;
 	}
 
+	bool CanCreatePropertyBagPlaceholderSuperType(const UE::Serialization::Private::FTypeResource& SuperType)
+	{
+		bool bResult = false;
+
+		if ((SuperType.TypeName == NAME_Component) && (SuperType.ClassName == NAME_VerseClass))
+		{
+			bResult = true;
+		}
+		else if ((SuperType.TypeName == NAME_Entity) && (SuperType.ClassName == NAME_Class))
+		{
+			bResult = true;
+		}
+
+		return bResult;
+	}
+
 	bool IsClassOfInstanceDataObjectClass(const UStruct* Class)
 	{
 		return Class->IsA(UInstanceDataObjectClass::StaticClass()) || Class->IsA(UInstanceDataObjectStruct::StaticClass());
diff --git a/Engine/Source/Runtime/CoreUObject/Private/UObject/LinkerLoad.cpp b/Engine/Source/Runtime/CoreUObject/Private/UObject/LinkerLoad.cpp
--- a/Engine/Source/Runtime/CoreUObject/Private/UObject/LinkerLoad.cpp
+++ b/Engine/Source/Runtime/CoreUObject/Private/UObject/LinkerLoad.cpp
@@ -4549,6 +4549,42 @@ UClass* FLinkerLoad::GetExportLoadClass(int32 Index)
 }
 
 #if WITH_EDITOR
+static UStruct* FindKnownSuperStruct(const UE::Serialization::Private::FImportTypeHierarchy* ImportTypeHierarchy)
+{
+	UStruct* SuperStruct = nullptr;
+
+	if (ImportTypeHierarchy)
+	{
+		const TArray<UE::Serialization::Private::FTypeResource>& SuperTypes = ImportTypeHierarchy->GetSuperTypes();
+
+		// Types are ordered bottom-up (child classes to parent classes).
+		for (int32 TypeIndex = 0; TypeIndex < SuperTypes.Num(); ++TypeIndex)
+		{
+			const UE::Serialization::Private::FTypeResource& SuperTypeResource = SuperTypes[TypeIndex];
+
+			if (UE::CanCreatePropertyBagPlaceholderSuperType(SuperTypeResource))
+			{
+				if (UPackage* SuperTypePackage = FindObjectFast<UPackage>(nullptr, SuperTypeResource.PackageName, EFindObjectFlags::ExactClass))
+				{
+					SuperStruct = FindObjectFast<UStruct>(SuperTypePackage, SuperTypeResource.TypeName, EFindObjectFlags::None);
+				}
+			}
+
+			if (SuperStruct)
+			{
+				break;
+			}
+		}
+	}
+
+	if (!SuperStruct)
+	{
+		SuperStruct = UObject::StaticClass();
+	}
+
+	return SuperStruct;
+}
+
 UClass* FLinkerLoad::TryCreatePlaceholderClassImport(int32 ImportIndex)
 {
 	const bool bAllowPlaceholderImportTypes = UE::FPropertyBagRepository::IsPropertyBagPlaceholderObjectFeatureEnabled(UE::FPropertyBagRepository::EPlaceholderObjectFeature::ReplaceMissingTypeImportsOnLoad);
@@ -4594,8 +4630,11 @@ UClass* FLinkerLoad::TryCreatePlaceholderClassImport(int32 ImportIndex)
 
 				if (ClassObjectPackage)
 				{
+					UStruct* SuperStruct = FindKnownSuperStruct(ImportTypeHierarchy);
+
 					// Create an opaque, non-native public type object that has no reflected properties.
-					ClassObject = UE::FPropertyBagRepository::CreatePropertyBagPlaceholderClass(ClassObjectPackage, ImportClass, Import.ObjectName, RF_Public);
+					UStruct* StructObject = UE::FPropertyBagRepository::CreatePropertyBagPlaceholderType(ClassObjectPackage, ImportClass, Import.ObjectName, RF_Public, SuperStruct);
+					ClassObject = CastChecked<UClass>(StructObject);
 
 					// Patch it into the import table so that we resolve to this class for any future exports of this type.
 					Import.XObject = ClassObject;
