AssetRegistry: Remove 2GB limit on loading DevelopmentAssetRegistry.bin files that was imposed by AssetDataTagMap.cpp; change to a FMemoryWriter64. Large projects can go over 2GB in size for the DevelopmenAssetRegistry.
#rnx
#Synced-CL 49551581

diff --git a/Engine/Source/Runtime/CoreUObject/Private/AssetRegistry/AssetDataTagMap.cpp b/Engine/Source/Runtime/CoreUObject/Private/AssetRegistry/AssetDataTagMap.cpp
--- a/Engine/Source/Runtime/CoreUObject/Private/AssetRegistry/AssetDataTagMap.cpp
+++ b/Engine/Source/Runtime/CoreUObject/Private/AssetRegistry/AssetDataTagMap.cpp
@@ -1400,13 +1400,14 @@ namespace FixedTagPrivate
 
 	TRefCountPtr<const FStore> InitializeFromBuilder(FStoreBuilder& InBuilder)
 	{
-		TArray<uint8> Data;
-		FMemoryWriter MemoryWriter(Data, true);
+		TArray64<uint8> Data;
+		FMemoryWriter64 MemoryWriter(Data, true);
 		FStoreData StoreData = InBuilder.Finalize();
 
 		SaveStore(StoreData, MemoryWriter);
 
-		FMemoryReader MemoryReader(Data, true);
+		FMemoryView DataView(Data);
+		FMemoryReaderView MemoryReader(DataView, true);
 		return LoadStore(MemoryReader);
 	}
 
