Reintroduce placeholder creation logic removed in 49041576, this allows display of paths to objects external to a project to be displayed in UE_INTERNAL_UNIVERSAL_PLACEHOLDERS configurations

#rnx
#Synced-CL 49566423

diff --git a/Engine/Source/Runtime/CoreUObject/Private/UObject/LinkerLoad.cpp b/Engine/Source/Runtime/CoreUObject/Private/UObject/LinkerLoad.cpp
--- a/Engine/Source/Runtime/CoreUObject/Private/UObject/LinkerLoad.cpp
+++ b/Engine/Source/Runtime/CoreUObject/Private/UObject/LinkerLoad.cpp
@@ -4687,6 +4687,42 @@ void FLinkerLoad::TryCreateImportedPackage(FObjectImport& Import)
 	Import.XObject = Pkg;
 }
 
+void FLinkerLoad::CreatePlaceholderImport(FObjectImport& Import)
+{
+	// Packages and obviously identifiable classes have already been created
+	// by the prepass (CreateUniversalPlaceholders):
+	UPackage* ClassPackage = FindObjectFast<UPackage>( nullptr, Import.ClassPackage, EFindObjectFlags::None );
+	check(ClassPackage);
+	UClass*	FindClass = FindObjectFast<UClass>( ClassPackage, Import.ClassName, EFindObjectFlags::None );
+	check(FindClass);
+	UObject* FindOuter = FindOuter = IndexToObject( Import.OuterIndex );;
+	UObject* FindObject = FindImportFast(FindClass, FindOuter, Import.ObjectName); // import was already created, likely by native code
+	if(FindObject)
+	{
+		return;
+	}
+
+	if(FindClass->IsChildOf(UClass::StaticClass()))
+	{
+		ForceCreatePlaceholderClassImport(Import);
+	}
+	else
+	{
+		FindObject = NewObject<UObject>(FindOuter, FindClass, Import.ObjectName, RF_Public);
+		if(UScriptStruct* Struct = Cast<UScriptStruct>(FindObject))
+		{
+			// add a single property to avoid creating a 0 sized struct. Zero sized allocations
+			// are supported by some systems, but containers are generally unhappy with zero stride
+			// container elements. We can't just use the fallback struct as we want a proper
+			// object name:
+			FBoolProperty* Property = new FBoolProperty(Struct, "__PlaceholderSentinel", RF_NoFlags);
+			Struct->ChildProperties = Property;
+			Struct->StaticLink(true); // make the struct roughly usable, it's just a dummy type
+		}
+		Import.XObject = FindObject;
+	}
+}
+
 void FLinkerLoad::TryCreatePlaceholderClassExport(const FObjectExport& Export)
 {
 	if(!Export.ClassIndex.IsExport())
@@ -4711,11 +4747,11 @@ void FLinkerLoad::TryCreatePlaceholderClassExport(const FObjectExport& Export)
 	// topo sort is needed for exports:
 	check(ExportClass.ClassIndex.IsImport());
 	UClass* const ClassOfClass = ForceCreatePlaceholderClassImport(
-		ExportClass.ClassIndex.ToImport(), UClass::StaticClass());
+		ImportMap[ExportClass.ClassIndex.ToImport()], UClass::StaticClass());
 		
 	check(ExportClass.SuperIndex.IsImport());
 	UClass* const SuperClass = ForceCreatePlaceholderClassImport(
-		ExportClass.SuperIndex.ToImport());
+		ImportMap[ExportClass.SuperIndex.ToImport()]);
 
 	// With required imports resolved, we can now create the export object (which is
 	// a class) and its cdo:
@@ -4783,7 +4819,7 @@ void FLinkerLoad::TryCreateImportedClassForExport(const FObjectExport& Export)
 	// failed to resolve, create the placeholder:
 	if(!Import.XObject)
 	{
-		ForceCreatePlaceholderClassImport(Export.ClassIndex.ToImport());
+		ForceCreatePlaceholderClassImport(ImportMap[Export.ClassIndex.ToImport()]);
 	}
 }
 
@@ -4941,7 +4977,7 @@ void FLinkerLoad::TryCreateImportedClassesForImports()
 		const int32 ImportIndex = ClassToImportIndex.FindChecked(Node);
 		if(ImportIndex != INDEX_NONE)
 		{
-			ForceCreatePlaceholderClassImport(ImportIndex);
+			ForceCreatePlaceholderClassImport(ImportMap[ImportIndex]);
 		}
 		else
 		{
@@ -4966,10 +5002,8 @@ void FLinkerLoad::TryCreateImportedClassesForImports()
 	}
 }
 
-UClass* FLinkerLoad::ForceCreatePlaceholderClassImport(int32 ImportIndex, UClass* SuperClassOverride)
+UClass* FLinkerLoad::ForceCreatePlaceholderClassImport(FObjectImport& Import, UClass* SuperClassOverride)
 {
-	FObjectImport& Import = ImportMap[ImportIndex];
-
 	// Placeholder already created:
 	if (Import.XObject)
 	{
@@ -6502,6 +6536,10 @@ UObject* FLinkerLoad::CreateImport( int32 Index )
 	// Imports can have no name if they were filtered out due to package redirects, skip in that case
 	if (Import.XObject == nullptr && Import.ObjectName != NAME_None)
 	{
+#if UE_INTERNAL_UNIVERSAL_PLACEHOLDERS
+		CreatePlaceholderImport(Import);
+#endif // UE_INTERNAL_UNIVERSAL_PLACEHOLDERS
+
 		if( Import.XObject == NULL )
 		{
 			EVerifyResult VerifyImportResult = VERIFY_Success;
diff --git a/Engine/Source/Runtime/CoreUObject/Public/UObject/LinkerLoad.h b/Engine/Source/Runtime/CoreUObject/Public/UObject/LinkerLoad.h
--- a/Engine/Source/Runtime/CoreUObject/Public/UObject/LinkerLoad.h
+++ b/Engine/Source/Runtime/CoreUObject/Public/UObject/LinkerLoad.h
@@ -924,10 +924,11 @@ private:
 	// package. To understand these functions simply read CreateUniversalPlaceholders
 	void CreateUniversalPlaceholders();
 	void TryCreateImportedPackage(FObjectImport& Import);
+	void CreatePlaceholderImport(FObjectImport& Import);
 	void TryCreatePlaceholderClassExport(const FObjectExport& Export);
 	void TryCreateImportedClassForExport(const FObjectExport& Export);
 	void TryCreateImportedClassesForImports();
-	UClass* ForceCreatePlaceholderClassImport(int32 ImportIndex, UClass* SuperClassOverride = nullptr);
+	UClass* ForceCreatePlaceholderClassImport(FObjectImport& Import, UClass* SuperClassOverride = nullptr);
 	
 	// FArchive bounds implementation - needed for high placeholder environments but
 	// likely useful everywhere that can afford them:
