Rename FRemoteObjectMaps::FindRemoteObjectStub and FindOrAddRemoteObjectStub to FindStub and FindOrAddStub

#Synced-CL 49620103

diff --git a/Engine/Source/Runtime/CoreUObject/Private/UObject/RemoteObject.cpp b/Engine/Source/Runtime/CoreUObject/Private/UObject/RemoteObject.cpp
--- a/Engine/Source/Runtime/CoreUObject/Private/UObject/RemoteObject.cpp
+++ b/Engine/Source/Runtime/CoreUObject/Private/UObject/RemoteObject.cpp
@@ -128,7 +128,7 @@ public:
 		}
 
 	}
-	UE::RemoteObject::Handle::FRemoteObjectStub* FindRemoteObjectStub(FRemoteObjectId Id)
+	UE::RemoteObject::Handle::FRemoteObjectStub* FindStub(FRemoteObjectId Id)
 	{
 		using namespace UE::RemoteObject::Handle;
 
@@ -136,7 +136,7 @@ public:
 		return RemoteObjects.Find(Id);
 	}
 
-	UE::RemoteObject::Handle::FRemoteObjectStub* FindOrAddRemoteObjectStub(FRemoteObjectId Id, FRemoteServerId ResidentServerId = FRemoteServerId())
+	UE::RemoteObject::Handle::FRemoteObjectStub* FindOrAddStub(FRemoteObjectId Id, FRemoteServerId ResidentServerId = FRemoteServerId())
 	{
 		using namespace UE::RemoteObject::Handle;
 
@@ -161,7 +161,7 @@ public:
 		return Stub;
 	}
 
-	UE::RemoteObject::Handle::FRemoteObjectStub* FindOrAddRemoteObjectStub(const UObject* Object, FRemoteServerId ResidentServerId = FRemoteServerId())
+	UE::RemoteObject::Handle::FRemoteObjectStub* FindOrAddStub(const UObject* Object, FRemoteServerId ResidentServerId = FRemoteServerId())
 	{
 		using namespace UE::RemoteObject::Handle;
 		using namespace UE::CoreUObject::Private;
@@ -169,7 +169,7 @@ public:
 		checkf(Object, TEXT("Unable to add a stub for a null object"));
 
 		UE::TScopeLock MapLock(ObjectMapCritical);
-		FRemoteObjectStub* Stub = FindOrAddRemoteObjectStub(FRemoteObjectId(Object), ResidentServerId);
+		FRemoteObjectStub* Stub = FindOrAddStub(FRemoteObjectId(Object), ResidentServerId);
 
 		// Always try and update the stub with the most recent information about the object
 		FUObjectItem* ObjectItem = GUObjectArray.ObjectToObjectItem(Object);
@@ -314,7 +314,7 @@ void ShutdownRemoteObjects()
 
 void RegisterRemoteObjectId(FRemoteObjectId ObjectId, FRemoteServerId ResidentServerId)
 {
-	ObjectMaps->FindOrAddRemoteObjectStub(ObjectId, ResidentServerId);
+	ObjectMaps->FindOrAddStub(ObjectId, ResidentServerId);
 }
 
 void RegisterSharedObject(UObject* Object)
@@ -332,7 +332,7 @@ void MarkAsRemote(UObject* Object, FRemoteServerId DestinationServerId)
 	ObjectItem->SetFlags(EInternalObjectFlags::Remote);
 	ObjectItem->ClearFlags(EInternalObjectFlags_RootFlags | EInternalObjectFlags::RemoteReference | EInternalObjectFlags::Borrowed);
 #endif
-	ObjectMaps->FindOrAddRemoteObjectStub(Object, DestinationServerId);
+	ObjectMaps->FindOrAddStub(Object, DestinationServerId);
 }
 
 void MarkAsRemoteReference(const UObject* Object)
@@ -378,13 +378,13 @@ void MarkAsLocal(UObject* Object)
 	ObjectItem->ClearFlags(EInternalObjectFlags::Remote);
 	ObjectItem->SetFlags(EInternalObjectFlags::RemoteReference);
 #endif
-	ObjectMaps->FindOrAddRemoteObjectStub(Object, FRemoteServerId::GetLocalServerId());
+	ObjectMaps->FindOrAddStub(Object, FRemoteServerId::GetLocalServerId());
 }
 
 void StoreAssetPath(UObject* Object)
 {
 	// Make sure the asset has a stub and that the stub knows the owner if this asset is the asset server (disk / content)
-	ObjectMaps->FindOrAddRemoteObjectStub(Object, FRemoteServerId(ERemoteServerIdConstants::Asset));
+	ObjectMaps->FindOrAddStub(Object, FRemoteServerId(ERemoteServerIdConstants::Asset));
 	ObjectMaps->StoreAssetPath(Object);
 }
 
@@ -395,7 +395,7 @@ FRemoteObjectPathName* FindAssetPath(FRemoteObjectId RemoteId)
 
 UE::RemoteObject::Handle::FRemoteObjectStub* FindRemoteObjectStub(FRemoteObjectId ObjectId)
 {
-	return ObjectMaps ? ObjectMaps->FindRemoteObjectStub(ObjectId) : nullptr;
+	return ObjectMaps ? ObjectMaps->FindStub(ObjectId) : nullptr;
 }
 
 FName GetServerBaseNameForUniqueName(const UClass* Class)
@@ -590,7 +590,7 @@ bool IsValid(FRemoteObjectId ObjectId)
 		return IsValid(Object);
 	}
 
-	if (FRemoteObjectStub* Stub = ObjectMaps->FindRemoteObjectStub(ObjectId))
+	if (FRemoteObjectStub* Stub = ObjectMaps->FindStub(ObjectId))
 	{
 		return !Stub->bWasGarbage;
 	}
@@ -712,7 +712,7 @@ UObject* ResolveObject(UObject* Object, ERemoteReferenceType RefType /*= ERemote
 	// we can just return it and let the owner finish its cleanup.
 	if (IsSafeToMigrateObjects()) // Note: this is not a thread-safe check
 	{
-		FRemoteObjectStub* Stub = ObjectMaps->FindRemoteObjectStub(UE::CoreUObject::Private::FObjectHandleUtils::GetRemoteId(Object));
+		FRemoteObjectStub* Stub = ObjectMaps->FindStub(UE::CoreUObject::Private::FObjectHandleUtils::GetRemoteId(Object));
 		checkf(Stub, TEXT("Failed to find remote object stub for %s"), *GetPathNameSafe(Object));
 		return ResolveObject(Stub, RefType);
 	}
@@ -749,7 +749,7 @@ bool CanResolveObject(FRemoteObjectId ObjectId, ECanResolveObjectBehavior CanRes
 		return false;
 	}
 
-	if (FRemoteObjectStub* Stub = ObjectMaps->FindRemoteObjectStub(ObjectId))
+	if (FRemoteObjectStub* Stub = ObjectMaps->FindStub(ObjectId))
 	{
 		// A stub exists so the object memory is not local but we can (attempt to) migrate it back if:
 		// 1. It's not already Garbage (unless we're allowing a migration of Garbage objects)
@@ -831,7 +831,7 @@ namespace UE::CoreUObject::Private
 		using namespace UE::RemoteObject::Handle;
 		using namespace UE::RemoteObject::Private;		
 
-		FRemoteObjectStub* Stub = ObjectMaps->FindRemoteObjectStub(UE::CoreUObject::Private::FObjectHandleUtils::GetRemoteId(Object));
+		FRemoteObjectStub* Stub = ObjectMaps->FindStub(UE::CoreUObject::Private::FObjectHandleUtils::GetRemoteId(Object));
 		checkf(Stub, TEXT("Failed to find remote object stub for %s"), *GetPathNameSafe(Object));
 		return FRemoteObjectHandlePrivate(Stub);
 	}
@@ -849,7 +849,7 @@ namespace UE::CoreUObject::Private
 			{
 				return FRemoteObjectHandlePrivate(Obj);
 			}
-			else if (const FRemoteObjectStub* const FoundStub = ObjectMaps->FindRemoteObjectStub(ObjectId))
+			else if (const FRemoteObjectStub* const FoundStub = ObjectMaps->FindStub(ObjectId))
 			{
 				return FRemoteObjectHandlePrivate(FoundStub);
 			}
@@ -860,7 +860,7 @@ namespace UE::CoreUObject::Private
 				// Create a null handle instead of using the RemoteId because there's no remote object to find.
 				return FRemoteObjectHandlePrivate(Obj);
 			}
-			else if (FRemoteObjectStub* Stub = ObjectMaps->FindOrAddRemoteObjectStub(ObjectId))
+			else if (FRemoteObjectStub* Stub = ObjectMaps->FindOrAddStub(ObjectId))
 			{
 				return FRemoteObjectHandlePrivate(Stub);
 			}
@@ -1314,7 +1314,7 @@ COREUOBJECT_API UE::RemoteObject::Handle::FRemoteObjectStub* DebugFindRemoteObje
 	}
 
 	FRemoteObjectId ObjId { Object };
-	return UE::RemoteObject::Private::ObjectMaps->FindRemoteObjectStub(ObjId);
+	return UE::RemoteObject::Private::ObjectMaps->FindStub(ObjId);
 }
 
 /**
