Add an ECoreRedirectFlags::Option_Iterate flag that can be used to mark that the result of applying the redirect should be iteratively redirected.
#Synced-CL 49648638

diff --git a/Engine/Config/BaseEngine.ini b/Engine/Config/BaseEngine.ini
--- a/Engine/Config/BaseEngine.ini
+++ b/Engine/Config/BaseEngine.ini
@@ -374,6 +374,13 @@ ClassName=/Script/OnlineSubsystemUtils.OnlineEngineInterfaceImpl
 ;
 ; NOTE, wildcard redirects are much slower than direct matches (non-wildcard matching) and should only be used temporarily! Each wildcard redirect defined
 ; will slow down the CoreRedirect system, regardless of whether they can/will ever match. Use with caution.
+;
+; CoreRedirects can also specify that upon applying a specific redirect, that further redirects must be iteratively applied to the result until a fixpoint is reached:
+;
+; +PackageRedirects=(OldName="/OldPlugin/...", NewName="/NewPlugin/", MatchWildcard=true, Iterate=true)
+; +ClassRedirects=(OldName="/NewPlugin/Module.OldClass", NewName="/NewPlugin/Module.NewClass")
+;
+; NOTE, iterative redirects are expensive when they are applied, and you should prefer to explicitly define chained redirects.
 
 ; 4.10/4.11
 
diff --git a/Engine/Source/Runtime/CoreUObject/Private/UObject/CoreRedirects.cpp b/Engine/Source/Runtime/CoreUObject/Private/UObject/CoreRedirects.cpp
--- a/Engine/Source/Runtime/CoreUObject/Private/UObject/CoreRedirects.cpp
+++ b/Engine/Source/Runtime/CoreUObject/Private/UObject/CoreRedirects.cpp
@@ -1701,63 +1701,92 @@ bool FCoreRedirects::RedirectNameAndValuesUnderReadLock(ECoreRedirectFlags Type,
 	FCoreRedirectObjectName& NewObjectName, const FCoreRedirect** FoundValueRedirect, ECoreRedirectMatchFlags MatchFlags,
 	FScopeCoreRedirectsReadLockedContext& LockedContext)
 {
-	auto ProcessRedirect = [&FoundValueRedirect, OldObjectName](const FCoreRedirect* Redirect, const FCoreRedirectObjectName& NewObjectName)
+	// Iteratively apply redirects until a fixed point or a non-iterative redirect is reached.
+	FCoreRedirectObjectName ObjectName = OldObjectName;
+	for (int NumIterations = 0;;++NumIterations)
 	{
-		if (FoundValueRedirect && (Redirect->HasValueChanges() || Redirect->OverrideClassName.IsValid()))
+		bool bShouldIterate = false;
+		auto ProcessRedirect = [&FoundValueRedirect, OldObjectName, &bShouldIterate](const FCoreRedirect* Redirect, const FCoreRedirectObjectName& NewObjectName)
 		{
-			if (*FoundValueRedirect)
+			if (FoundValueRedirect && (Redirect->HasValueChanges() || Redirect->OverrideClassName.IsValid()))
 			{
-				if ((*FoundValueRedirect)->ValueChanges.OrderIndependentCompareEqual(Redirect->ValueChanges) == false)
+				if (*FoundValueRedirect)
 				{
-					UE_LOG(LogCoreRedirects, Error, TEXT("RedirectNameAndValues(%s) found multiple conflicting value redirects, %s and %s!"),
-						*OldObjectName.ToString(), *(*FoundValueRedirect)->OldName.ToString(), *Redirect->OldName.ToString());
+					if ((*FoundValueRedirect)->ValueChanges.OrderIndependentCompareEqual(Redirect->ValueChanges) == false)
+					{
+						UE_LOG(LogCoreRedirects, Error, TEXT("RedirectNameAndValues(%s) found multiple conflicting value redirects, %s and %s!"),
+							*OldObjectName.ToString(), *(*FoundValueRedirect)->OldName.ToString(), *Redirect->OldName.ToString());
+					}
+				}
+				else
+				{
+					// Set value redirects for processing outside
+					*FoundValueRedirect = Redirect;
 				}
 			}
-			else
+
+			if (!!(Redirect->RedirectFlags & ECoreRedirectFlags::Option_Iterate))
 			{
-				// Set value redirects for processing outside
-				*FoundValueRedirect = Redirect;
+				bShouldIterate = true;
 			}
-		}
 
-		return Redirect->RedirectName(NewObjectName, true /* bIsKnownToMatch */);
-	};
+			return Redirect->RedirectName(NewObjectName, true /* bIsKnownToMatch */);
+		};
 
-	NewObjectName = OldObjectName;
-	TArray<const FCoreRedirect*> FoundRedirects;
-	if (GetMatchingRedirectsUnderReadLock(Type, OldObjectName, FoundRedirects, MatchFlags, LockedContext))
-	{
-		if (FoundRedirects.Num() > 1)
+		NewObjectName = ObjectName;
+		TArray<const FCoreRedirect*> FoundRedirects;
+		if (GetMatchingRedirectsUnderReadLock(Type, ObjectName, FoundRedirects, MatchFlags, LockedContext))
 		{
-			// Sort them based on match
-			FoundRedirects.Sort([&OldObjectName, MatchFlags](const FCoreRedirect& A, const FCoreRedirect& B)
-				{ return A.OldName.MatchScore(OldObjectName, A.RedirectFlags, MatchFlags) > B.OldName.MatchScore(OldObjectName, B.RedirectFlags, MatchFlags); });
-			NewObjectName = ProcessRedirect(FoundRedirects[0], NewObjectName);
-
-			// Apply in order
-			for (int32 i = 1; i < FoundRedirects.Num(); i++)
+			if (FoundRedirects.Num() > 1)
 			{
-				const FCoreRedirect* Redirect = FoundRedirects[i];
+				// Sort them based on match
+				FoundRedirects.Sort([&ObjectName, MatchFlags](const FCoreRedirect& A, const FCoreRedirect& B)
+					{ return A.OldName.MatchScore(ObjectName, A.RedirectFlags, MatchFlags) > B.OldName.MatchScore(ObjectName, B.RedirectFlags, MatchFlags); });
+				NewObjectName = ProcessRedirect(FoundRedirects[0], NewObjectName);
 
-				// Only apply if name match is still valid, if it already renamed part of it it may not apply any more. 
-				// Don't want to check flags as those were checked in the gather step
-				if (Redirect->Matches(NewObjectName, MatchFlags))
+				// Apply in order
+				for (int32 i = 1; i < FoundRedirects.Num(); i++)
 				{
-					NewObjectName = ProcessRedirect(Redirect, NewObjectName);
+					const FCoreRedirect* Redirect = FoundRedirects[i];
+
+					// Only apply if name match is still valid, if it already renamed part of it it may not apply any more. 
+					// Don't want to check flags as those were checked in the gather step
+					if (Redirect->Matches(NewObjectName, MatchFlags))
+					{
+						NewObjectName = ProcessRedirect(Redirect, NewObjectName);
+					}
 				}
 			}
+			else
+			{
+				const FCoreRedirect* Redirect = FoundRedirects[0];
+				NewObjectName = ProcessRedirect(Redirect, NewObjectName);
+			}
+		}
+
+		if (NewObjectName != ObjectName)
+		{
+			UE_CLOG(LockedContext.Get().IsInDebugMode(), LogCoreRedirects, Verbose,
+				TEXT("RedirectNameAndValues(%s) replaced by %s"), *ObjectName.ToString(), *NewObjectName.ToString());
+			ObjectName = NewObjectName;
 		}
 		else
 		{
-			const FCoreRedirect* Redirect = FoundRedirects[0];
-			NewObjectName = ProcessRedirect(Redirect, NewObjectName);
+			bShouldIterate = false;
 		}
-	}
 
-	const bool bDidRedirect = NewObjectName != OldObjectName;
-	UE_CLOG(LockedContext.Get().IsInDebugMode() && bDidRedirect, LogCoreRedirects, Verbose,
-		TEXT("RedirectNameAndValues(%s) replaced by %s"), *OldObjectName.ToString(), *NewObjectName.ToString());
-	return bDidRedirect;
+		// Limit the maximum number of iterations to ensure we don't get stuck in an infinite loop here if somebody accidentally created a redirect cycle.
+		constexpr int MaxIterations = 10;
+		if (!ensureMsgf(NumIterations < MaxIterations, TEXT("RedirectNameAndValues stopping after maximum number of iterations: %s -> %s"), *OldObjectName.ToString(), *NewObjectName.ToString()))
+		{
+			bShouldIterate = false;
+		}
+
+		if (!bShouldIterate)
+		{
+			return NewObjectName != OldObjectName;
+		}
+	};
 }
 
 bool FCoreRedirects::ValidateAssetRedirectsUnderReadLock(FScopeCoreRedirectsReadLockedContext& LockedContext)
@@ -1773,13 +1802,16 @@ bool FCoreRedirects::ValidateAssetRedirectsUnderReadLock(FScopeCoreRedirectsRead
 			{
 				const FCoreRedirectObjectName& SearchName = Redirect.NewName;
 
-				TArray<const FCoreRedirect*> MatchingRedirects;
-				GetMatchingRedirects(ECoreRedirectFlags::Type_Asset, SearchName, MatchingRedirects);
-				for (const FCoreRedirect* MatchingRedirect : MatchingRedirects)
+				if (!(Redirect.RedirectFlags & ECoreRedirectFlags::Option_Iterate))
 				{
-					UE_LOG(LogCoreRedirects, Warning, TEXT("Found redirect from existing redirect. Chained redirects will not be followed. %s --> %s --> %s"),
-						*Redirect.OldName.ToString(), *SearchName.ToString(), *MatchingRedirect->NewName.ToString());
-					bValidationSucceeded = false;
+					TArray<const FCoreRedirect*> MatchingRedirects;
+					GetMatchingRedirects(ECoreRedirectFlags::Type_Asset, SearchName, MatchingRedirects);
+					for (const FCoreRedirect* MatchingRedirect : MatchingRedirects)
+					{
+						UE_LOG(LogCoreRedirects, Warning, TEXT("Found redirect from existing redirect. Chained redirects will not be followed. %s --> %s --> %s"),
+							*Redirect.OldName.ToString(), *SearchName.ToString(), *MatchingRedirect->NewName.ToString());
+						bValidationSucceeded = false;
+					}
 				}
 			}
 		}
@@ -2112,6 +2144,9 @@ bool FCoreRedirects::RunTests()
 	NewRedirects.Emplace(ECoreRedirectFlags::Type_Package | ECoreRedirectFlags::Option_MatchSuffix, TEXT("/混合接尾辞"), TEXT("/mixed_suffix"));
 	NewRedirects.Emplace(ECoreRedirectFlags::Type_Package | ECoreRedirectFlags::Category_Removed, TEXT("/Game/RemovedPackage"), TEXT("/Game/RemovedPackage"));
 	NewRedirects.Emplace(ECoreRedirectFlags::Type_Package | ECoreRedirectFlags::Category_Removed | ECoreRedirectFlags::Option_MissingLoad, TEXT("/Game/MissingLoadPackage"), TEXT("/Game/MissingLoadPackage"));
+	NewRedirects.Emplace(ECoreRedirectFlags::Type_Package | ECoreRedirectFlags::Option_MatchPrefix | ECoreRedirectFlags::Option_Iterate, TEXT("/oldoldprefix"), TEXT("/oldprefix"));
+	NewRedirects.Emplace(ECoreRedirectFlags::Type_Package | ECoreRedirectFlags::Option_MatchPrefix | ECoreRedirectFlags::Option_Iterate, TEXT("/oldoldoldprefix"), TEXT("/oldoldprefix"));
+	NewRedirects.Emplace(ECoreRedirectFlags::Type_Package | ECoreRedirectFlags::Option_MatchPrefix, TEXT("/noiterate_oldprefix"), TEXT("/oldprefix"));
 
 	FCoreRedirect EnumRedirectWithValues(ECoreRedirectFlags::Type_Enum, TEXT("MaskTarget_PhysMesh"), TEXT("/Script/ClothingSystemRuntimeCommon.EWeightMapTargetCommon"));
 	EnumRedirectWithValues.ValueChanges.Add(TEXT("MaskTarget_PhysMesh::MaxDistance"), TEXT("EWeightMapTargetCommon::MaxDistance"));
@@ -2172,6 +2207,10 @@ bool FCoreRedirects::RunTests()
 	Tests.Emplace(TEXT("/old/object.prefix.subobjects.do.not.change"), TEXT("/new/superobject.prefix2.subobjects.do.not.change"), ECoreRedirectFlags::Type_Object);
 	Tests.Emplace(TEXT("/PrefixOverlappingDirectMatch/Path.Remain.Unchanged"), TEXT("/DirectMatchIsPreferred/NewPath.Remain:Unchanged"), ECoreRedirectFlags::Type_Package);
 
+	// Iteration tests
+	Tests.Emplace(TEXT("/oldoldoldprefix_SomeGame/Package.DefaultClass"), TEXT("/newprefix_SomeGame/Package.DefaultClass"), ECoreRedirectFlags::Type_Package);
+	Tests.Emplace(TEXT("/noiterate_oldprefix_SomeGame/Package.DefaultClass"), TEXT("/oldprefix_SomeGame/Package.DefaultClass"), ECoreRedirectFlags::Type_Package);
+
 	// Suffix tests
 	Tests.Emplace(TEXT("/Game/Package/oldsuffix"), TEXT("/Game/Package/newsuffix"), ECoreRedirectFlags::Type_Package);
 	Tests.Emplace(TEXT("/Game/Package/古い接尾辞"), TEXT("/Game/Package/新しい接尾辞"), ECoreRedirectFlags::Type_Package);
@@ -2636,6 +2675,7 @@ bool FCoreRedirects::ReadRedirectsFromIni(const FString& IniName)
 				FString OldName, NewName, OverrideClassName;
 
 				bool bInstanceOnly = false;
+				bool bIterate = false;
 				bool bRemoved = false;
 				bool bMatchSubstring = false;
 				bool bMatchWildcard = false;
@@ -2643,6 +2683,7 @@ bool FCoreRedirects::ReadRedirectsFromIni(const FString& IniName)
 				const FString& ValueString = It.Value().GetValue();
 
 				FParse::Bool(*ValueString, TEXT("InstanceOnly="), bInstanceOnly);
+				FParse::Bool(*ValueString, TEXT("Iterate="), bIterate);
 				FParse::Bool(*ValueString, TEXT("Removed="), bRemoved);
 				FParse::Bool(*ValueString, TEXT("MatchSubstring="), bMatchSubstring);
 				FParse::Bool(*ValueString, TEXT("MatchWildcard="), bMatchWildcard);
@@ -2669,6 +2710,15 @@ bool FCoreRedirects::ReadRedirectsFromIni(const FString& IniName)
 						NewFlags |= ECoreRedirectFlags::Category_Removed;
 					}
 
+					if (bIterate)
+					{
+						NewFlags |= ECoreRedirectFlags::Option_Iterate;
+
+						UE_LOG(LogCoreRedirects, Warning,
+							TEXT("ReadRedirectsFromIni(%s) Redirect %s is iterative. Iterative redirects are expensive, and should be used with caution."),
+							*IniName, *ValueString);
+					}
+
 					if (bMatchWildcard || bMatchSubstring)
 					{
 						UE_CLOG(bMatchSubstring, LogCoreRedirects, Warning, TEXT("ReadRedirectsFromIni(%s) 'MatchSubstring=' is deprecated. "
diff --git a/Engine/Source/Runtime/CoreUObject/Public/UObject/CoreRedirects.h b/Engine/Source/Runtime/CoreUObject/Public/UObject/CoreRedirects.h
--- a/Engine/Source/Runtime/CoreUObject/Public/UObject/CoreRedirects.h
+++ b/Engine/Source/Runtime/CoreUObject/Public/UObject/CoreRedirects.h
@@ -64,6 +64,7 @@ enum class ECoreRedirectFlags : uint32
 	Option_MatchWildcardMask =	Option_MatchSubstring, // Bit mask of all possible wildcards
 
 	Option_MissingLoad =		0x04000000, // An automatically-created redirect that was created in response to a missing Thing during load. Redirect will be removed if and when the Thing is loaded.
+	Option_Iterate =			0x08000000, // Indicates that the result of applying this redirect should be iteratively redirected.
 	Option_AllMask =			0xFF000000, // Bit mask of all possible Options
 };
 ENUM_CLASS_FLAGS(ECoreRedirectFlags);
