When receiving FRemoteObjectReference from another server the current server will only set the resident server id for object ids that have not yet been registered locally. This is because we can't rely on the other server's information about the resident server id for an object as it may be out of date.

Renamed UE::RemoteObject::Private::RegisterRemoteObjectId to FindOrAddRemoteObjectId to better reflect what it's actually doing
Changed the behavior of unused UE::RemoteObject::Transfer::RegisterRemoteObjectId to match the description of this change

#Synced-CL 49620103

diff --git a/Engine/Source/Runtime/CoreUObject/Private/UObject/RemoteObject.cpp b/Engine/Source/Runtime/CoreUObject/Private/UObject/RemoteObject.cpp
--- a/Engine/Source/Runtime/CoreUObject/Private/UObject/RemoteObject.cpp
+++ b/Engine/Source/Runtime/CoreUObject/Private/UObject/RemoteObject.cpp
@@ -312,9 +312,10 @@ void ShutdownRemoteObjects()
 	ObjectMaps = nullptr;
 }
 
-void RegisterRemoteObjectId(FRemoteObjectId ObjectId, FRemoteServerId ResidentServerId)
+UE::RemoteObject::Handle::FRemoteObjectStub* FindOrAddRemoteObjectStub(FRemoteObjectId ObjectId, FRemoteServerId ResidentServerId)
 {
-	ObjectMaps->FindOrAddStub(ObjectId, ResidentServerId);
+	checkf(ObjectMaps, TEXT("FindOrAddRemoteObjectStub(%s, %s): Remote object system is not initialized"), *ObjectId.ToString(), *ResidentServerId.ToString());
+	return ObjectMaps->FindOrAddStub(ObjectId, ResidentServerId);
 }
 
 void RegisterSharedObject(UObject* Object)
diff --git a/Engine/Source/Runtime/CoreUObject/Private/UObject/RemoteObjectPrivate.h b/Engine/Source/Runtime/CoreUObject/Private/UObject/RemoteObjectPrivate.h
--- a/Engine/Source/Runtime/CoreUObject/Private/UObject/RemoteObjectPrivate.h
+++ b/Engine/Source/Runtime/CoreUObject/Private/UObject/RemoteObjectPrivate.h
@@ -55,9 +55,12 @@ namespace UE::RemoteObject::Private
 	bool IsBorrowed(const UObject* Object);
 
 	/**
-	* Registers a stub for a remote object that is known to be resident on a specific server
+	* Finds or adds a stub for a remote object that is known to be resident on a specific server and updates its ResidentServerId
+	* @param ObjectId Remote object id to find or create a stub for
+	* @param ResidentServerId (optional) New resident server id to set for the new or existing stub (can be invalid in which case the resident server id will be set to local for a new stub or will not be updated for an existing stub)
+	* @return Stub for the specified remote object id
 	*/
-	void RegisterRemoteObjectId(FRemoteObjectId ObjectId, FRemoteServerId ResidentServerId);
+	UE::RemoteObject::Handle::FRemoteObjectStub* FindOrAddRemoteObjectStub(FRemoteObjectId ObjectId, FRemoteServerId ResidentServerId = FRemoteServerId());
 
 	/**
 	* Registers object for sharing, marking it as owned by the current server
diff --git a/Engine/Source/Runtime/CoreUObject/Private/UObject/RemoteObjectSerialization.cpp b/Engine/Source/Runtime/CoreUObject/Private/UObject/RemoteObjectSerialization.cpp
--- a/Engine/Source/Runtime/CoreUObject/Private/UObject/RemoteObjectSerialization.cpp
+++ b/Engine/Source/Runtime/CoreUObject/Private/UObject/RemoteObjectSerialization.cpp
@@ -1367,7 +1367,8 @@ int32 DeserializeObjectData(FRemoteObjectData& ObjectData,
 			{
 				if (UObject* Object = OutObjects[ObjectIndex])
 				{
-					UE::RemoteObject::Private::RegisterRemoteObjectId(FRemoteObjectId(Object), LocalServerId);
+					// Call FindOrAddRemoteObjectStub instead of RegisterRemoteObjectId to always update the resident server id
+					UE::RemoteObject::Private::FindOrAddRemoteObjectStub(FRemoteObjectId(Object), LocalServerId);
 				}
 			}
 
diff --git a/Engine/Source/Runtime/CoreUObject/Private/UObject/RemoteObjectTransfer.cpp b/Engine/Source/Runtime/CoreUObject/Private/UObject/RemoteObjectTransfer.cpp
--- a/Engine/Source/Runtime/CoreUObject/Private/UObject/RemoteObjectTransfer.cpp
+++ b/Engine/Source/Runtime/CoreUObject/Private/UObject/RemoteObjectTransfer.cpp
@@ -1149,7 +1149,14 @@ void RegisterSharedObject(UObject* Object)
 
 void RegisterRemoteObjectId(FRemoteObjectId Id, FRemoteServerId ResidentServerId)
 {
-	UE::RemoteObject::Private::RegisterRemoteObjectId(Id, ResidentServerId);
+	using namespace UE::RemoteObject::Private;
+
+	if (!FindRemoteObjectStub(Id))
+	{
+		// The reson why we don't simply always call FindOrAddRemoteObjectStub is that FindOrAddRemoteObjectStub always updates 
+		// the ResidentServerId (even if the stub already exists) and we only want to update the ResidentServerId when the stub is created
+		FindOrAddRemoteObjectStub(Id, ResidentServerId);
+	}
 }
 
 void GetAllBorrowedObjects(TArray<FRemoteObjectId>& OutBorrowedObjectIds)
@@ -1420,7 +1427,7 @@ void FRemoteObjectReference::NetDequantize(FRemoteObjectId InObjectId, FRemoteSe
 	ObjectId = InObjectId;
 	ServerId = InServerId;
 
-	UE::RemoteObject::Private::RegisterRemoteObjectId(ObjectId, ServerId);
+	UE::RemoteObject::Transfer::RegisterRemoteObjectId(ObjectId, ServerId);
 
 	if (InPath.Num())
 	{
@@ -1446,7 +1453,7 @@ bool FRemoteObjectReference::Serialize(FArchive& Ar)
 
 	if (Ar.IsLoading())
 	{
-		UE::RemoteObject::Private::RegisterRemoteObjectId(ObjectId, ServerId);
+		UE::RemoteObject::Transfer::RegisterRemoteObjectId(ObjectId, ServerId);
 	}
 
 	return true;
diff --git a/Engine/Source/Runtime/CoreUObject/Public/UObject/RemoteObjectTransfer.h b/Engine/Source/Runtime/CoreUObject/Public/UObject/RemoteObjectTransfer.h
--- a/Engine/Source/Runtime/CoreUObject/Public/UObject/RemoteObjectTransfer.h
+++ b/Engine/Source/Runtime/CoreUObject/Public/UObject/RemoteObjectTransfer.h
@@ -240,9 +240,9 @@ namespace UE::RemoteObject::Transfer
 	/**
 	* Registers object ID as known to be owned by another server, without migrating it
 	* @param ObjectId Remote object id
-	* @param OwnerServerId Id of a remote server that currently owns the object
+	* @param ResidentServerId (optional) Id of a remote server where the object currently resides. Note that the resident server id will only be updated for objects that have not yet been registered on this server.
 	*/
-	COREUOBJECT_API void RegisterRemoteObjectId(FRemoteObjectId Id, FRemoteServerId ResidentServerId);
+	COREUOBJECT_API void RegisterRemoteObjectId(FRemoteObjectId Id, FRemoteServerId ResidentServerId = FRemoteServerId());
 
 	/**
 	* Returns the list of all object IDs currently borrowed from another server.
