Preventing Blueprint exposed Remote Object structures from crashing the engine when used in a non remote object handle enabled environment

#Synced-CL 49620103

diff --git a/Engine/Source/Runtime/CoreUObject/Private/UObject/RemoteObject.cpp b/Engine/Source/Runtime/CoreUObject/Private/UObject/RemoteObject.cpp
--- a/Engine/Source/Runtime/CoreUObject/Private/UObject/RemoteObject.cpp
+++ b/Engine/Source/Runtime/CoreUObject/Private/UObject/RemoteObject.cpp
@@ -434,6 +434,11 @@ bool IsForceReturnObjectHandles()
 	return GForceReturnObjectHandles > 0;
 }
 
+bool IsRemoteObjectSystemInitialized()
+{
+	return !!ObjectMaps;
+}
+
 FScopedForceReturnObjectHandles::FScopedForceReturnObjectHandles()
 {
 	Private::GForceReturnObjectHandles++;
diff --git a/Engine/Source/Runtime/CoreUObject/Private/UObject/RemoteObjectPrivate.h b/Engine/Source/Runtime/CoreUObject/Private/UObject/RemoteObjectPrivate.h
--- a/Engine/Source/Runtime/CoreUObject/Private/UObject/RemoteObjectPrivate.h
+++ b/Engine/Source/Runtime/CoreUObject/Private/UObject/RemoteObjectPrivate.h
@@ -123,6 +123,23 @@ namespace UE::RemoteObject::Private
 	* Returns true if UE::RemoteObject::Handle::Resolve() functions are forced to return FObjectHandles instead of the actual UObject pointers
 	*/
 	bool IsForceReturnObjectHandles();
+
+	/**
+	* Returns true if the internal structures used by the remote object system are initialized
+	*/
+	bool IsRemoteObjectSystemInitialized();
+
+	/**
+	* Returns true if the remote object handle support is compiled in and internal structures used by the remote object system are initialized
+	*/
+	inline bool IsRemoteObjectSystemCompiledInAndInitialized()
+	{
+		if (FRemoteObjectId::RemoteObjectSupportCompiledIn)
+		{
+			return IsRemoteObjectSystemInitialized();
+		}
+		return false;
+	}
 }
 
 namespace UE::RemoteObject::Transfer::Private
diff --git a/Engine/Source/Runtime/CoreUObject/Private/UObject/RemoteObjectTransfer.cpp b/Engine/Source/Runtime/CoreUObject/Private/UObject/RemoteObjectTransfer.cpp
--- a/Engine/Source/Runtime/CoreUObject/Private/UObject/RemoteObjectTransfer.cpp
+++ b/Engine/Source/Runtime/CoreUObject/Private/UObject/RemoteObjectTransfer.cpp
@@ -1386,13 +1386,14 @@ UObject* FRemoteObjectReference::Resolve() const
 bool FRemoteObjectReference::NetSerialize(FArchive& Ar, class UPackageMap* Map, bool& bOutSuccess)
 {
 	using namespace UE::RemoteObject::Serialization;
+	using namespace UE::RemoteObject::Private;
 
 	// Serialize remote object id
 	Serialize(Ar);
 
 	// Serialize the pathname of the object if possible which will then be used on the receiving end to find assets loaded by both servers in memory	
 	FRemoteObjectPathName PathName;
-	if (Ar.IsSaving() && ObjectId.IsValid())
+	if (Ar.IsSaving() && ObjectId.IsValid() && IsRemoteObjectSystemCompiledInAndInitialized())
 	{		
 		if (UObject* ExistingObject = StaticFindObjectFastInternal(ObjectId))
 		{
@@ -1403,7 +1404,7 @@ bool FRemoteObjectReference::NetSerialize(FArchive& Ar, class UPackageMap* Map,
 
 	Ar << PathName;
 
-	if (Ar.IsLoading() && PathName.Num())
+	if (Ar.IsLoading() && PathName.Num() && IsRemoteObjectSystemCompiledInAndInitialized())
 	{
 		// We currently can't differentiate between assets that exist on both servers which we don't want to migrate
 		// so try to find the object in memory first by its id and if it doesn't exist try to find the object by pathname
@@ -1424,23 +1425,29 @@ bool FRemoteObjectReference::NetSerialize(FArchive& Ar, class UPackageMap* Map,
 
 void FRemoteObjectReference::NetDequantize(FRemoteObjectId InObjectId, FRemoteServerId InServerId, const FRemoteObjectPathName& InPath)
 {
+	using namespace UE::RemoteObject::Transfer;
+	using namespace UE::RemoteObject::Private;
+
 	ObjectId = InObjectId;
 	ServerId = InServerId;
 
-	UE::RemoteObject::Transfer::RegisterRemoteObjectId(ObjectId, ServerId);
-
-	if (InPath.Num())
+	if (IsRemoteObjectSystemCompiledInAndInitialized())
 	{
-		// We currently can't differentiate between assets that exist on both servers which we don't want to migrate
-		// so try to find the object in memory first by its id and if it doesn't exist try to find the object by pathname
-		// and if such object does exist use its remote id to resolve this reference to avoid migrating the asset from another server
-		UObject* ExistingObject = StaticFindObjectFastInternal(ObjectId);
-		if (!ExistingObject)
+		RegisterRemoteObjectId(ObjectId, ServerId);
+
+		if (InPath.Num())
 		{
-			ExistingObject = InPath.Resolve();
-			if (ExistingObject)
+			// We currently can't differentiate between assets that exist on both servers which we don't want to migrate
+			// so try to find the object in memory first by its id and if it doesn't exist try to find the object by pathname
+			// and if such object does exist use its remote id to resolve this reference to avoid migrating the asset from another server
+			UObject* ExistingObject = StaticFindObjectFastInternal(ObjectId);
+			if (!ExistingObject)
 			{
-				ObjectId = FRemoteObjectId(ExistingObject);
+				ExistingObject = InPath.Resolve();
+				if (ExistingObject)
+				{
+					ObjectId = FRemoteObjectId(ExistingObject);
+				}
 			}
 		}
 	}
@@ -1448,12 +1455,15 @@ void FRemoteObjectReference::NetDequantize(FRemoteObjectId InObjectId, FRemoteSe
 
 bool FRemoteObjectReference::Serialize(FArchive& Ar)
 {
+	using namespace UE::RemoteObject::Transfer;
+	using namespace UE::RemoteObject::Private;
+
 	Ar << ObjectId;
 	Ar << ServerId;
 
-	if (Ar.IsLoading())
+	if (Ar.IsLoading() && IsRemoteObjectSystemCompiledInAndInitialized())
 	{
-		UE::RemoteObject::Transfer::RegisterRemoteObjectId(ObjectId, ServerId);
+		RegisterRemoteObjectId(ObjectId, ServerId);
 	}
 
 	return true;
