[UObject][AsyncLoading]

- When renaming an object, the object will be detached from it's linker (if any). Previously we only cleared the linker for the object when renaming the object to a new package (this still happens), and before that would always called ResetLoaders (this never happens now) which was too heavy weight and resulted in hard to debug issues.
- When a loaded object is renamed, a new CoreUObject delegate OnLoadedObjectRenamed is invoked to allow informing listeners (in this case the loader) that the loaded package is no longer fully loaded and requires work when the package is requested for load again (Importantly, this works even before a GC has had a chance to run)
- The loader now maintains a queue of objects that have been renamed and will flush that queue when the ALT thread runs. Renamed objects will remove the object from the ImportStore implicitly marking its owning package as bAreAllPublicExports = false allowing the package to properly reload when requested.
- In order to allow for the same level of validation for releasing references on FLoadedPackageRefs, a new bWasFullyLoaded flag has been added so we can ensure when releasing packages we are not incorrectly working on an object that hasn't been loaded yet, while still allowing packages to become partially loaded later.

#rnx
#Synced-CL 49684429

diff --git a/Engine/Source/Runtime/CoreUObject/Private/Serialization/AsyncLoading2.cpp b/Engine/Source/Runtime/CoreUObject/Private/Serialization/AsyncLoading2.cpp
--- a/Engine/Source/Runtime/CoreUObject/Private/Serialization/AsyncLoading2.cpp
+++ b/Engine/Source/Runtime/CoreUObject/Private/Serialization/AsyncLoading2.cpp
@@ -741,6 +741,12 @@ struct FPackageReferencer
 	}
 };
 
+struct FLoadedObjectRenamed
+{
+	int32 ObjectIndex;
+	UObject* Object;
+};
+
 struct FPackageRequest
 {
 	int32 RequestId = -1;
@@ -1419,7 +1425,8 @@ private:
 	FAsyncLoadingMutex Mutex;
 	EPackageExtension PackageHeaderExtension = EPackageExtension::Unspecified;
 	EPackageLoader PackageLoader = EPackageLoader::Unknown;
-	uint8 bAreAllPublicExportsLoaded : 1 = false;
+	uint8 bWasFullyLoaded : 1 = false;				// Was the package loading ever completed
+	uint8 bAreAllPublicExportsLoaded : 1 = false;	// Are all loaded objects still valid
 	uint8 bIsNotInstalled : 1 = false;
 	uint8 bIsMissing : 1 = false;
 	uint8 bHasFailed : 1 = false;
@@ -1435,6 +1442,7 @@ private:
 		ExportCount = -1;
 		PackageHeaderExtension = EPackageExtension::Unspecified;
 		PackageLoader = EPackageLoader::Unknown;
+		bWasFullyLoaded = false;
 		bAreAllPublicExportsLoaded = false;
 		bIsNotInstalled = false;
 		bIsMissing = false;
@@ -1549,6 +1557,7 @@ public:
 	inline void SetPackage(UPackage* InPackage)
 	{
 		check(Mutex.IsLocked());
+		check(!bWasFullyLoaded);
 		check(!bAreAllPublicExportsLoaded);
 		check(!bIsMissing);
 		check(!bIsNotInstalled);
@@ -1577,6 +1586,7 @@ public:
 		check(Mutex.IsLocked());
 		// keep RefCount and PublicExportMap while resetting all other state,
 		// the public exports will be replaced one by one from StoreGlobalObject
+		bWasFullyLoaded = false;
 		bAreAllPublicExportsLoaded = false;
 		bIsMissing = false;
 		bHasFailed = false;
@@ -1599,6 +1609,7 @@ public:
 		check(!bHasFailed);
 		check(HasPackage());
 		bIsMissing = false;
+		bWasFullyLoaded = true;
 		bAreAllPublicExportsLoaded = true;
 		bHasBeenLoadedDebug = true;
 #if WITH_EDITOR
@@ -1612,18 +1623,22 @@ public:
 	inline void SetIsMissingPackage()
 	{
 		check(Mutex.IsLocked());
+		check(!bWasFullyLoaded);
 		check(!bAreAllPublicExportsLoaded);
 		check(!HasPackage());
 		bIsMissing = true;
+		bWasFullyLoaded = false;
 		bAreAllPublicExportsLoaded = false;
 	}
 
 	inline void SetIsNotInstalledPackage()
 	{
 		check(Mutex.IsLocked());
+		check(!bWasFullyLoaded);
 		check(!bAreAllPublicExportsLoaded);
 		check(!HasPackage());
 		bIsNotInstalled = true;
+		bWasFullyLoaded = false;
 		bAreAllPublicExportsLoaded = false;
 	}
 
@@ -2223,7 +2238,7 @@ public:
 		PackageRef->ReleasePackageRef();
 
 #if DO_CHECK
-		ensureMsgf(!PackageRef->bHasBeenLoadedDebug || PackageRef->bAreAllPublicExportsLoaded || PackageRef->bIsMissing || PackageRef->bIsNotInstalled || PackageRef->bHasFailed,
+		ensureMsgf(!PackageRef->bHasBeenLoadedDebug || PackageRef->bWasFullyLoaded || PackageRef->bIsMissing || PackageRef->bIsNotInstalled || PackageRef->bHasFailed,
 			TEXT("LoadedPackageRef from None (0x%s) to %s (0x%s) should not have been released when the package is not complete.")
 			TEXT("RefCount=%d, AreAllExportsLoaded=%d, IsMissing=%d, IsNotInstalled=%d, HasFailed=%d, HasBeenLoaded=%d"),
 			*LexToString(FromPackageId),
@@ -2236,13 +2251,13 @@ public:
 			PackageRef->bHasFailed,
 			PackageRef->bHasBeenLoadedDebug);
 
-		if (PackageRef->bAreAllPublicExportsLoaded)
+		if (PackageRef->bWasFullyLoaded)
 		{
 			check(!PackageRef->bIsMissing && !PackageRef->bIsNotInstalled);
 		}
 		if (PackageRef->bIsMissing || PackageRef->bIsNotInstalled)
 		{
-			check(!PackageRef->bAreAllPublicExportsLoaded);
+			check(!PackageRef->bWasFullyLoaded);
 		}
 #endif
 		// is this the last reference to a loaded package?
@@ -2276,6 +2291,34 @@ public:
 		);
 	}
 
+	bool RemoveDeferredPublicExport(int32 ObjectIndex, UObject* Object)
+	{
+		check(ObjectIndex >= 0);
+
+		FPublicExportKey PublicExportKey;
+		if (ObjectIndexToPublicExport.RemoveAndCopyValue(ObjectIndex, PublicExportKey))
+		{
+			// We are dealing with deferred object removal/renames. In which case we can't be sure the passed in Object
+			// is still valid, so lookup the object via index and only remove if we have a match. If we have a mismatch
+			// the object was already collected and would have been removed during collection.
+			UObject* FoundObject = static_cast<UObject*>(GUObjectArray.IndexToObject(ObjectIndex)->GetObject());
+			if(FoundObject == Object)
+			{
+				FPackageId PackageId = PublicExportKey.GetPackageId();
+				UE_ASYNC_PACKAGEID_DEBUG(PackageId);
+				FWriteLockedPackageRef PackageRef = FindPackageRefForWrite(PackageId);
+				PackageRef->RemovePublicExport(PublicExportKey.GetExportHash(), Object->GetFName());
+
+				// Allow this object to now be collectable
+				checkf(Object, TEXT("We expect renamed loaded objects to remain alive until we remove the LoaderImport flag."));
+				Object->AtomicallyClearInternalFlags(EInternalObjectFlags::LoaderImport);
+			}
+
+			return true;
+		}
+		return false;
+	}
+
 	void CleanupPublicExportMap(FWriteLockedPackageRef& PackageRef)
 	{
 		for (auto It = PackageRef->PublicExportMap.CreateValueIterator(); It; ++It)
@@ -4767,6 +4810,9 @@ private:
 	FCriticalSection UnreachableObjectsCritical;
 	FUnreachableObjects UnreachableObjects;
 
+	/** [ASYNC/GAME THREAD] Object Renames can occur on any thread, will be flushed by the ALT exclusively. */
+	TMpscQueue<FLoadedObjectRenamed> DeferredLoadedObjectRenameQueue;
+
 	TUniquePtr<FAsyncLoadingThreadState2> GameThreadState;
 	TUniquePtr<FAsyncLoadingThreadState2> AsyncLoadingThreadState;
 
@@ -5242,6 +5288,11 @@ public:
 private:
 	void OnGarbageCollectStarted();
 
+#if WITH_EDITOR
+	void OnLoadedObjectRenamed(UObject* Obj, UObject* OldOuter, FName OldName);
+#endif
+	void FlushDeferredLoadedObjectRenameQueue();
+
 	void CollectUnreachableObjects(TArrayView<FUObjectItem*> UnreachableObjectItems, FUnreachableObjects& OutUnreachableObjects);
 
 	void RemoveUnreachableObjects(FUnreachableObjects& ObjectsToRemove);
@@ -10880,6 +10931,7 @@ EAsyncPackageState::Type FAsyncLoadingThread2::TickAsyncLoadingFromGameThread(FA
 
 		if (!bIsMultithreaded)
 		{
+			FlushDeferredLoadedObjectRenameQueue();
 			RemoveUnreachableObjects(UnreachableObjects);
 		}
 		UpdateSyncLoadContext(ThreadState);
@@ -10950,6 +11002,10 @@ FAsyncLoadingThread2::FAsyncLoadingThread2(FIoDispatcher& InIoDispatcher)
 
 	FCoreUObjectDelegates::GetGarbageCollectStartedDelegate().AddRaw(this, &FAsyncLoadingThread2::OnGarbageCollectStarted);
 
+#if WITH_EDITOR
+	FCoreUObjectDelegates::OnLoadedObjectRenamed.AddRaw(this, &FAsyncLoadingThread2::OnLoadedObjectRenamed);
+#endif
+
 	FAsyncLoadingThreadState2::TlsSlot = FPlatformTLS::AllocTlsSlot();
 	GameThreadState = MakeUnique<FAsyncLoadingThreadState2>();
 	GameThreadState->bCanAccessAsyncLoadingThreadData = true;
@@ -11145,6 +11201,8 @@ uint32 FAsyncLoadingThread2::Run()
 					};
 #endif
 
+					FlushDeferredLoadedObjectRenameQueue();
+
 					{
 						FScopeLock UnreachableObjectsLock(&UnreachableObjectsCritical);
 						RemoveUnreachableObjects(UnreachableObjects);
@@ -11757,6 +11815,35 @@ void FAsyncLoadingThread2::NotifyRegistrationComplete()
 		GlobalImportStore.GetStoredScriptObjectsCount(), (float)GlobalImportStore.GetStoredScriptObjectsAllocatedSize() / 1024.f);
 }
 
+#if WITH_EDITOR
+void FAsyncLoadingThread2::OnLoadedObjectRenamed(UObject* Obj, UObject* OldOuter, FName OldName)
+{
+	const int32 Index = GUObjectArray.ObjectToIndex(Obj);
+	DeferredLoadedObjectRenameQueue.Enqueue(Index, Obj);
+}
+#endif
+
+void FAsyncLoadingThread2::FlushDeferredLoadedObjectRenameQueue()
+{
+#if WITH_EDITOR
+	TRACE_CPUPROFILER_EVENT_SCOPE(FAsyncLoadingThread2::FlushDeferredLoadedObjectRenameQueue);
+	FLoadedObjectRenamed RenamedObject;
+	while (DeferredLoadedObjectRenameQueue.Dequeue(RenamedObject))
+	{
+		// Attempt to remove this object as a public export (it might not be). 
+		// 
+		// If we just renamed a loaded object to no longer match the public export, we want the export removed.
+		// 
+		// If we just renamed a new object to have the same name as a previously removed public export, do nothing. Why?
+		// We can't rename new object instances on top of public export names without having first renamed/removed them. So if we 
+		// are in this call with an object that was renamed to be "back" in the export table we must have already marked the package 
+		// as not fully loaded. Thus, any load of the package will now either reload the removed export from disk or find the newly renamed 
+		// object instance in memory (if it is indeed an object being renamed back into the export table) and use that.
+		GlobalImportStore.RemoveDeferredPublicExport(RenamedObject.ObjectIndex, RenamedObject.Object);
+	}
+#endif
+}
+
 /*-----------------------------------------------------------------------------
 	FAsyncPackage implementation.
 -----------------------------------------------------------------------------*/
diff --git a/Engine/Source/Runtime/CoreUObject/Private/UObject/Obj.cpp b/Engine/Source/Runtime/CoreUObject/Private/UObject/Obj.cpp
--- a/Engine/Source/Runtime/CoreUObject/Private/UObject/Obj.cpp
+++ b/Engine/Source/Runtime/CoreUObject/Private/UObject/Obj.cpp
@@ -252,7 +252,7 @@ bool UObject::Rename(const TCHAR* InName, UObject* NewOuter, ERenameFlags Flags)
 	// the metadata will be moved with it.
 	FMetaDataUtilities::FMoveMetadataHelperContext MoveMetaData(this, true);
 #endif // WITH_METADATA
-
+	const bool bIsTest = Flags & REN_Test;
 	if (NewOuter && NewOuter != GetOuter())
 	{
 		// Renaming the CDO of a Blueprint is a special case so we do not validate what would otherwise be incorrect use of Rename.
@@ -263,7 +263,7 @@ bool UObject::Rename(const TCHAR* InName, UObject* NewOuter, ERenameFlags Flags)
 			// Check that we are not renaming a within object into an Outer of the wrong type
 			if (!NewOuter->IsA(GetClass()->ClassWithin))
 			{
-				if (Flags & REN_Test)
+				if (bIsTest)
 				{
 					return false;
 				}
@@ -286,6 +286,11 @@ bool UObject::Rename(const TCHAR* InName, UObject* NewOuter, ERenameFlags Flags)
 			{
 				if (Linker && HasAnyFlags(RF_NeedLoad | RF_NeedPostLoad))
 				{
+					if (bIsTest)
+					{
+						return false;
+					}
+
 					const bool bIsPreloading = HasAnyFlags(RF_NeedLoad);
 					const TCHAR* PreloadingGuidance =
 						TEXT("Ensure the target package has completed loading before renaming to correct this error. ") 
@@ -311,6 +316,20 @@ bool UObject::Rename(const TCHAR* InName, UObject* NewOuter, ERenameFlags Flags)
 			}
 		}
 	}
+	// If we are renaming a loaded object while keeping its outer and won't be creating a redirector, we need to detach from the linker if the user
+	// hasn't specifically requested to allow for the linker mismatch. This way, we can reload the renamed object if the package is rerequested
+	// Users who don't want such reload behaviour for the export should be calling InvalidateExport() on the object before renaming.
+	else if (!bIsTest && ((Flags & (REN_AllowPackageLinkerMismatch | REN_DontCreateRedirectors)) == REN_DontCreateRedirectors) && HasAnyFlags(RF_WasLoaded))
+	{
+		if (FLinkerLoad* Linker = GetLinker())
+		{
+			const int32 LinkerIndex = GetLinkerIndex();
+			if (LinkerIndex != INDEX_NONE)
+			{
+				Linker->DetachExport(LinkerIndex);
+			}
+		}
+	}
 
 	// find an object with the same name and same class in the new outer
 	bool bIsCaseOnlyChange = false;
@@ -341,7 +360,7 @@ bool UObject::Rename(const TCHAR* InName, UObject* NewOuter, ERenameFlags Flags)
 		}
 		else if (ExistingObject)
 		{
-			if (Flags & REN_Test)
+			if (bIsTest)
 			{
 				return false;
 			}
@@ -353,7 +372,7 @@ bool UObject::Rename(const TCHAR* InName, UObject* NewOuter, ERenameFlags Flags)
 	}
 
 	// if we are just testing, and there was no conflict, then return a success
-	if (Flags & REN_Test)
+	if (bIsTest)
 	{
 		return true;
 	}
@@ -896,6 +915,10 @@ bool UObject::IsSelectedInEditor() const
 void UObject::PostRename(UObject* OldOuter, const FName OldName)
 {
 #if WITH_EDITOR
+	if (HasAnyFlags(RF_WasLoaded))
+	{
+		FCoreUObjectDelegates::OnLoadedObjectRenamed.Broadcast(this, OldOuter, OldName);
+	}
 	FCoreUObjectDelegates::OnObjectRenamed.Broadcast(this, OldOuter, OldName);
 #endif
 }
diff --git a/Engine/Source/Runtime/CoreUObject/Private/UObject/UObjectGlobals.cpp b/Engine/Source/Runtime/CoreUObject/Private/UObject/UObjectGlobals.cpp
--- a/Engine/Source/Runtime/CoreUObject/Private/UObject/UObjectGlobals.cpp
+++ b/Engine/Source/Runtime/CoreUObject/Private/UObject/UObjectGlobals.cpp
@@ -198,6 +198,7 @@ FCoreUObjectDelegates::FOnObjectPropertyChanged FCoreUObjectDelegates::OnObjectP
 TSet<UObject*> FCoreUObjectDelegates::ObjectsModifiedThisFrame;
 FCoreUObjectDelegates::FOnObjectModified FCoreUObjectDelegates::OnObjectModified;
 FCoreUObjectDelegates::FOnObjectRenamed FCoreUObjectDelegates::OnObjectRenamed;
+FCoreUObjectDelegates::FOnObjectRenamed FCoreUObjectDelegates::OnLoadedObjectRenamed;
 FCoreUObjectDelegates::FOnObjectTransacted FCoreUObjectDelegates::OnObjectTransacted;
 FCoreUObjectDelegates::FOnObjectsReplaced FCoreUObjectDelegates::OnObjectsReplaced;
 FCoreUObjectDelegates::FOnObjectsReinstanced FCoreUObjectDelegates::OnObjectsReinstanced;
diff --git a/Engine/Source/Runtime/CoreUObject/Public/UObject/UObjectGlobals.h b/Engine/Source/Runtime/CoreUObject/Public/UObject/UObjectGlobals.h
--- a/Engine/Source/Runtime/CoreUObject/Public/UObject/UObjectGlobals.h
+++ b/Engine/Source/Runtime/CoreUObject/Public/UObject/UObjectGlobals.h
@@ -3300,10 +3300,13 @@ struct FCoreUObjectDelegates
 		}
 	}
 
-	/** Called when an object is renamed this gets called in both the editor and standalone game editor builds, for every object renamed */
 	DECLARE_TS_MULTICAST_DELEGATE_ThreeParams(FOnObjectRenamed, UObject*, UObject*, FName);
+	/** Called when an object is renamed. This gets called in both the editor and standalone game editor builds, for every object renamed */
 	static COREUOBJECT_API FOnObjectRenamed OnObjectRenamed;
 
+	/** Called when an object that was loaded is renamed. This gets called in both the editor and standalone game editor builds, for every renamed loaded object */
+	static COREUOBJECT_API FOnObjectRenamed OnLoadedObjectRenamed;
+
 	/** Broadcast OnObjectRenamed event */
 	static void BroadcastOnObjectRenamed(UObject* Object, UObject* OldOuter, FName OldName)
 	{
diff --git a/Engine/Source/Runtime/Engine/Private/Tests/Loading/AsyncLoadingTests_Renames.cpp b/Engine/Source/Runtime/Engine/Private/Tests/Loading/AsyncLoadingTests_Renames.cpp
--- a/Engine/Source/Runtime/Engine/Private/Tests/Loading/AsyncLoadingTests_Renames.cpp
+++ b/Engine/Source/Runtime/Engine/Private/Tests/Loading/AsyncLoadingTests_Renames.cpp
@@ -132,4 +132,91 @@ bool FLoadingTests_Renames_ObjectMovedToAnotherPackage_OnDiskReload::RunTest(con
 	return true;
 }
 
+/**
+*/
+IMPLEMENT_SIMPLE_AUTOMATION_TEST(FLoadingTests_Renames_ImportReplacedObject, TEXT("System.Engine.Loading.Renames.ImportReplacedObject"), EAutomationTestFlags::EditorContext | EAutomationTestFlags::EngineFilter)
+bool FLoadingTests_Renames_ImportReplacedObject::RunTest(const FString& Parameters)
+{
+	// Simulate Parent (Object1) <- Child (Object2) relationship 
+	auto MutateObjects =
+		[](FLoadingTestsScope& Scope)
+		{
+			Scope.Object2->HardReference = Scope.Object1;
+		};
+
+	FLoadingTestsScope LoadingTestScope(this, MutateObjects);
+
+	// Load Parent package
+	UPackage* ParentPackage = LoadPackage(nullptr, FLoadingTestsScope::PackagePath1, LOAD_None);
+
+	// Rename one of the objects so the package is now effectively missing a public export
+	UAsyncLoadingTests_Shared* Parent = FindObject<UAsyncLoadingTests_Shared>(nullptr, FLoadingTestsScope::ObjectPath1);
+	TestNotNull(TEXT("Expected to be able to load our Parent object"), Parent);
+
+	FName OriginalName = Parent->GetFName();
+	Parent->Rename(TEXT("TRASH_TestObject"), nullptr, REN_DontCreateRedirectors);
+	Parent->ClearFlags(RF_Standalone);
+	Parent->MarkAsGarbage();
+
+	// Create a new object on top of where the missing public export should be
+	UAsyncLoadingTests_Shared* ReplacementObject = NewObject<UAsyncLoadingTests_Shared>(ParentPackage, OriginalName);
+	UAsyncLoadingTests_Shared* NewParent = FindObject<UAsyncLoadingTests_Shared>(nullptr, FLoadingTestsScope::ObjectPath1);
+	TestEqual(TEXT("Expect new object named ontop of orignal object name to be the same instance"), ReplacementObject, NewParent);
+
+	// Load Child Package which imports the Parent package
+	LoadPackage(nullptr, FLoadingTestsScope::PackagePath2, LOAD_None);
+
+	UAsyncLoadingTests_Shared* Child = FindObject<UAsyncLoadingTests_Shared>(nullptr, FLoadingTestsScope::ObjectPath2);
+	TestNotNull(TEXT("Expected to be able to load our Child object"), Child);
+	TestEqual(TEXT("Expected Child package to import the new renamed Parent object, not the trashed object or a new object create from disk."), Cast<UAsyncLoadingTests_Shared>(Child->HardReference), NewParent);
+
+	ReplacementObject->ClearFlags(RF_Standalone);
+	ReplacementObject->MarkAsGarbage();
+
+	CollectGarbage(GARBAGE_COLLECTION_KEEPFLAGS);
+	return true;
+}
+
+IMPLEMENT_SIMPLE_AUTOMATION_TEST(FLoadingTests_Renames_ImportRenamedObjectReloadsImport, TEXT("System.Engine.Loading.Renames.ImportRenamedObjectReloadsImport"), EAutomationTestFlags::EditorContext | EAutomationTestFlags::EngineFilter)
+bool FLoadingTests_Renames_ImportRenamedObjectReloadsImport::RunTest(const FString& Parameters)
+{
+	// Simulate Parent (Object1) <- Child (Object2) relationship 
+	auto MutateObjects =
+		[](FLoadingTestsScope& Scope)
+		{
+			Scope.Object2->HardReference = Scope.Object1;
+		};
+
+	FLoadingTestsScope LoadingTestScope(this, MutateObjects);
+
+	// Load Parent package
+	UPackage* ParentPackage = LoadPackage(nullptr, FLoadingTestsScope::PackagePath1, LOAD_None);
+
+	// Rename one of the objects so the package is now effectively missing a public export
+	UAsyncLoadingTests_Shared* Parent = FindObject<UAsyncLoadingTests_Shared>(nullptr, FLoadingTestsScope::ObjectPath1);
+	TestNotNull(TEXT("Expected to be able to load our Parent object"), Parent);
+
+	FName OriginalName = Parent->GetFName();
+	Parent->Rename(TEXT("TRASH_TestObject"), nullptr, REN_DontCreateRedirectors);
+	Parent->ClearFlags(RF_Standalone);
+	Parent->MarkAsGarbage();
+
+	UAsyncLoadingTests_Shared* ReloadedParent = FindObject<UAsyncLoadingTests_Shared>(nullptr, FLoadingTestsScope::ObjectPath1);
+	TestNull(TEXT("Expected to be able Parent to no longer exist at the name FLoadingTestsScope::ObjectPath1"), ReloadedParent);
+
+	// Load Child Package which imports the Parent package
+	LoadPackage(nullptr, FLoadingTestsScope::PackagePath2, LOAD_None);
+
+	UAsyncLoadingTests_Shared* Child = FindObject<UAsyncLoadingTests_Shared>(nullptr, FLoadingTestsScope::ObjectPath2);
+	TestNotNull(TEXT("Expected to be able to load our Child object"), Child);
+
+	ReloadedParent = FindObject<UAsyncLoadingTests_Shared>(nullptr, FLoadingTestsScope::ObjectPath1);
+	TestNotNull(TEXT("Expected to be able to load our ReloadedParent object"), ReloadedParent);
+
+	TestEqual(TEXT("Expected Child package to import the renamed Parent object"), Cast<UAsyncLoadingTests_Shared>(Child->HardReference), ReloadedParent);
+	TestNotEqual(TEXT("Expected the reloaded parent to be a new instance and not the trashed parent object"), Parent, ReloadedParent);
+
+	return true;
+}
+
 #endif // WITH_DEV_AUTOMATION_TESTS
