Fixed FPackedObjectRef::GetTypeHash failing to compile if there's another GetTypeHash overload in the outer namespace(s), which can happen in some unity build cases.

diff --git a/Engine/Source/Runtime/CoreUObject/Public/UObject/PackedObjectRef.h b/Engine/Source/Runtime/CoreUObject/Public/UObject/PackedObjectRef.h
--- a/Engine/Source/Runtime/CoreUObject/Public/UObject/PackedObjectRef.h
+++ b/Engine/Source/Runtime/CoreUObject/Public/UObject/PackedObjectRef.h
@@ -38,7 +38,7 @@ namespace UE::CoreUObject::Private
 
 		friend inline uint32 GetTypeHash(FPackedObjectRef ObjectRef)
 		{
-			return GetTypeHash(ObjectRef.EncodedRef);
+			return GetTypeHashHelper(ObjectRef.EncodedRef);
 		}
 
 		inline bool IsNull() { return EncodedRef == 0; }
