Remove the last-accessed cache from FUObjectAnnotationSparse, due to changes over the years it was almost always slower and had complicated registration logic.
Add FUObjectAnnotationTemp which has the same API and a last-access cache, but without the slow thread safety checks.
Optimize DuplicateObject by using FUObjectAnnotationTemp (with a typedef), which is faster than FUObjectAnnotationSparse-with-cache was for this use case.

#Synced-CL 49736630

diff --git a/Engine/Source/Runtime/CoreUObject/Private/Serialization/DuplicateDataReader.cpp b/Engine/Source/Runtime/CoreUObject/Private/Serialization/DuplicateDataReader.cpp
--- a/Engine/Source/Runtime/CoreUObject/Private/Serialization/DuplicateDataReader.cpp
+++ b/Engine/Source/Runtime/CoreUObject/Private/Serialization/DuplicateDataReader.cpp
@@ -20,7 +20,7 @@
  * @param	InDuplicatedObjects		map of original object to copy of that object
  * @param	InObjectData			object data to read from
  */
-FDuplicateDataReader::FDuplicateDataReader( class FUObjectAnnotationSparse<FDuplicatedObject,false>& InDuplicatedObjects ,const FLargeMemoryData& InObjectData, uint32 InPortFlags, UObject* InDestOuter )
+FDuplicateDataReader::FDuplicateDataReader(FDuplicatedObjectAnnotation& InDuplicatedObjects ,const FLargeMemoryData& InObjectData, uint32 InPortFlags, UObject* InDestOuter )
 	: DuplicatedObjectAnnotation(InDuplicatedObjects)
 	, ObjectData(InObjectData)
 	, Offset(0)
diff --git a/Engine/Source/Runtime/CoreUObject/Private/Serialization/DuplicateDataWriter.cpp b/Engine/Source/Runtime/CoreUObject/Private/Serialization/DuplicateDataWriter.cpp
--- a/Engine/Source/Runtime/CoreUObject/Private/Serialization/DuplicateDataWriter.cpp
+++ b/Engine/Source/Runtime/CoreUObject/Private/Serialization/DuplicateDataWriter.cpp
@@ -24,7 +24,7 @@
  * @param	InApplyFlags			the flags that should always be set on the duplicated objects (regardless of whether they're set on the source)
  * @param	InInstanceGraph			the instancing graph to use when creating the duplicate objects.
  */
-FDuplicateDataWriter::FDuplicateDataWriter( FUObjectAnnotationSparse<FDuplicatedObject,false>& InDuplicatedObjects, FLargeMemoryData& InObjectData, UObject* SourceObject,
+FDuplicateDataWriter::FDuplicateDataWriter(FDuplicatedObjectAnnotation& InDuplicatedObjects, FLargeMemoryData& InObjectData, UObject* SourceObject,
 	UObject* DestObject, EObjectFlags InFlagMask, EObjectFlags InApplyFlags, EInternalObjectFlags InInternalFlagMask, EInternalObjectFlags InApplyInternalFlags, FObjectInstancingGraph* InInstanceGraph, uint32 InPortFlags,
 	bool InAssignExternalPackages)
 : DuplicatedObjectAnnotation(InDuplicatedObjects)
diff --git a/Engine/Source/Runtime/CoreUObject/Private/UObject/UObjectGlobals.cpp b/Engine/Source/Runtime/CoreUObject/Private/UObject/UObjectGlobals.cpp
--- a/Engine/Source/Runtime/CoreUObject/Private/UObject/UObjectGlobals.cpp
+++ b/Engine/Source/Runtime/CoreUObject/Private/UObject/UObjectGlobals.cpp
@@ -2774,7 +2774,7 @@ namespace UE::ObjectDuplication
 	};
 
 	// Helper method intended to gather up all default subobjects that have already been created and prepare them for duplication.
-	void GatherDefaultSubobjectsForDuplication(UObject* SrcObject, UObject* DstObject, FUObjectAnnotationSparse<FDuplicatedObject, false>& DuplicatedObjectAnnotation, FDuplicateDataWriter& Writer)
+	void GatherDefaultSubobjectsForDuplication(UObject* SrcObject, UObject* DstObject, FDuplicatedObjectAnnotation& DuplicatedObjectAnnotation, FDuplicateDataWriter& Writer)
 	{
 		TArray<UObject*> SrcDefaultSubobjects;
 		SrcObject->GetDefaultSubobjects(SrcDefaultSubobjects);
@@ -3003,7 +3003,7 @@ UObject* StaticDuplicateObjectEx( FObjectDuplicationParameters& Parameters )
 
 	FPooledLargeMemoryData ObjectData;
 
-	FUObjectAnnotationSparse<FDuplicatedObject,false>  DuplicatedObjectAnnotation;
+	FDuplicatedObjectAnnotation DuplicatedObjectAnnotation;
 	TMap<Verse::VCell*, FDuplicatedCell> DuplicatedCells;
 
 	// if seed objects were specified, add those to the DuplicatedObjects map now
diff --git a/Engine/Source/Runtime/CoreUObject/Public/Serialization/DuplicatedDataReader.h b/Engine/Source/Runtime/CoreUObject/Public/Serialization/DuplicatedDataReader.h
--- a/Engine/Source/Runtime/CoreUObject/Public/Serialization/DuplicatedDataReader.h
+++ b/Engine/Source/Runtime/CoreUObject/Public/Serialization/DuplicatedDataReader.h
@@ -21,7 +21,7 @@ class FDuplicateDataReader : public FArchiveUObject
 {
 private:
 
-	class FUObjectAnnotationSparse<FDuplicatedObject,false>&	DuplicatedObjectAnnotation;
+	FDuplicatedObjectAnnotation&			DuplicatedObjectAnnotation;
 	TMap<Verse::VCell*, FDuplicatedCell>*	DuplicatedCells;
 	const FLargeMemoryData&					ObjectData;
 	int64									Offset;
@@ -80,6 +80,6 @@ public:
 	 * @param	InDuplicatedObjectAnnotation		Annotation for storing a mapping from source to duplicated object
 	 * @param	InObjectData					Object data to read from
 	 */
-	FDuplicateDataReader( FUObjectAnnotationSparse<FDuplicatedObject,false>& InDuplicatedObjectAnnotation, const FLargeMemoryData& InObjectData, uint32 InPortFlags, UObject* InDestOuter );
+	FDuplicateDataReader(FDuplicatedObjectAnnotation& InDuplicatedObjectAnnotation, const FLargeMemoryData& InObjectData, uint32 InPortFlags, UObject* InDestOuter);
 	void SetDuplicatedCells(TMap<Verse::VCell*, FDuplicatedCell>& InDuplicatedCells);
 };
\ No newline at end of file
diff --git a/Engine/Source/Runtime/CoreUObject/Public/Serialization/DuplicatedDataWriter.h b/Engine/Source/Runtime/CoreUObject/Public/Serialization/DuplicatedDataWriter.h
--- a/Engine/Source/Runtime/CoreUObject/Public/Serialization/DuplicatedDataWriter.h
+++ b/Engine/Source/Runtime/CoreUObject/Public/Serialization/DuplicatedDataWriter.h
@@ -22,7 +22,7 @@ class FDuplicateDataWriter : public FArchiveUObject
 {
 private:
 
-	class FUObjectAnnotationSparse<FDuplicatedObject,false>&	DuplicatedObjectAnnotation;
+	FDuplicatedObjectAnnotation&			DuplicatedObjectAnnotation;
 	TMap<Verse::VCell*, FDuplicatedCell>*	DuplicatedCells;
 	FLargeMemoryData&						ObjectData;
 	int64									Offset;
@@ -115,7 +115,7 @@ public:
 	 * @param	InInstanceGraph			the instancing graph to use when creating the duplicate objects.
 	 */
 	FDuplicateDataWriter(
-		FUObjectAnnotationSparse<FDuplicatedObject, false>& InDuplicatedObjects,
+		FDuplicatedObjectAnnotation& InDuplicatedObjects,
 		FLargeMemoryData& InObjectData,
 		UObject* SourceObject,
 		UObject* DestObject,
diff --git a/Engine/Source/Runtime/CoreUObject/Public/Serialization/DuplicatedObject.h b/Engine/Source/Runtime/CoreUObject/Public/Serialization/DuplicatedObject.h
--- a/Engine/Source/Runtime/CoreUObject/Public/Serialization/DuplicatedObject.h
+++ b/Engine/Source/Runtime/CoreUObject/Public/Serialization/DuplicatedObject.h
@@ -4,6 +4,7 @@
 
 #include "CoreMinimal.h"
 #include "UObject/WeakObjectPtr.h"
+#include "UObject/UObjectAnnotation.h"
 
 /**
  * Information about a duplicated object
@@ -39,6 +40,9 @@ struct FDuplicatedObject
 
 template <> struct TIsPODType<FDuplicatedObject> { enum { Value = true }; };
 
+/** Annotation type for storing duplicated object infos */
+typedef FUObjectAnnotationTemp<FDuplicatedObject> FDuplicatedObjectAnnotation;
+
 struct FDuplicatedCell
 {
 	Verse::VCell* DuplicatedCell = nullptr;
diff --git a/Engine/Source/Runtime/CoreUObject/Public/UObject/UObjectAnnotation.h b/Engine/Source/Runtime/CoreUObject/Public/UObject/UObjectAnnotation.h
--- a/Engine/Source/Runtime/CoreUObject/Public/UObject/UObjectAnnotation.h
+++ b/Engine/Source/Runtime/CoreUObject/Public/UObject/UObjectAnnotation.h
@@ -67,7 +67,7 @@ public:
 	FUObjectAnnotationSparse()
 	{
 		// default constructor is required to be default annotation
-		check(AnnotationCacheValue.IsDefault());
+		check(TAnnotation().IsDefault());
 	}
 
 	/**
@@ -75,16 +75,6 @@ public:
 	 */
 	virtual ~FUObjectAnnotationSparse()
 	{
-#if UE_AUTORTFM
-		if (bRegisteredAutoRTFMHandlers)
-		{
-			AutoRTFM::EContextStatus Status = AutoRTFM::Close([this]
-			{
-				AutoRTFM::PopOnCommitHandler(this);
-			});
-			checkSlow(Status == AutoRTFM::EContextStatus::OnTrack);
-		}
-#endif
 		RemoveAllAnnotations();
 	}
 
@@ -105,7 +95,6 @@ private:
 			const bool bWasEmpty = (AnnotationMap.Num() == 0);
 
 			AnnotationMap.Add(Object, LocalAnnotation);
-			SetAnnotationCacheKeyAndValue(Object, MoveTemp(LocalAnnotation));
 
 			if (bWasEmpty)
 			{
@@ -183,8 +172,6 @@ public:
 
 		TAnnotation Result;
 		{
-			SetAnnotationCacheKeyAndValue(Object, TAnnotation());
-
 			UE::TRWScopeLock AnnotationMapLock(AnnotationMapCritical, SLT_Write);
 			const bool bHadElements = (AnnotationMap.Num() > 0);
 			AnnotationMap.RemoveAndCopyValue(Object, Result);
@@ -219,8 +206,6 @@ public:
 		bool bProcessRegistration = false;
 
 		{
-			SetAnnotationCacheKeyAndValue(Object, TAnnotation());
-
 			UE::TRWScopeLock AnnotationMapLock(AnnotationMapCritical, SLT_Write);
 			const bool bHadElements = (AnnotationMap.Num() > 0);
 			AnnotationMap.Remove(Object);
@@ -252,7 +237,6 @@ public:
 		bool bProcessRegistration = false;
 
 		{
-			SetAnnotationCacheKeyAndValue(nullptr, TAnnotation());
 			const bool bHadElements = (AnnotationMap.Num() > 0);
 			UE::TRWScopeLock AnnotationMapLock(AnnotationMapCritical, SLT_Write);
 			AnnotationMap.Empty();
@@ -284,30 +268,8 @@ public:
 	{
 		check(Object);
 
-		
-		// If we are within a transaction, we don't touch the annotation cache, as it can be altered from the open.
-		// Look up the annotation directly from the map.
-		if (AutoRTFM::IsTransactional())
-		{
-			UE::TRWScopeLock AnnotationMapLock(AnnotationMapCritical, SLT_ReadOnly);
-			TAnnotation* Entry = AnnotationMap.Find(Object);
-			return Entry ? *Entry : TAnnotation();
-		}
-
-		{
-			// Check cache
-			UE::TRWScopeLock AnnotationCacheLock(AnnotationCacheCritical, SLT_ReadOnly);
-			if (Object == AnnotationCacheKey)
-			{
-				return AnnotationCacheValue;
-			}
-		}
-		
-		// Update the cache if it isn't a match.
 		UE::TRWScopeLock AnnotationMapLock(AnnotationMapCritical, SLT_ReadOnly);
 		TAnnotation* Entry = AnnotationMap.Find(Object);
-		SetAnnotationCacheKeyAndValue(Object, Entry ? *Entry : TAnnotation());
-
 		return Entry ? *Entry : TAnnotation();
 	}
 
@@ -335,100 +297,21 @@ public:
 		return AnnotationMap.GetAllocatedSize();
 	}
 
-private:
-	// Sets the AnnotationCacheKey and AnnotationCacheValue if called outside of
-	// an AutoRTFM transaction.
-	// This is a no-op in a transaction to avoid writes to the same memory locations 
-	// from open and closed code. If a transaction is committed or aborted, we reset
-	// the cache to avoid reading back a stale value.
-	void SetAnnotationCacheKeyAndValue(const UObjectBase* Key, const TAnnotation& Value)
-	{
-		if (AutoRTFM::IsTransactional())
-		{
-			MaybeRegisterAutoRTFMHandlers();
-		}
-		else
-		{
-			UE::TRWScopeLock AnnotationCacheLock(AnnotationCacheCritical, SLT_Write);
-			AnnotationCacheKey = Key;
-			AnnotationCacheValue = Value;
-		}
-	}
-
-	void SetAnnotationCacheKeyAndValue(const UObjectBase* Key, TAnnotation&& Value)
-	{
-		if (AutoRTFM::IsTransactional())
-		{
-			MaybeRegisterAutoRTFMHandlers();
-		}
-		else
-		{
-			UE::TRWScopeLock AnnotationCacheLock(AnnotationCacheCritical, SLT_Write);
-			AnnotationCacheKey = Key;
-			AnnotationCacheValue = MoveTemp(Value);
-		}
-	}
-
-	void MaybeRegisterAutoRTFMHandlers()
-	{
-#if UE_AUTORTFM
-		if (!bRegisteredAutoRTFMHandlers)
-		{
-			checkSlow(AutoRTFM::IsTransactional());  // Precondition: our caller should have checked this
-
-			AutoRTFM::EContextStatus Status = AutoRTFM::Close([this]
-			{
-				bRegisteredAutoRTFMHandlers = true;
-
-				// We avoid touching the cache key and value within a transaction, so after a transaction completes,
-				// the cache value may be outdated. To avoid stale data, we null out the cache key at the end of a
-				// transaction. It will be automatically repopulated the next time GetAnnotation is called.
-				AutoRTFM::PushOnCommitHandler(this, [this] 
-				{ 
-					AnnotationCacheKey = nullptr; 
-					bRegisteredAutoRTFMHandlers = false;
-				});
-			});
-
-			checkSlow(Status == AutoRTFM::EContextStatus::OnTrack);
-		}
-#endif
-	}
-
 protected:
 	/**
-	 * Guards the annotation map, cache key and cache value.
+	 * Guards the annotation map
 	 */
 	mutable FTransactionallySafeRWLock AnnotationMapCritical;
-	mutable FTransactionallySafeRWLock AnnotationCacheCritical;
 
 	/**
 	 * Map from live objects to an annotation. Guarded by AnnotationMapCritical.
 	 */
 	TMap<const UObjectBase *, TAnnotation> AnnotationMap;
 
-	/**
-	 * Key for a one-item cache of the last lookup into AnnotationMap. Guarded by AnnotationCacheCritical.
-	 * Annotation are often called back-to-back so this is a performance optimization for that.
-	 */
-	const UObjectBase* AnnotationCacheKey = nullptr;
-
-	/**
-	 * Value for a one-item cache of the last lookup into AnnotationMap. Guarded by AnnotationCacheCritical.
-	 */
-	TAnnotation AnnotationCacheValue;
-
 	/**
 	 * Monitor registration to GUObjectArray/UObjectDeleteListener.
 	 */
 	bool bRegistered = false;
-
-#if UE_AUTORTFM
-	/**
-	 * Tracks the presence of a AutoRTFM handlers for cleaning up the annotation cache.
-	 */
-	bool bRegisteredAutoRTFMHandlers = false;
-#endif
 };
 
 
@@ -1592,3 +1475,161 @@ private:
 	 */
 	TArray<TBitType> AnnotationArray;
 };
+
+/**
+ * FUObjectAnnotationTemp is a temporary annotation with the same API as FUObjectAnnotationSparse but without the thread safety or GC checks.
+ * It should only be used for single threaded functions and must be deleted before garbage collection occurs.
+ */
+template<typename TAnnotation>
+class FUObjectAnnotationTemp
+{
+public:
+
+	/**
+	 * Constructor, initializes to nothing
+	 */
+	FUObjectAnnotationTemp()
+	{
+		// default constructor is required to be default annotation
+		check(AnnotationCacheValue.IsDefault());
+	}
+
+	/**
+	 * Destructor, removes all annotations
+	 */
+	virtual ~FUObjectAnnotationTemp()
+	{
+		RemoveAllAnnotations();
+	}
+
+private:
+	template<typename T>
+	void AddAnnotationInternal(const UObjectBase* Object, T&& Annotation)
+	{
+		TAnnotation LocalAnnotation = Forward<T>(Annotation);
+		if (LocalAnnotation.IsDefault())
+		{
+			RemoveAnnotation(Object); // adding the default annotation is the same as removing an annotation
+		}
+		else
+		{
+			AnnotationMap.Add(Object, LocalAnnotation);
+			SetAnnotationCacheKeyAndValue(Object, MoveTemp(LocalAnnotation));
+		}
+	}
+
+public:
+	/**
+	 * Add an annotation to the annotation list. If the Annotation is the default, then the annotation is removed.
+	 *
+	 * @param Object        Object to annotate.
+	 * @param Annotation    Annotation to associate with Object.
+	 */
+	void AddAnnotation(const UObjectBase* Object, TAnnotation&& Annotation)
+	{
+		AddAnnotationInternal(Object, MoveTemp(Annotation));
+	}
+
+	void AddAnnotation(const UObjectBase* Object, const TAnnotation& Annotation)
+	{
+		AddAnnotationInternal(Object, Annotation);
+	}
+	/**
+	 * Removes an annotation from the annotation list and returns the annotation if it had one 
+	 *
+	 * @param Object		Object to de-annotate.
+	 * @return				Old annotation
+	 */
+	TAnnotation GetAndRemoveAnnotation(const UObjectBase *Object)
+	{
+		TAnnotation Result;
+		SetAnnotationCacheKeyAndValue(Object, TAnnotation());
+		AnnotationMap.RemoveAndCopyValue(Object, Result);
+
+		return Result;
+	}
+	/**
+	 * Removes an annotation from the annotation list. 
+	 *
+	 * @param Object		Object to de-annotate.
+	 */
+	void RemoveAnnotation(const UObjectBase *Object)
+	{
+		SetAnnotationCacheKeyAndValue(Object, TAnnotation());
+		AnnotationMap.Remove(Object);
+	}
+	/**
+	 * Removes all annotation from the annotation list. 
+	 */
+	void RemoveAllAnnotations()
+	{
+		SetAnnotationCacheKeyAndValue(nullptr, TAnnotation());
+		AnnotationMap.Empty();
+	}
+
+	/**
+	 * Return the annotation associated with a uobject 
+	 *
+	 * @param Object		Object to return the annotation for
+	 */
+	inline TAnnotation GetAnnotation(const UObjectBase *Object)
+	{
+		// Check cache
+		if (Object == AnnotationCacheKey)
+		{
+			return AnnotationCacheValue;
+		}
+		
+		// Update the cache if it isn't a match.
+		TAnnotation* Entry = AnnotationMap.Find(Object);
+		SetAnnotationCacheKeyAndValue(Object, Entry ? *Entry : TAnnotation());
+
+		return Entry ? *Entry : TAnnotation();
+	}
+
+	/**
+	 * Return the annotation map. Caution, this is for low level use 
+	 * @return A mapping from UObjectBase to annotation for non-default annotations
+	 */
+	const TMap<const UObjectBase *,TAnnotation>& GetAnnotationMap() const
+	{
+		return AnnotationMap;
+	}
+
+	/** 
+	 * Reserves memory for the annotation map for the specified number of elements, used to avoid reallocations. 
+	 */
+	void Reserve(int32 ExpectedNumElements)
+	{
+		AnnotationMap.Empty(ExpectedNumElements);
+	}
+private:
+	void SetAnnotationCacheKeyAndValue(const UObjectBase* Key, const TAnnotation& Value)
+	{
+		AnnotationCacheKey = Key;
+		AnnotationCacheValue = Value;
+	}
+
+	void SetAnnotationCacheKeyAndValue(const UObjectBase* Key, TAnnotation&& Value)
+	{
+		AnnotationCacheKey = Key;
+		AnnotationCacheValue = MoveTemp(Value);
+	}
+
+protected:
+	/**
+	 * Map from live objects to an annotation.
+	 */
+	TMap<const UObjectBase *, TAnnotation> AnnotationMap;
+
+	/**
+	 * Key for a one-item cache of the last lookup into AnnotationMap.
+	 * Annotation are often called back-to-back so this is a performance optimization for that.
+	 */
+	const UObjectBase* AnnotationCacheKey = nullptr;
+
+	/**
+	 * Value for a one-item cache of the last lookup into AnnotationMap.
+	 */
+	TAnnotation AnnotationCacheValue;
+};
diff --git a/Engine/Source/Runtime/Engine/Private/ComponentInstanceDataCache.cpp b/Engine/Source/Runtime/Engine/Private/ComponentInstanceDataCache.cpp
--- a/Engine/Source/Runtime/Engine/Private/ComponentInstanceDataCache.cpp
+++ b/Engine/Source/Runtime/Engine/Private/ComponentInstanceDataCache.cpp
@@ -155,7 +155,7 @@ protected:
 
 private:
 	const UObject* CacheObject;
-	FUObjectAnnotationSparse<FDuplicatedObject,false> DuplicatedObjectAnnotation;
+	FDuplicatedObjectAnnotation DuplicatedObjectAnnotation;
 	FInstanceCacheDataBase& InstanceData;
 };
 
