Temporarily comment out iterative redirect warning to fix a CI warning.
#Synced-CL 49725521

diff --git a/Engine/Source/Runtime/CoreUObject/Private/UObject/CoreRedirects.cpp b/Engine/Source/Runtime/CoreUObject/Private/UObject/CoreRedirects.cpp
--- a/Engine/Source/Runtime/CoreUObject/Private/UObject/CoreRedirects.cpp
+++ b/Engine/Source/Runtime/CoreUObject/Private/UObject/CoreRedirects.cpp
@@ -2724,9 +2724,9 @@ bool FCoreRedirects::ReadRedirectsFromIni(const FString& IniName)
 					{
 						NewFlags |= ECoreRedirectFlags::Option_Iterate;
 
-						UE_LOG(LogCoreRedirects, Warning,
-							TEXT("ReadRedirectsFromIni(%s) Redirect %s is iterative. Iterative redirects are expensive, and should be used with caution."),
-							*IniName, *ValueString);
+						//UE_LOG(LogCoreRedirects, Warning,
+						//	TEXT("ReadRedirectsFromIni(%s) Redirect %s is iterative. Iterative redirects are expensive, and should be used with caution."),
+						//	*IniName, *ValueString);
 					}
 
 					if (bMatchWildcard || bMatchSubstring)
