silence a pvs warning about nullptr

#Synced-CL 49706563

diff --git a/Engine/Source/Runtime/CoreUObject/Private/UObject/RemoteHeap.cpp b/Engine/Source/Runtime/CoreUObject/Private/UObject/RemoteHeap.cpp
--- a/Engine/Source/Runtime/CoreUObject/Private/UObject/RemoteHeap.cpp
+++ b/Engine/Source/Runtime/CoreUObject/Private/UObject/RemoteHeap.cpp
@@ -913,7 +913,7 @@ void FRemoteHeapAllocator::ForAnyElementType::ResizeAllocation(FRemoteHeapAlloca
 		
 		if (CurrentNum != 0)
 		{
-			FMemory::Memcpy(NewMemory, Allocation, CurrentNum * NumBytesPerElement);
+			FMemory::Memcpy(NewMemory, Allocation, CurrentNum * NumBytesPerElement); //-V575
 		}
 
 		GRemoteHeap.RemoteHeapDeallocate(Allocation);
