PVS Studio 7.40

#rnx
#Synced-CL 49749982

diff --git a/Engine/Plugins/Media/ElectraUtil/Source/ElectraHTTPStream/Private/Generic/GenericElectraHTTPStream.cpp b/Engine/Plugins/Media/ElectraUtil/Source/ElectraHTTPStream/Private/Generic/GenericElectraHTTPStream.cpp
--- a/Engine/Plugins/Media/ElectraUtil/Source/ElectraHTTPStream/Private/Generic/GenericElectraHTTPStream.cpp
+++ b/Engine/Plugins/Media/ElectraUtil/Source/ElectraHTTPStream/Private/Generic/GenericElectraHTTPStream.cpp
@@ -753,7 +753,7 @@ void FElectraHTTPStreamGeneric::SetupNewRequests()
 	TArray<TSharedPtr<FElectraHTTPStreamRequestGeneric, ESPMode::ThreadSafe>> NewReqs;
 	Swap(NewReqs, NewRequests);
 	RequestLock.Unlock();
-	for(auto &Request : NewReqs) //-V1078 7.39 Swap Bug
+	for(auto &Request : NewReqs)
 	{
 		if (Request->Setup())
 		{
@@ -824,7 +824,7 @@ void FElectraHTTPStreamGeneric::HandleCompletedRequests()
 	{
 		TArray<TSharedPtr<FElectraHTTPStreamRequestGeneric, ESPMode::ThreadSafe>> TempRequests;
 		Swap(CompletedRequests, TempRequests);
-		for(auto &Finished : TempRequests) //-V1078 7.39 Swap Bug
+		for(auto &Finished : TempRequests)
 		{
 			Finished->SetFinished();
 			Finished->Close();
diff --git a/Engine/Plugins/Runtime/RigVM/Source/RigVM/Public/RigVMCore/RigVMNativized.h b/Engine/Plugins/Runtime/RigVM/Source/RigVM/Public/RigVMCore/RigVMNativized.h
--- a/Engine/Plugins/Runtime/RigVM/Source/RigVM/Public/RigVMCore/RigVMNativized.h
+++ b/Engine/Plugins/Runtime/RigVM/Source/RigVM/Public/RigVMCore/RigVMNativized.h
@@ -295,7 +295,7 @@ protected:
 		TArray<T> Temp;
 		Swap(InOutA, Temp);
 
-		for(const T& Element : Temp) //-V1078 7.39 Swap Bug
+		for(const T& Element : Temp)
 		{
 			if(!InB.Contains(Element))
 			{
@@ -317,7 +317,7 @@ protected:
 		TArray<T> Temp;
 		Swap(InOutA, Temp);
 
-		for(const T& Element : Temp) //-V1078 7.39 Swap Bug
+		for(const T& Element : Temp)
 		{
 			if(InB.Contains(Element))
 			{
diff --git a/Engine/Plugins/Runtime/RigVM/Source/RigVMDeveloper/Private/RigVMModel/RigVMController.cpp b/Engine/Plugins/Runtime/RigVM/Source/RigVMDeveloper/Private/RigVMModel/RigVMController.cpp
--- a/Engine/Plugins/Runtime/RigVM/Source/RigVMDeveloper/Private/RigVMModel/RigVMController.cpp
+++ b/Engine/Plugins/Runtime/RigVM/Source/RigVMDeveloper/Private/RigVMModel/RigVMController.cpp
@@ -23024,7 +23024,7 @@ void URigVMController::ProcessDetachedLinks(const FRestoreLinkedPathSettings& In
 	Swap(GetGraph()->DetachedLinks, DetachedLinks);
 
 	// destroy the links finally
-	for(URigVMLink* DetachedLink : DetachedLinks) //-V1078 7.39 Swap Bug
+	for(URigVMLink* DetachedLink : DetachedLinks)
 	{
 		ReportRemovedLink(DetachedLink->GetSourcePinPath(), DetachedLink->GetTargetPinPath());
 		Notify(ERigVMGraphNotifType::LinkRemoved, DetachedLink);
diff --git a/Engine/Source/Developer/Windows/LiveCoding/Private/LiveCodingModule.cpp b/Engine/Source/Developer/Windows/LiveCoding/Private/LiveCodingModule.cpp
--- a/Engine/Source/Developer/Windows/LiveCoding/Private/LiveCodingModule.cpp
+++ b/Engine/Source/Developer/Windows/LiveCoding/Private/LiveCodingModule.cpp
@@ -1356,7 +1356,7 @@ void FLiveCodingModule::UpdateModules(bool bAllowStarting)
 
 		TArray<FString> EnableModules;
 		TArray<FString> LazyLoadModules;
-		for (const ModuleChange& Change : Changes) //-V1078 7.39 Swap Bug
+		for (const ModuleChange& Change : Changes)
 		{
 			if (Change.bLoaded)
 			{
diff --git a/Engine/Source/Editor/Sequencer/Private/SequencerKeyParams.cpp b/Engine/Source/Editor/Sequencer/Private/SequencerKeyParams.cpp
--- a/Engine/Source/Editor/Sequencer/Private/SequencerKeyParams.cpp
+++ b/Engine/Source/Editor/Sequencer/Private/SequencerKeyParams.cpp
@@ -145,7 +145,7 @@ void FKeyOperation::FSectionCandidates::FilterOperations(UMovieSceneTrack* Track
 		Swap(OldOperations, Operations);
 
 		TArray<UMovieSceneSection*, TInlineAllocator<1>> Sections;
-		for (const FKeySectionOperation& Operation : OldOperations) //-V1078 7.39 Swap Bug
+		for (const FKeySectionOperation& Operation : OldOperations)
 		{
 			UMovieSceneSection* Section = Operation.Section->GetSectionObject();
 			Sections.Add(Section);
diff --git a/Engine/Source/Editor/UnrealEd/Private/Cooker/CookWorkerClient.cpp b/Engine/Source/Editor/UnrealEd/Private/Cooker/CookWorkerClient.cpp
--- a/Engine/Source/Editor/UnrealEd/Private/Cooker/CookWorkerClient.cpp
+++ b/Engine/Source/Editor/UnrealEd/Private/Cooker/CookWorkerClient.cpp
@@ -699,7 +699,7 @@ void FCookWorkerClient::SendPendingMessages()
 		FScopeLock FileTransferScopeLock(&FileTransferLock);
 		Swap(LocalPendingFileTransfers, PendingFileTransfers);
 	}
-	for (FFileTransferMessage& FileTransferMessage : LocalPendingFileTransfers) //-V1078 7.39 Swap Bug
+	for (FFileTransferMessage& FileTransferMessage : LocalPendingFileTransfers)
 	{
 		SendMessage(FileTransferMessage);
 	}
diff --git a/Engine/Source/Runtime/Core/Private/IO/IoDispatcher.cpp b/Engine/Source/Runtime/Core/Private/IO/IoDispatcher.cpp
--- a/Engine/Source/Runtime/Core/Private/IO/IoDispatcher.cpp
+++ b/Engine/Source/Runtime/Core/Private/IO/IoDispatcher.cpp
@@ -947,7 +947,7 @@ private:
 				Swap(LocalRequestsToCancel, RequestsToCancel);
 				Swap(LocalRequestsToReprioritize, RequestsToReprioritize);
 			}
-			for (FIoRequestImpl* RequestToCancel : LocalRequestsToCancel) //-V1078 7.39 Swap Bug
+			for (FIoRequestImpl* RequestToCancel : LocalRequestsToCancel)
 			{
 				if (!RequestToCancel->IsCancelled())
 				{
@@ -959,7 +959,7 @@ private:
 				}
 				RequestToCancel->ReleaseRef();
 			}
-			for (FIoRequestImpl* RequestToRePrioritize : LocalRequestsToReprioritize) //-V1078 7.39 Swap Bug
+			for (FIoRequestImpl* RequestToRePrioritize : LocalRequestsToReprioritize)
 			{
 				if (RequestToRePrioritize->Backend)
 				{
diff --git a/Engine/Source/Runtime/CoreUObject/Private/Serialization/AsyncLoading2.cpp b/Engine/Source/Runtime/CoreUObject/Private/Serialization/AsyncLoading2.cpp
--- a/Engine/Source/Runtime/CoreUObject/Private/Serialization/AsyncLoading2.cpp
+++ b/Engine/Source/Runtime/CoreUObject/Private/Serialization/AsyncLoading2.cpp
@@ -12203,7 +12203,7 @@ void FAsyncLoadingThread2::ConditionalProcessEditorCallbacks()
 
 		// Call the global delegate for package endloads and set the bHasBeenLoaded flag that is used to
 		// check which packages have reached this state
-		for (UPackage* CompletedUPackage : LocalEditorCompletedUPackages) //-V1078 7.39 Swap Bug
+		for (UPackage* CompletedUPackage : LocalEditorCompletedUPackages)
 		{
 			CompletedUPackage->SetHasBeenEndLoaded(true);
 		}
@@ -12211,7 +12211,7 @@ void FAsyncLoadingThread2::ConditionalProcessEditorCallbacks()
 			FEndLoadPackageContext{ LocalEditorCompletedUPackages, 0, false /* bSynchronous */ });
 
 		// In editor builds, call the asset load callback. This happens in both editor and standalone to match EndLoad
-		for (UObject* LoadedObject : LocalEditorLoadedAssets) //-V1078 7.39 Swap Bug
+		for (UObject* LoadedObject : LocalEditorLoadedAssets)
 		{
 			if (LoadedObject)
 			{
diff --git a/Engine/Source/Runtime/Engine/Private/UnrealEngine.cpp b/Engine/Source/Runtime/Engine/Private/UnrealEngine.cpp
--- a/Engine/Source/Runtime/Engine/Private/UnrealEngine.cpp
+++ b/Engine/Source/Runtime/Engine/Private/UnrealEngine.cpp
@@ -11552,7 +11552,7 @@ bool UEngine::PerformError(const TCHAR* Cmd, FOutputDevice& Ar)
 			{
 				void* Eat = FMemory::Malloc(AllocSizeBytes);
 				FMemory::Memset(Eat, 0, AllocSizeBytes);
-			}
+			} //-V773
 
 			return true;
 		}));
@@ -11581,7 +11581,7 @@ bool UEngine::PerformError(const TCHAR* Cmd, FOutputDevice& Ar)
 			// newly allocated memory must be touched
 			FMemory::Memset(Eat, 0, BytesPerMegabyte);
 			MBToAllocate--;
-		}
+		} //-V773
 
 		return true;
 	}
diff --git a/Engine/Source/Runtime/Landscape/Private/LandscapeSplines.cpp b/Engine/Source/Runtime/Landscape/Private/LandscapeSplines.cpp
--- a/Engine/Source/Runtime/Landscape/Private/LandscapeSplines.cpp
+++ b/Engine/Source/Runtime/Landscape/Private/LandscapeSplines.cpp
@@ -3194,7 +3194,7 @@ void ULandscapeSplineSegment::UpdateSplinePoints(bool bUpdateCollision, bool bUp
 	TMap<ULandscapeSplinesComponent*, TArray<USplineMeshComponent*>> ForeignMeshComponentsMap = GetForeignMeshComponents();
 	
 	// Unregister components, Remove Foreign/Local Associations
-	for (TObjectPtr<USplineMeshComponent>& LocalMeshComponent : OldLocalMeshComponents) //-V1078 7.39 Swap Bug
+	for (TObjectPtr<USplineMeshComponent>& LocalMeshComponent : OldLocalMeshComponents)
 	{
 		USplineMeshComponent* Comp = LocalMeshComponent.Get();
 		checkSlow(OuterSplines->MeshComponentLocalOwnersMap.FindRef(Comp) == this);
diff --git a/Engine/Source/Runtime/MovieScene/Private/EntitySystem/MovieSceneEntitySystemRunner.cpp b/Engine/Source/Runtime/MovieScene/Private/EntitySystem/MovieSceneEntitySystemRunner.cpp
--- a/Engine/Source/Runtime/MovieScene/Private/EntitySystem/MovieSceneEntitySystemRunner.cpp
+++ b/Engine/Source/Runtime/MovieScene/Private/EntitySystem/MovieSceneEntitySystemRunner.cpp
@@ -1327,7 +1327,7 @@ void FMovieSceneEntitySystemRunner::GameThread_PostEvaluationPhase(UMovieSceneEn
 	{
 		FMovieSceneEntitySystemEvaluationReentrancyWindow Window(this, Linker);
 
-		for (const FQueuedUpdateParams& UpdateParams : TmpCurrentInstances) //-V1078 7.39 Swap Bug
+		for (const FQueuedUpdateParams& UpdateParams : TmpCurrentInstances)
 		{
 			// We must check for validity here because the cache handles may have become invalid
 			// during this iteration (since there is a re-entrancy window open)
@@ -1354,7 +1354,7 @@ void FMovieSceneEntitySystemRunner::GameThread_PostEvaluationPhase(UMovieSceneEn
 			}
 		}
 
-		for (const FSimpleDelegate& OnFlushed : TmpOnFlushedDelegates) //-V1078 7.39 Swap Bug
+		for (const FSimpleDelegate& OnFlushed : TmpOnFlushedDelegates)
 		{
 			OnFlushed.Execute();
 		}
diff --git a/Engine/Source/Runtime/MovieScene/Private/MovieScene.cpp b/Engine/Source/Runtime/MovieScene/Private/MovieScene.cpp
--- a/Engine/Source/Runtime/MovieScene/Private/MovieScene.cpp
+++ b/Engine/Source/Runtime/MovieScene/Private/MovieScene.cpp
@@ -692,7 +692,7 @@ void UMovieScene::EmptyRootFolders()
 
 	RootFolders.Empty();
 
-	for (UMovieSceneFolder* Folder : OldFolders) //-V1078 7.39 Swap Bug
+	for (UMovieSceneFolder* Folder : OldFolders)
 	{
 		EventHandlers.Trigger(&UE::MovieScene::ISequenceDataEventHandler::OnRootFolderRemoved, Folder);
 	}
diff --git a/Engine/Source/Runtime/MovieSceneCapture/Private/ActiveMovieSceneCaptures.cpp b/Engine/Source/Runtime/MovieSceneCapture/Private/ActiveMovieSceneCaptures.cpp
--- a/Engine/Source/Runtime/MovieSceneCapture/Private/ActiveMovieSceneCaptures.cpp
+++ b/Engine/Source/Runtime/MovieSceneCapture/Private/ActiveMovieSceneCaptures.cpp
@@ -45,7 +45,7 @@ void FActiveMovieSceneCaptures::Shutdown()
 	TArray<TObjectPtr<UMovieSceneCapture>> ActiveCapturesCopy;
 	Swap(ActiveCaptures, ActiveCapturesCopy);
 
-	for (auto& Obj : ActiveCapturesCopy) //-V1078 7.39 Swap Bug
+	for (auto& Obj : ActiveCapturesCopy)
 	{
 		Obj->Finalize();
 	}
diff --git a/Engine/Source/Runtime/MovieSceneCapture/Private/FrameGrabberProtocol.cpp b/Engine/Source/Runtime/MovieSceneCapture/Private/FrameGrabberProtocol.cpp
--- a/Engine/Source/Runtime/MovieSceneCapture/Private/FrameGrabberProtocol.cpp
+++ b/Engine/Source/Runtime/MovieSceneCapture/Private/FrameGrabberProtocol.cpp
@@ -37,7 +37,7 @@ void UFrameGrabberProtocol::TickImpl()
 	{
 		TArray<FCapturedFrameData> CapturedFrames = FrameGrabber->GetCapturedFrames();
 
-		for (FCapturedFrameData& Frame : CapturedFrames) //-V1078 7.39 Swap Bug
+		for (FCapturedFrameData& Frame : CapturedFrames)
 		{
 			ProcessFrame(MoveTemp(Frame));
 		}
diff --git a/Engine/Source/Runtime/MovieSceneCapture/Private/UserDefinedCaptureProtocol.cpp b/Engine/Source/Runtime/MovieSceneCapture/Private/UserDefinedCaptureProtocol.cpp
--- a/Engine/Source/Runtime/MovieSceneCapture/Private/UserDefinedCaptureProtocol.cpp
+++ b/Engine/Source/Runtime/MovieSceneCapture/Private/UserDefinedCaptureProtocol.cpp
@@ -115,7 +115,7 @@ void UUserDefinedCaptureProtocol::TickImpl()
 	{
 		TArray<FCapturedFrameData> CapturedFrames = FinalPixelsFrameGrabber->GetCapturedFrames();
 
-		for (FCapturedFrameData& Frame : CapturedFrames) //-V1078 7.39 Swap Bug
+		for (FCapturedFrameData& Frame : CapturedFrames)
 		{
 			// Steal the frame and make it shareable
 			FCapturedPixels CapturedPixels { MakeShared<TImagePixelData<FColor>, ESPMode::ThreadSafe>( Frame.BufferSize, TArray64<FColor>(MoveTemp(Frame.ColorBuffer) ) ) };
diff --git a/Engine/Source/Runtime/UMG/Private/Animation/UMGSequenceTickManager.cpp b/Engine/Source/Runtime/UMG/Private/Animation/UMGSequenceTickManager.cpp
--- a/Engine/Source/Runtime/UMG/Private/Animation/UMGSequenceTickManager.cpp
+++ b/Engine/Source/Runtime/UMG/Private/Animation/UMGSequenceTickManager.cpp
@@ -178,7 +178,7 @@ void UUMGSequenceTickManager::TickWidgetAnimations(float DeltaSeconds)
 		TArray<TObjectPtr<UUserWidget>> NewUserWidgets;
 		Swap(PendingUserWidgets, NewUserWidgets);
 
-		for (TObjectPtr<UUserWidget> NewWidget : NewUserWidgets) //-V1078 7.39 Swap Bug
+		for (TObjectPtr<UUserWidget> NewWidget : NewUserWidgets)
 		{
 			FSequenceTickManagerWidgetData NewWidgetData;
 
