Added tests for correct dispatch in dynamic delegates, after the recent change to using templates for delegate wrappers instead of UHT.
Added tests for dynamic delegate payload support, which is currently disabled but in development.

diff --git a/Engine/Source/Runtime/Core/Public/UObject/ScriptDelegates.h b/Engine/Source/Runtime/Core/Public/UObject/ScriptDelegates.h
--- a/Engine/Source/Runtime/Core/Public/UObject/ScriptDelegates.h
+++ b/Engine/Source/Runtime/Core/Public/UObject/ScriptDelegates.h
@@ -16,6 +16,16 @@
 #include "UObject/NameTypes.h"
 #include <type_traits>
 
+/* Specifies whether or not dynamic delegates can have payloads. */
+#ifndef UE_USE_DYNAMIC_DELEGATE_PAYLOADS
+	#define UE_USE_DYNAMIC_DELEGATE_PAYLOADS 0
+
+	/* Dynamic delegates with payloads and return values need special handling. */
+	#ifndef UE_USE_DYNAMIC_DELEGATE_PAYLOADS_RETVAL
+		#define UE_USE_DYNAMIC_DELEGATE_PAYLOADS_RETVAL 0
+	#endif
+#endif
+
 namespace UE::Core::Private
 {
 	template <typename InThreadSafetyMode>
diff --git a/Engine/Source/Runtime/CoreUObject/Public/UObject/DynamicDelegateTestObject.cpp b/Engine/Source/Runtime/CoreUObject/Public/UObject/DynamicDelegateTestObject.cpp
new file mode 100644
--- /dev/null
+++ b/Engine/Source/Runtime/CoreUObject/Public/UObject/DynamicDelegateTestObject.cpp
@@ -0,0 +1,807 @@
+// Copyright Epic Games, Inc. All Rights Reserved.
+
+#if WITH_TESTS
+
+#include "DynamicDelegateTestObject.h"
+
+#include "Containers/AnsiString.h"
+#include "Tests/TestHarnessAdapter.h"
+
+// Set by the name of the function to say which one was called.
+FAnsiString GCalledFunctionName;
+
+int32                                      GIntArg        = 5;
+FString                                    GStrArg        = "Hello";
+const UDynamicDelegateTestObject*          GPtrArg        = nullptr;
+FName                                      GNameArg       = "Package";
+bool                                       bGBoolArg      = true;
+TWeakObjectPtr<UDynamicDelegateTestObject> GWeakObjPtrArg = nullptr;
+
+int32   GIntRet = 99;
+FString GStrRet = "World";
+
+UDynamicDelegateTestObject::UDynamicDelegateTestObject()
+{
+}
+
+void UDynamicDelegateTestObject::Func0()
+{
+	CHECK(GCalledFunctionName.IsEmpty());
+	GCalledFunctionName = "Func0";
+}
+
+void UDynamicDelegateTestObject::Func1(int Int)
+{
+	CHECK(GCalledFunctionName.IsEmpty());
+	GCalledFunctionName = "Func1";
+
+	CHECK(Int == GIntArg);
+}
+
+void UDynamicDelegateTestObject::Func2(int Int, FString Str)
+{
+	CHECK(GCalledFunctionName.IsEmpty());
+	GCalledFunctionName = "Func2";
+
+	CHECK(Int == GIntArg);
+	CHECK(Str == GStrArg);
+}
+
+void UDynamicDelegateTestObject::Func3(int Int, FString Str, const UDynamicDelegateTestObject* Ptr)
+{
+	CHECK(GCalledFunctionName.IsEmpty());
+	GCalledFunctionName = "Func3";
+
+	CHECK(Int == GIntArg);
+	CHECK(Str == GStrArg);
+	CHECK(Ptr == GPtrArg);
+}
+
+void UDynamicDelegateTestObject::Func4(int Int, FString Str, const UDynamicDelegateTestObject* Ptr, FName Name)
+{
+	CHECK(GCalledFunctionName.IsEmpty());
+	GCalledFunctionName = "Func4";
+
+	CHECK(Int  == GIntArg);
+	CHECK(Str  == GStrArg);
+	CHECK(Ptr  == GPtrArg);
+	CHECK(Name == GNameArg);
+}
+
+void UDynamicDelegateTestObject::Func5(int Int, FString Str, const UDynamicDelegateTestObject* Ptr, FName Name, bool bBool)
+{
+	CHECK(GCalledFunctionName.IsEmpty());
+	GCalledFunctionName = "Func5";
+
+	CHECK(Int   == GIntArg);
+	CHECK(Str   == GStrArg);
+	CHECK(Ptr   == GPtrArg);
+	CHECK(Name  == GNameArg);
+	CHECK(bBool == bGBoolArg);
+}
+
+void UDynamicDelegateTestObject::Func6(int Int, FString Str, const UDynamicDelegateTestObject* Ptr, FName Name, bool bBool, TWeakObjectPtr<UDynamicDelegateTestObject> WeakObjPtr)
+{
+	CHECK(GCalledFunctionName.IsEmpty());
+	GCalledFunctionName = "Func6";
+
+	CHECK(Int        == GIntArg);
+	CHECK(Str        == GStrArg);
+	CHECK(Ptr        == GPtrArg);
+	CHECK(Name       == GNameArg);
+	CHECK(bBool      == bGBoolArg);
+	CHECK(WeakObjPtr == GWeakObjPtrArg);
+}
+
+void UDynamicDelegateTestObject::ConstFunc0() const
+{
+	CHECK(GCalledFunctionName.IsEmpty());
+	GCalledFunctionName = "ConstFunc0";
+}
+
+void UDynamicDelegateTestObject::ConstFunc1(int Int) const
+{
+	CHECK(GCalledFunctionName.IsEmpty());
+	GCalledFunctionName = "ConstFunc1";
+
+	CHECK(Int == GIntArg);
+}
+
+void UDynamicDelegateTestObject::ConstFunc2(int Int, FString Str) const
+{
+	CHECK(GCalledFunctionName.IsEmpty());
+	GCalledFunctionName = "ConstFunc2";
+
+	CHECK(Int == GIntArg);
+	CHECK(Str == GStrArg);
+}
+
+void UDynamicDelegateTestObject::ConstFunc3(int Int, FString Str, const UDynamicDelegateTestObject* Ptr) const
+{
+	CHECK(GCalledFunctionName.IsEmpty());
+	GCalledFunctionName = "ConstFunc3";
+
+	CHECK(Int == GIntArg);
+	CHECK(Str == GStrArg);
+	CHECK(Ptr == GPtrArg);
+}
+
+void UDynamicDelegateTestObject::ConstFunc4(int Int, FString Str, const UDynamicDelegateTestObject* Ptr, FName Name) const
+{
+	CHECK(GCalledFunctionName.IsEmpty());
+	GCalledFunctionName = "ConstFunc4";
+
+	CHECK(Int  == GIntArg);
+	CHECK(Str  == GStrArg);
+	CHECK(Ptr  == GPtrArg);
+	CHECK(Name == GNameArg);
+}
+
+void UDynamicDelegateTestObject::ConstFunc5(int Int, FString Str, const UDynamicDelegateTestObject* Ptr, FName Name, bool bBool) const
+{
+	CHECK(GCalledFunctionName.IsEmpty());
+	GCalledFunctionName = "ConstFunc5";
+
+	CHECK(Int   == GIntArg);
+	CHECK(Str   == GStrArg);
+	CHECK(Ptr   == GPtrArg);
+	CHECK(Name  == GNameArg);
+	CHECK(bBool == bGBoolArg);
+}
+
+void UDynamicDelegateTestObject::ConstFunc6(int Int, FString Str, const UDynamicDelegateTestObject* Ptr, FName Name, bool bBool, TWeakObjectPtr<UDynamicDelegateTestObject> WeakObjPtr) const
+{
+	CHECK(GCalledFunctionName.IsEmpty());
+	GCalledFunctionName = "ConstFunc6";
+
+	CHECK(Int        == GIntArg);
+	CHECK(Str        == GStrArg);
+	CHECK(Ptr        == GPtrArg);
+	CHECK(Name       == GNameArg);
+	CHECK(bBool      == bGBoolArg);
+	CHECK(WeakObjPtr == GWeakObjPtrArg);
+}
+
+int UDynamicDelegateTestObject::FuncRetInt0()
+{
+	Func0();
+	GCalledFunctionName = "FuncRetInt0";
+	return GIntRet;
+}
+
+int UDynamicDelegateTestObject::FuncRetInt1(int Int)
+{
+	Func1(Int);
+	GCalledFunctionName = "FuncRetInt1";
+	return GIntRet;
+}
+
+int UDynamicDelegateTestObject::FuncRetInt2(int Int, FString Str)
+{
+	Func2(Int, Str);
+	GCalledFunctionName = "FuncRetInt2";
+	return GIntRet;
+}
+
+int UDynamicDelegateTestObject::FuncRetInt3(int Int, FString Str, const UDynamicDelegateTestObject* Ptr)
+{
+	Func3(Int, Str, Ptr);
+	GCalledFunctionName = "FuncRetInt3";
+	return GIntRet;
+}
+
+int UDynamicDelegateTestObject::FuncRetInt4(int Int, FString Str, const UDynamicDelegateTestObject* Ptr, FName Name)
+{
+	Func4(Int, Str, Ptr, Name);
+	GCalledFunctionName = "FuncRetInt4";
+	return GIntRet;
+}
+
+int UDynamicDelegateTestObject::FuncRetInt5(int Int, FString Str, const UDynamicDelegateTestObject* Ptr, FName Name, bool bBool)
+{
+	Func5(Int, Str, Ptr, Name, bBool);
+	GCalledFunctionName = "FuncRetInt5";
+	return GIntRet;
+}
+
+int UDynamicDelegateTestObject::FuncRetInt6(int Int, FString Str, const UDynamicDelegateTestObject* Ptr, FName Name, bool bBool, TWeakObjectPtr<UDynamicDelegateTestObject> WeakObjPtr)
+{
+	Func6(Int, Str, Ptr, Name, bBool, WeakObjPtr);
+	GCalledFunctionName = "FuncRetInt6";
+	return GIntRet;
+}
+
+FString UDynamicDelegateTestObject::FuncRetStr0()
+{
+	Func0();
+	GCalledFunctionName = "FuncRetStr0";
+	return GStrRet;
+}
+
+FString UDynamicDelegateTestObject::FuncRetStr1(int Int)
+{
+	Func1(Int);
+	GCalledFunctionName = "FuncRetStr1";
+	return GStrRet;
+}
+
+FString UDynamicDelegateTestObject::FuncRetStr2(int Int, FString Str)
+{
+	Func2(Int, Str);
+	GCalledFunctionName = "FuncRetStr2";
+	return GStrRet;
+}
+
+FString UDynamicDelegateTestObject::FuncRetStr3(int Int, FString Str, const UDynamicDelegateTestObject* Ptr)
+{
+	Func3(Int, Str, Ptr);
+	GCalledFunctionName = "FuncRetStr3";
+	return GStrRet;
+}
+
+FString UDynamicDelegateTestObject::FuncRetStr4(int Int, FString Str, const UDynamicDelegateTestObject* Ptr, FName Name)
+{
+	Func4(Int, Str, Ptr, Name);
+	GCalledFunctionName = "FuncRetStr4";
+	return GStrRet;
+}
+
+FString UDynamicDelegateTestObject::FuncRetStr5(int Int, FString Str, const UDynamicDelegateTestObject* Ptr, FName Name, bool bBool)
+{
+	Func5(Int, Str, Ptr, Name, bBool);
+	GCalledFunctionName = "FuncRetStr5";
+	return GStrRet;
+}
+
+FString UDynamicDelegateTestObject::FuncRetStr6(int Int, FString Str, const UDynamicDelegateTestObject* Ptr, FName Name, bool bBool, TWeakObjectPtr<UDynamicDelegateTestObject> WeakObjPtr)
+{
+	Func6(Int, Str, Ptr, Name, bBool, WeakObjPtr);
+	GCalledFunctionName = "FuncRetStr6";
+	return GStrRet;
+}
+
+int UDynamicDelegateTestObject::ConstFuncRetInt0() const
+{
+	ConstFunc0();
+	GCalledFunctionName = "ConstFuncRetInt0";
+	return GIntRet;
+}
+
+int UDynamicDelegateTestObject::ConstFuncRetInt1(int Int) const
+{
+	ConstFunc1(Int);
+	GCalledFunctionName = "ConstFuncRetInt1";
+	return GIntRet;
+}
+
+int UDynamicDelegateTestObject::ConstFuncRetInt2(int Int, FString Str) const
+{
+	ConstFunc2(Int, Str);
+	GCalledFunctionName = "ConstFuncRetInt2";
+	return GIntRet;
+}
+
+int UDynamicDelegateTestObject::ConstFuncRetInt3(int Int, FString Str, const UDynamicDelegateTestObject* Ptr) const
+{
+	ConstFunc3(Int, Str, Ptr);
+	GCalledFunctionName = "ConstFuncRetInt3";
+	return GIntRet;
+}
+
+int UDynamicDelegateTestObject::ConstFuncRetInt4(int Int, FString Str, const UDynamicDelegateTestObject* Ptr, FName Name) const
+{
+	ConstFunc4(Int, Str, Ptr, Name);
+	GCalledFunctionName = "ConstFuncRetInt4";
+	return GIntRet;
+}
+
+int UDynamicDelegateTestObject::ConstFuncRetInt5(int Int, FString Str, const UDynamicDelegateTestObject* Ptr, FName Name, bool bBool) const
+{
+	ConstFunc5(Int, Str, Ptr, Name, bBool);
+	GCalledFunctionName = "ConstFuncRetInt5";
+	return GIntRet;
+}
+
+int UDynamicDelegateTestObject::ConstFuncRetInt6(int Int, FString Str, const UDynamicDelegateTestObject* Ptr, FName Name, bool bBool, TWeakObjectPtr<UDynamicDelegateTestObject> WeakObjPtr) const
+{
+	ConstFunc6(Int, Str, Ptr, Name, bBool, WeakObjPtr);
+	GCalledFunctionName = "ConstFuncRetInt6";
+	return GIntRet;
+}
+
+FString UDynamicDelegateTestObject::ConstFuncRetStr0() const
+{
+	ConstFunc0();
+	GCalledFunctionName = "ConstFuncRetStr0";
+	return GStrRet;
+}
+
+FString UDynamicDelegateTestObject::ConstFuncRetStr1(int Int) const
+{
+	ConstFunc1(Int);
+	GCalledFunctionName = "ConstFuncRetStr1";
+	return GStrRet;
+}
+
+FString UDynamicDelegateTestObject::ConstFuncRetStr2(int Int, FString Str) const
+{
+	ConstFunc2(Int, Str);
+	GCalledFunctionName = "ConstFuncRetStr2";
+	return GStrRet;
+}
+
+FString UDynamicDelegateTestObject::ConstFuncRetStr3(int Int, FString Str, const UDynamicDelegateTestObject* Ptr) const
+{
+	ConstFunc3(Int, Str, Ptr);
+	GCalledFunctionName = "ConstFuncRetStr3";
+	return GStrRet;
+}
+
+FString UDynamicDelegateTestObject::ConstFuncRetStr4(int Int, FString Str, const UDynamicDelegateTestObject* Ptr, FName Name) const
+{
+	ConstFunc4(Int, Str, Ptr, Name);
+	GCalledFunctionName = "ConstFuncRetStr4";
+	return GStrRet;
+}
+
+FString UDynamicDelegateTestObject::ConstFuncRetStr5(int Int, FString Str, const UDynamicDelegateTestObject* Ptr, FName Name, bool bBool) const
+{
+	ConstFunc5(Int, Str, Ptr, Name, bBool);
+	GCalledFunctionName = "ConstFuncRetStr5";
+	return GStrRet;
+}
+
+FString UDynamicDelegateTestObject::ConstFuncRetStr6(int Int, FString Str, const UDynamicDelegateTestObject* Ptr, FName Name, bool bBool, TWeakObjectPtr<UDynamicDelegateTestObject> WeakObjPtr) const
+{
+	ConstFunc6(Int, Str, Ptr, Name, bBool, WeakObjPtr);
+	GCalledFunctionName = "ConstFuncRetStr6";
+	return GStrRet;
+}
+
+#define BIND_CALL_CHECK_DELEGATE(DelegateType, bMulticast, FuncPtr, CallArgs, ...) /* CallArgs is a parentheses-wrapped set of arguments to invoke the delegate, and the ellipsis is the payload arguments */ \
+	{ \
+		GCalledFunctionName.Reset(); \
+		DelegateType Delegate; \
+		Delegate.UE_IF(bMulticast, AddDynamic, BindDynamic)(this, FuncPtr, ##__VA_ARGS__); \
+		Delegate.UE_IF(bMulticast, Broadcast, Execute) CallArgs; \
+		CHECK(GCalledFunctionName == UE::Core::Private::TrimmedMemberFunctionName<#FuncPtr>.Str); \
+	}
+
+#define BIND_CALL_CHECK_DELEGATE_RETVAL(ExpectedRetVal, DelegateType, bMulticast, FuncPtr, CallArgs, ...) /* CallArgs is a parentheses-wrapped set of arguments to invoke the delegate, and the ellipsis is the payload arguments */ \
+	{ \
+		GCalledFunctionName.Reset(); \
+		DelegateType Delegate; \
+		Delegate.UE_IF(bMulticast, AddDynamic, BindDynamic)(this, FuncPtr, ##__VA_ARGS__); \
+		auto RetVal = Delegate.UE_IF(bMulticast, Broadcast, Execute) CallArgs; \
+		CHECK(RetVal == ExpectedRetVal); \
+		CHECK(GCalledFunctionName == UE::Core::Private::TrimmedMemberFunctionName<#FuncPtr>.Str); \
+	}
+
+void UDynamicDelegateTestObject::RunTests()
+{
+	// Update ptr args, because they can't be given a reasonable default as a global
+	TGuardValue<const UDynamicDelegateTestObject*>          PtrArgGuard(GPtrArg, this);
+	TGuardValue<TWeakObjectPtr<UDynamicDelegateTestObject>> WeakPtrArgGuard(GWeakObjPtrArg, this);
+
+	// 0 param function
+	{
+		// 0 arg payload
+		BIND_CALL_CHECK_DELEGATE       (         FTestDynamicMulticastDelegate0::FDelegate, 0, &UDynamicDelegateTestObject::Func0,            ());
+		BIND_CALL_CHECK_DELEGATE       (         FTestDynamicDelegate0,                     0, &UDynamicDelegateTestObject::Func0,            ());
+		BIND_CALL_CHECK_DELEGATE       (         FTestDynamicMulticastDelegate0,            1, &UDynamicDelegateTestObject::Func0,            ());
+		BIND_CALL_CHECK_DELEGATE       (         FTestDynamicMulticastDelegate0::FDelegate, 0, &UDynamicDelegateTestObject::ConstFunc0,       ());
+		BIND_CALL_CHECK_DELEGATE       (         FTestDynamicDelegate0,                     0, &UDynamicDelegateTestObject::ConstFunc0,       ());
+		BIND_CALL_CHECK_DELEGATE       (         FTestDynamicMulticastDelegate0,            1, &UDynamicDelegateTestObject::ConstFunc0,       ());
+		BIND_CALL_CHECK_DELEGATE_RETVAL(GIntRet, FTestDynamicDelegateRetInt0,               0, &UDynamicDelegateTestObject::FuncRetInt0,      ());
+		BIND_CALL_CHECK_DELEGATE_RETVAL(GIntRet, FTestDynamicDelegateRetInt0,               0, &UDynamicDelegateTestObject::ConstFuncRetInt0, ());
+		BIND_CALL_CHECK_DELEGATE_RETVAL(GStrRet, FTestDynamicDelegateRetStr0,               0, &UDynamicDelegateTestObject::FuncRetStr0,      ());
+		BIND_CALL_CHECK_DELEGATE_RETVAL(GStrRet, FTestDynamicDelegateRetStr0,               0, &UDynamicDelegateTestObject::ConstFuncRetStr0, ());
+	}
+
+	// 1 param function
+	{
+		// 0 arg payload
+		BIND_CALL_CHECK_DELEGATE       (         FTestDynamicMulticastDelegate1::FDelegate, 0, &UDynamicDelegateTestObject::Func1,            (GIntArg));
+		BIND_CALL_CHECK_DELEGATE       (         FTestDynamicDelegate1,                     0, &UDynamicDelegateTestObject::Func1,            (GIntArg));
+		BIND_CALL_CHECK_DELEGATE       (         FTestDynamicMulticastDelegate1,            1, &UDynamicDelegateTestObject::Func1,            (GIntArg));
+		BIND_CALL_CHECK_DELEGATE       (         FTestDynamicMulticastDelegate1::FDelegate, 0, &UDynamicDelegateTestObject::ConstFunc1,       (GIntArg));
+		BIND_CALL_CHECK_DELEGATE       (         FTestDynamicDelegate1,                     0, &UDynamicDelegateTestObject::ConstFunc1,       (GIntArg));
+		BIND_CALL_CHECK_DELEGATE       (         FTestDynamicMulticastDelegate1,            1, &UDynamicDelegateTestObject::ConstFunc1,       (GIntArg));
+		BIND_CALL_CHECK_DELEGATE_RETVAL(GIntRet, FTestDynamicDelegateRetInt1,               0, &UDynamicDelegateTestObject::FuncRetInt1,      (GIntArg));
+		BIND_CALL_CHECK_DELEGATE_RETVAL(GIntRet, FTestDynamicDelegateRetInt1,               0, &UDynamicDelegateTestObject::ConstFuncRetInt1, (GIntArg));
+		BIND_CALL_CHECK_DELEGATE_RETVAL(GStrRet, FTestDynamicDelegateRetStr1,               0, &UDynamicDelegateTestObject::FuncRetStr1,      (GIntArg));
+		BIND_CALL_CHECK_DELEGATE_RETVAL(GStrRet, FTestDynamicDelegateRetStr1,               0, &UDynamicDelegateTestObject::ConstFuncRetStr1, (GIntArg));
+
+#if defined(UE_USE_DYNAMIC_DELEGATE_PAYLOADS) && UE_USE_DYNAMIC_DELEGATE_PAYLOADS
+		// 1 arg payload
+		BIND_CALL_CHECK_DELEGATE       (         FTestDynamicMulticastDelegate0::FDelegate, 0, &UDynamicDelegateTestObject::Func1,            (), GIntArg);
+		BIND_CALL_CHECK_DELEGATE       (         FTestDynamicDelegate0,                     0, &UDynamicDelegateTestObject::Func1,            (), GIntArg);
+		BIND_CALL_CHECK_DELEGATE       (         FTestDynamicMulticastDelegate0,            1, &UDynamicDelegateTestObject::Func1,            (), GIntArg);
+		BIND_CALL_CHECK_DELEGATE       (         FTestDynamicMulticastDelegate0::FDelegate, 0, &UDynamicDelegateTestObject::ConstFunc1,       (), GIntArg);
+		BIND_CALL_CHECK_DELEGATE       (         FTestDynamicDelegate0,                     0, &UDynamicDelegateTestObject::ConstFunc1,       (), GIntArg);
+		BIND_CALL_CHECK_DELEGATE       (         FTestDynamicMulticastDelegate0,            1, &UDynamicDelegateTestObject::ConstFunc1,       (), GIntArg);
+#if UE_USE_DYNAMIC_DELEGATE_PAYLOADS_RETVAL
+		BIND_CALL_CHECK_DELEGATE_RETVAL(GIntRet, FTestDynamicDelegateRetInt0,               0, &UDynamicDelegateTestObject::FuncRetInt1,      (), GIntArg);
+		BIND_CALL_CHECK_DELEGATE_RETVAL(GIntRet, FTestDynamicDelegateRetInt0,               0, &UDynamicDelegateTestObject::ConstFuncRetInt1, (), GIntArg);
+		BIND_CALL_CHECK_DELEGATE_RETVAL(GStrRet, FTestDynamicDelegateRetStr0,               0, &UDynamicDelegateTestObject::FuncRetStr1,      (), GIntArg);
+		BIND_CALL_CHECK_DELEGATE_RETVAL(GStrRet, FTestDynamicDelegateRetStr0,               0, &UDynamicDelegateTestObject::ConstFuncRetStr1, (), GIntArg);
+#endif
+#endif
+	}
+
+	// 2 param function
+	{
+		// 0 arg payload
+		BIND_CALL_CHECK_DELEGATE       (         FTestDynamicMulticastDelegate2::FDelegate, 0, &UDynamicDelegateTestObject::Func2,            (GIntArg, GStrArg));
+		BIND_CALL_CHECK_DELEGATE       (         FTestDynamicDelegate2,                     0, &UDynamicDelegateTestObject::Func2,            (GIntArg, GStrArg));
+		BIND_CALL_CHECK_DELEGATE       (         FTestDynamicMulticastDelegate2,            1, &UDynamicDelegateTestObject::Func2,            (GIntArg, GStrArg));
+		BIND_CALL_CHECK_DELEGATE       (         FTestDynamicMulticastDelegate2::FDelegate, 0, &UDynamicDelegateTestObject::ConstFunc2,       (GIntArg, GStrArg));
+		BIND_CALL_CHECK_DELEGATE       (         FTestDynamicDelegate2,                     0, &UDynamicDelegateTestObject::ConstFunc2,       (GIntArg, GStrArg));
+		BIND_CALL_CHECK_DELEGATE       (         FTestDynamicMulticastDelegate2,            1, &UDynamicDelegateTestObject::ConstFunc2,       (GIntArg, GStrArg));
+		BIND_CALL_CHECK_DELEGATE_RETVAL(GIntRet, FTestDynamicDelegateRetInt2,               0, &UDynamicDelegateTestObject::FuncRetInt2,      (GIntArg, GStrArg));
+		BIND_CALL_CHECK_DELEGATE_RETVAL(GIntRet, FTestDynamicDelegateRetInt2,               0, &UDynamicDelegateTestObject::ConstFuncRetInt2, (GIntArg, GStrArg));
+		BIND_CALL_CHECK_DELEGATE_RETVAL(GStrRet, FTestDynamicDelegateRetStr2,               0, &UDynamicDelegateTestObject::FuncRetStr2,      (GIntArg, GStrArg));
+		BIND_CALL_CHECK_DELEGATE_RETVAL(GStrRet, FTestDynamicDelegateRetStr2,               0, &UDynamicDelegateTestObject::ConstFuncRetStr2, (GIntArg, GStrArg));
+
+#if defined(UE_USE_DYNAMIC_DELEGATE_PAYLOADS) && UE_USE_DYNAMIC_DELEGATE_PAYLOADS
+		// 1 arg payload
+		BIND_CALL_CHECK_DELEGATE       (         FTestDynamicMulticastDelegate1::FDelegate, 0, &UDynamicDelegateTestObject::Func2,            (GIntArg), GStrArg);
+		BIND_CALL_CHECK_DELEGATE       (         FTestDynamicDelegate1,                     0, &UDynamicDelegateTestObject::Func2,            (GIntArg), GStrArg);
+		BIND_CALL_CHECK_DELEGATE       (         FTestDynamicMulticastDelegate1,            1, &UDynamicDelegateTestObject::Func2,            (GIntArg), GStrArg);
+		BIND_CALL_CHECK_DELEGATE       (         FTestDynamicMulticastDelegate1::FDelegate, 0, &UDynamicDelegateTestObject::ConstFunc2,       (GIntArg), GStrArg);
+		BIND_CALL_CHECK_DELEGATE       (         FTestDynamicDelegate1,                     0, &UDynamicDelegateTestObject::ConstFunc2,       (GIntArg), GStrArg);
+		BIND_CALL_CHECK_DELEGATE       (         FTestDynamicMulticastDelegate1,            1, &UDynamicDelegateTestObject::ConstFunc2,       (GIntArg), GStrArg);
+#if UE_USE_DYNAMIC_DELEGATE_PAYLOADS_RETVAL
+		BIND_CALL_CHECK_DELEGATE_RETVAL(GIntRet, FTestDynamicDelegateRetInt1,               0, &UDynamicDelegateTestObject::FuncRetInt2,      (GIntArg), GStrArg);
+		BIND_CALL_CHECK_DELEGATE_RETVAL(GIntRet, FTestDynamicDelegateRetInt1,               0, &UDynamicDelegateTestObject::ConstFuncRetInt2, (GIntArg), GStrArg);
+		BIND_CALL_CHECK_DELEGATE_RETVAL(GStrRet, FTestDynamicDelegateRetStr1,               0, &UDynamicDelegateTestObject::FuncRetStr2,      (GIntArg), GStrArg);
+		BIND_CALL_CHECK_DELEGATE_RETVAL(GStrRet, FTestDynamicDelegateRetStr1,               0, &UDynamicDelegateTestObject::ConstFuncRetStr2, (GIntArg), GStrArg);
+#endif
+
+		// 2 arg payload
+		BIND_CALL_CHECK_DELEGATE       (         FTestDynamicMulticastDelegate0::FDelegate, 0, &UDynamicDelegateTestObject::Func2,            (), GIntArg, GStrArg);
+		BIND_CALL_CHECK_DELEGATE       (         FTestDynamicDelegate0,                     0, &UDynamicDelegateTestObject::Func2,            (), GIntArg, GStrArg);
+		BIND_CALL_CHECK_DELEGATE       (         FTestDynamicMulticastDelegate0,            1, &UDynamicDelegateTestObject::Func2,            (), GIntArg, GStrArg);
+		BIND_CALL_CHECK_DELEGATE       (         FTestDynamicMulticastDelegate0::FDelegate, 0, &UDynamicDelegateTestObject::ConstFunc2,       (), GIntArg, GStrArg);
+		BIND_CALL_CHECK_DELEGATE       (         FTestDynamicDelegate0,                     0, &UDynamicDelegateTestObject::ConstFunc2,       (), GIntArg, GStrArg);
+		BIND_CALL_CHECK_DELEGATE       (         FTestDynamicMulticastDelegate0,            1, &UDynamicDelegateTestObject::ConstFunc2,       (), GIntArg, GStrArg);
+#if UE_USE_DYNAMIC_DELEGATE_PAYLOADS_RETVAL
+		BIND_CALL_CHECK_DELEGATE_RETVAL(GIntRet, FTestDynamicDelegateRetInt0,               0, &UDynamicDelegateTestObject::FuncRetInt2,      (), GIntArg, GStrArg);
+		BIND_CALL_CHECK_DELEGATE_RETVAL(GIntRet, FTestDynamicDelegateRetInt0,               0, &UDynamicDelegateTestObject::ConstFuncRetInt2, (), GIntArg, GStrArg);
+		BIND_CALL_CHECK_DELEGATE_RETVAL(GStrRet, FTestDynamicDelegateRetStr0,               0, &UDynamicDelegateTestObject::FuncRetStr2,      (), GIntArg, GStrArg);
+		BIND_CALL_CHECK_DELEGATE_RETVAL(GStrRet, FTestDynamicDelegateRetStr0,               0, &UDynamicDelegateTestObject::ConstFuncRetStr2, (), GIntArg, GStrArg);
+#endif
+#endif
+	}
+
+	// 3 param function
+	{
+		// 0 arg payload
+		BIND_CALL_CHECK_DELEGATE       (         FTestDynamicMulticastDelegate3::FDelegate, 0, &UDynamicDelegateTestObject::Func3,            (GIntArg, GStrArg, GPtrArg));
+		BIND_CALL_CHECK_DELEGATE       (         FTestDynamicDelegate3,                     0, &UDynamicDelegateTestObject::Func3,            (GIntArg, GStrArg, GPtrArg));
+		BIND_CALL_CHECK_DELEGATE       (         FTestDynamicMulticastDelegate3,            1, &UDynamicDelegateTestObject::Func3,            (GIntArg, GStrArg, GPtrArg));
+		BIND_CALL_CHECK_DELEGATE       (         FTestDynamicMulticastDelegate3::FDelegate, 0, &UDynamicDelegateTestObject::ConstFunc3,       (GIntArg, GStrArg, GPtrArg));
+		BIND_CALL_CHECK_DELEGATE       (         FTestDynamicDelegate3,                     0, &UDynamicDelegateTestObject::ConstFunc3,       (GIntArg, GStrArg, GPtrArg));
+		BIND_CALL_CHECK_DELEGATE       (         FTestDynamicMulticastDelegate3,            1, &UDynamicDelegateTestObject::ConstFunc3,       (GIntArg, GStrArg, GPtrArg));
+		BIND_CALL_CHECK_DELEGATE_RETVAL(GIntRet, FTestDynamicDelegateRetInt3,               0, &UDynamicDelegateTestObject::FuncRetInt3,      (GIntArg, GStrArg, GPtrArg));
+		BIND_CALL_CHECK_DELEGATE_RETVAL(GIntRet, FTestDynamicDelegateRetInt3,               0, &UDynamicDelegateTestObject::ConstFuncRetInt3, (GIntArg, GStrArg, GPtrArg));
+		BIND_CALL_CHECK_DELEGATE_RETVAL(GStrRet, FTestDynamicDelegateRetStr3,               0, &UDynamicDelegateTestObject::FuncRetStr3,      (GIntArg, GStrArg, GPtrArg));
+		BIND_CALL_CHECK_DELEGATE_RETVAL(GStrRet, FTestDynamicDelegateRetStr3,               0, &UDynamicDelegateTestObject::ConstFuncRetStr3, (GIntArg, GStrArg, GPtrArg));
+
+#if defined(UE_USE_DYNAMIC_DELEGATE_PAYLOADS) && UE_USE_DYNAMIC_DELEGATE_PAYLOADS
+		// 1 arg payload
+		BIND_CALL_CHECK_DELEGATE       (         FTestDynamicMulticastDelegate2::FDelegate, 0, &UDynamicDelegateTestObject::Func3,            (GIntArg, GStrArg), GPtrArg);
+		BIND_CALL_CHECK_DELEGATE       (         FTestDynamicDelegate2,                     0, &UDynamicDelegateTestObject::Func3,            (GIntArg, GStrArg), GPtrArg);
+		BIND_CALL_CHECK_DELEGATE       (         FTestDynamicMulticastDelegate2,            1, &UDynamicDelegateTestObject::Func3,            (GIntArg, GStrArg), GPtrArg);
+		BIND_CALL_CHECK_DELEGATE       (         FTestDynamicMulticastDelegate2::FDelegate, 0, &UDynamicDelegateTestObject::ConstFunc3,       (GIntArg, GStrArg), GPtrArg);
+		BIND_CALL_CHECK_DELEGATE       (         FTestDynamicDelegate2,                     0, &UDynamicDelegateTestObject::ConstFunc3,       (GIntArg, GStrArg), GPtrArg);
+		BIND_CALL_CHECK_DELEGATE       (         FTestDynamicMulticastDelegate2,            1, &UDynamicDelegateTestObject::ConstFunc3,       (GIntArg, GStrArg), GPtrArg);
+#if UE_USE_DYNAMIC_DELEGATE_PAYLOADS_RETVAL
+		BIND_CALL_CHECK_DELEGATE_RETVAL(GIntRet, FTestDynamicDelegateRetInt2,               0, &UDynamicDelegateTestObject::FuncRetInt3,      (GIntArg, GStrArg), GPtrArg);
+		BIND_CALL_CHECK_DELEGATE_RETVAL(GIntRet, FTestDynamicDelegateRetInt2,               0, &UDynamicDelegateTestObject::ConstFuncRetInt3, (GIntArg, GStrArg), GPtrArg);
+		BIND_CALL_CHECK_DELEGATE_RETVAL(GStrRet, FTestDynamicDelegateRetStr2,               0, &UDynamicDelegateTestObject::FuncRetStr3,      (GIntArg, GStrArg), GPtrArg);
+		BIND_CALL_CHECK_DELEGATE_RETVAL(GStrRet, FTestDynamicDelegateRetStr2,               0, &UDynamicDelegateTestObject::ConstFuncRetStr3, (GIntArg, GStrArg), GPtrArg);
+#endif
+
+		// 2 arg payload
+		BIND_CALL_CHECK_DELEGATE       (         FTestDynamicMulticastDelegate1::FDelegate, 0, &UDynamicDelegateTestObject::Func3,            (GIntArg), GStrArg, GPtrArg);
+		BIND_CALL_CHECK_DELEGATE       (         FTestDynamicDelegate1,                     0, &UDynamicDelegateTestObject::Func3,            (GIntArg), GStrArg, GPtrArg);
+		BIND_CALL_CHECK_DELEGATE       (         FTestDynamicMulticastDelegate1,            1, &UDynamicDelegateTestObject::Func3,            (GIntArg), GStrArg, GPtrArg);
+		BIND_CALL_CHECK_DELEGATE       (         FTestDynamicMulticastDelegate1::FDelegate, 0, &UDynamicDelegateTestObject::ConstFunc3,       (GIntArg), GStrArg, GPtrArg);
+		BIND_CALL_CHECK_DELEGATE       (         FTestDynamicDelegate1,                     0, &UDynamicDelegateTestObject::ConstFunc3,       (GIntArg), GStrArg, GPtrArg);
+		BIND_CALL_CHECK_DELEGATE       (         FTestDynamicMulticastDelegate1,            1, &UDynamicDelegateTestObject::ConstFunc3,       (GIntArg), GStrArg, GPtrArg);
+#if UE_USE_DYNAMIC_DELEGATE_PAYLOADS_RETVAL
+		BIND_CALL_CHECK_DELEGATE_RETVAL(GIntRet, FTestDynamicDelegateRetInt1,               0, &UDynamicDelegateTestObject::FuncRetInt3,      (GIntArg), GStrArg, GPtrArg);
+		BIND_CALL_CHECK_DELEGATE_RETVAL(GIntRet, FTestDynamicDelegateRetInt1,               0, &UDynamicDelegateTestObject::ConstFuncRetInt3, (GIntArg), GStrArg, GPtrArg);
+		BIND_CALL_CHECK_DELEGATE_RETVAL(GStrRet, FTestDynamicDelegateRetStr1,               0, &UDynamicDelegateTestObject::FuncRetStr3,      (GIntArg), GStrArg, GPtrArg);
+		BIND_CALL_CHECK_DELEGATE_RETVAL(GStrRet, FTestDynamicDelegateRetStr1,               0, &UDynamicDelegateTestObject::ConstFuncRetStr3, (GIntArg), GStrArg, GPtrArg);
+#endif
+
+		// 3 arg payload
+		BIND_CALL_CHECK_DELEGATE       (         FTestDynamicMulticastDelegate0::FDelegate, 0, &UDynamicDelegateTestObject::Func3,            (), GIntArg, GStrArg, GPtrArg);
+		BIND_CALL_CHECK_DELEGATE       (         FTestDynamicDelegate0,                     0, &UDynamicDelegateTestObject::Func3,            (), GIntArg, GStrArg, GPtrArg);
+		BIND_CALL_CHECK_DELEGATE       (         FTestDynamicMulticastDelegate0,            1, &UDynamicDelegateTestObject::Func3,            (), GIntArg, GStrArg, GPtrArg);
+		BIND_CALL_CHECK_DELEGATE       (         FTestDynamicMulticastDelegate0::FDelegate, 0, &UDynamicDelegateTestObject::ConstFunc3,       (), GIntArg, GStrArg, GPtrArg);
+		BIND_CALL_CHECK_DELEGATE       (         FTestDynamicDelegate0,                     0, &UDynamicDelegateTestObject::ConstFunc3,       (), GIntArg, GStrArg, GPtrArg);
+		BIND_CALL_CHECK_DELEGATE       (         FTestDynamicMulticastDelegate0,            1, &UDynamicDelegateTestObject::ConstFunc3,       (), GIntArg, GStrArg, GPtrArg);
+#if UE_USE_DYNAMIC_DELEGATE_PAYLOADS_RETVAL
+		BIND_CALL_CHECK_DELEGATE_RETVAL(GIntRet, FTestDynamicDelegateRetInt0,               0, &UDynamicDelegateTestObject::FuncRetInt3,      (), GIntArg, GStrArg, GPtrArg);
+		BIND_CALL_CHECK_DELEGATE_RETVAL(GIntRet, FTestDynamicDelegateRetInt0,               0, &UDynamicDelegateTestObject::ConstFuncRetInt3, (), GIntArg, GStrArg, GPtrArg);
+		BIND_CALL_CHECK_DELEGATE_RETVAL(GStrRet, FTestDynamicDelegateRetStr0,               0, &UDynamicDelegateTestObject::FuncRetStr3,      (), GIntArg, GStrArg, GPtrArg);
+		BIND_CALL_CHECK_DELEGATE_RETVAL(GStrRet, FTestDynamicDelegateRetStr0,               0, &UDynamicDelegateTestObject::ConstFuncRetStr3, (), GIntArg, GStrArg, GPtrArg);
+#endif
+#endif
+	}
+
+	// 4 param function
+	{
+		// 0 arg payload
+		BIND_CALL_CHECK_DELEGATE       (         FTestDynamicMulticastDelegate4::FDelegate, 0, &UDynamicDelegateTestObject::Func4,            (GIntArg, GStrArg, GPtrArg, GNameArg));
+		BIND_CALL_CHECK_DELEGATE       (         FTestDynamicDelegate4,                     0, &UDynamicDelegateTestObject::Func4,            (GIntArg, GStrArg, GPtrArg, GNameArg));
+		BIND_CALL_CHECK_DELEGATE       (         FTestDynamicMulticastDelegate4,            1, &UDynamicDelegateTestObject::Func4,            (GIntArg, GStrArg, GPtrArg, GNameArg));
+		BIND_CALL_CHECK_DELEGATE       (         FTestDynamicMulticastDelegate4::FDelegate, 0, &UDynamicDelegateTestObject::ConstFunc4,       (GIntArg, GStrArg, GPtrArg, GNameArg));
+		BIND_CALL_CHECK_DELEGATE       (         FTestDynamicDelegate4,                     0, &UDynamicDelegateTestObject::ConstFunc4,       (GIntArg, GStrArg, GPtrArg, GNameArg));
+		BIND_CALL_CHECK_DELEGATE       (         FTestDynamicMulticastDelegate4,            1, &UDynamicDelegateTestObject::ConstFunc4,       (GIntArg, GStrArg, GPtrArg, GNameArg));
+		BIND_CALL_CHECK_DELEGATE_RETVAL(GIntRet, FTestDynamicDelegateRetInt4,               0, &UDynamicDelegateTestObject::FuncRetInt4,      (GIntArg, GStrArg, GPtrArg, GNameArg));
+		BIND_CALL_CHECK_DELEGATE_RETVAL(GIntRet, FTestDynamicDelegateRetInt4,               0, &UDynamicDelegateTestObject::ConstFuncRetInt4, (GIntArg, GStrArg, GPtrArg, GNameArg));
+		BIND_CALL_CHECK_DELEGATE_RETVAL(GStrRet, FTestDynamicDelegateRetStr4,               0, &UDynamicDelegateTestObject::FuncRetStr4,      (GIntArg, GStrArg, GPtrArg, GNameArg));
+		BIND_CALL_CHECK_DELEGATE_RETVAL(GStrRet, FTestDynamicDelegateRetStr4,               0, &UDynamicDelegateTestObject::ConstFuncRetStr4, (GIntArg, GStrArg, GPtrArg, GNameArg));
+
+#if defined(UE_USE_DYNAMIC_DELEGATE_PAYLOADS) && UE_USE_DYNAMIC_DELEGATE_PAYLOADS
+		// 1 arg payload
+		BIND_CALL_CHECK_DELEGATE       (         FTestDynamicMulticastDelegate3::FDelegate, 0, &UDynamicDelegateTestObject::Func4,            (GIntArg, GStrArg, GPtrArg), GNameArg);
+		BIND_CALL_CHECK_DELEGATE       (         FTestDynamicDelegate3,                     0, &UDynamicDelegateTestObject::Func4,            (GIntArg, GStrArg, GPtrArg), GNameArg);
+		BIND_CALL_CHECK_DELEGATE       (         FTestDynamicMulticastDelegate3,            1, &UDynamicDelegateTestObject::Func4,            (GIntArg, GStrArg, GPtrArg), GNameArg);
+		BIND_CALL_CHECK_DELEGATE       (         FTestDynamicMulticastDelegate3::FDelegate, 0, &UDynamicDelegateTestObject::ConstFunc4,       (GIntArg, GStrArg, GPtrArg), GNameArg);
+		BIND_CALL_CHECK_DELEGATE       (         FTestDynamicDelegate3,                     0, &UDynamicDelegateTestObject::ConstFunc4,       (GIntArg, GStrArg, GPtrArg), GNameArg);
+		BIND_CALL_CHECK_DELEGATE       (         FTestDynamicMulticastDelegate3,            1, &UDynamicDelegateTestObject::ConstFunc4,       (GIntArg, GStrArg, GPtrArg), GNameArg);
+#if UE_USE_DYNAMIC_DELEGATE_PAYLOADS_RETVAL
+		BIND_CALL_CHECK_DELEGATE_RETVAL(GIntRet, FTestDynamicDelegateRetInt3,               0, &UDynamicDelegateTestObject::FuncRetInt4,      (GIntArg, GStrArg, GPtrArg), GNameArg);
+		BIND_CALL_CHECK_DELEGATE_RETVAL(GIntRet, FTestDynamicDelegateRetInt3,               0, &UDynamicDelegateTestObject::ConstFuncRetInt4, (GIntArg, GStrArg, GPtrArg), GNameArg);
+		BIND_CALL_CHECK_DELEGATE_RETVAL(GStrRet, FTestDynamicDelegateRetStr3,               0, &UDynamicDelegateTestObject::FuncRetStr4,      (GIntArg, GStrArg, GPtrArg), GNameArg);
+		BIND_CALL_CHECK_DELEGATE_RETVAL(GStrRet, FTestDynamicDelegateRetStr3,               0, &UDynamicDelegateTestObject::ConstFuncRetStr4, (GIntArg, GStrArg, GPtrArg), GNameArg);
+#endif
+
+		// 2 arg payload
+		BIND_CALL_CHECK_DELEGATE       (         FTestDynamicMulticastDelegate2::FDelegate, 0, &UDynamicDelegateTestObject::Func4,            (GIntArg, GStrArg), GPtrArg, GNameArg);
+		BIND_CALL_CHECK_DELEGATE       (         FTestDynamicDelegate2,                     0, &UDynamicDelegateTestObject::Func4,            (GIntArg, GStrArg), GPtrArg, GNameArg);
+		BIND_CALL_CHECK_DELEGATE       (         FTestDynamicMulticastDelegate2,            1, &UDynamicDelegateTestObject::Func4,            (GIntArg, GStrArg), GPtrArg, GNameArg);
+		BIND_CALL_CHECK_DELEGATE       (         FTestDynamicMulticastDelegate2::FDelegate, 0, &UDynamicDelegateTestObject::ConstFunc4,       (GIntArg, GStrArg), GPtrArg, GNameArg);
+		BIND_CALL_CHECK_DELEGATE       (         FTestDynamicDelegate2,                     0, &UDynamicDelegateTestObject::ConstFunc4,       (GIntArg, GStrArg), GPtrArg, GNameArg);
+		BIND_CALL_CHECK_DELEGATE       (         FTestDynamicMulticastDelegate2,            1, &UDynamicDelegateTestObject::ConstFunc4,       (GIntArg, GStrArg), GPtrArg, GNameArg);
+#if UE_USE_DYNAMIC_DELEGATE_PAYLOADS_RETVAL
+		BIND_CALL_CHECK_DELEGATE_RETVAL(GIntRet, FTestDynamicDelegateRetInt2,               0, &UDynamicDelegateTestObject::FuncRetInt4,      (GIntArg, GStrArg), GPtrArg, GNameArg);
+		BIND_CALL_CHECK_DELEGATE_RETVAL(GIntRet, FTestDynamicDelegateRetInt2,               0, &UDynamicDelegateTestObject::ConstFuncRetInt4, (GIntArg, GStrArg), GPtrArg, GNameArg);
+		BIND_CALL_CHECK_DELEGATE_RETVAL(GStrRet, FTestDynamicDelegateRetStr2,               0, &UDynamicDelegateTestObject::FuncRetStr4,      (GIntArg, GStrArg), GPtrArg, GNameArg);
+		BIND_CALL_CHECK_DELEGATE_RETVAL(GStrRet, FTestDynamicDelegateRetStr2,               0, &UDynamicDelegateTestObject::ConstFuncRetStr4, (GIntArg, GStrArg), GPtrArg, GNameArg);
+#endif
+
+		// 3 arg payload
+		BIND_CALL_CHECK_DELEGATE       (         FTestDynamicMulticastDelegate1::FDelegate, 0, &UDynamicDelegateTestObject::Func4,            (GIntArg), GStrArg, GPtrArg, GNameArg);
+		BIND_CALL_CHECK_DELEGATE       (         FTestDynamicDelegate1,                     0, &UDynamicDelegateTestObject::Func4,            (GIntArg), GStrArg, GPtrArg, GNameArg);
+		BIND_CALL_CHECK_DELEGATE       (         FTestDynamicMulticastDelegate1,            1, &UDynamicDelegateTestObject::Func4,            (GIntArg), GStrArg, GPtrArg, GNameArg);
+		BIND_CALL_CHECK_DELEGATE       (         FTestDynamicMulticastDelegate1::FDelegate, 0, &UDynamicDelegateTestObject::ConstFunc4,       (GIntArg), GStrArg, GPtrArg, GNameArg);
+		BIND_CALL_CHECK_DELEGATE       (         FTestDynamicDelegate1,                     0, &UDynamicDelegateTestObject::ConstFunc4,       (GIntArg), GStrArg, GPtrArg, GNameArg);
+		BIND_CALL_CHECK_DELEGATE       (         FTestDynamicMulticastDelegate1,            1, &UDynamicDelegateTestObject::ConstFunc4,       (GIntArg), GStrArg, GPtrArg, GNameArg);
+#if UE_USE_DYNAMIC_DELEGATE_PAYLOADS_RETVAL
+		BIND_CALL_CHECK_DELEGATE_RETVAL(GIntRet, FTestDynamicDelegateRetInt1,               0, &UDynamicDelegateTestObject::FuncRetInt4,      (GIntArg), GStrArg, GPtrArg, GNameArg);
+		BIND_CALL_CHECK_DELEGATE_RETVAL(GIntRet, FTestDynamicDelegateRetInt1,               0, &UDynamicDelegateTestObject::ConstFuncRetInt4, (GIntArg), GStrArg, GPtrArg, GNameArg);
+		BIND_CALL_CHECK_DELEGATE_RETVAL(GStrRet, FTestDynamicDelegateRetStr1,               0, &UDynamicDelegateTestObject::FuncRetStr4,      (GIntArg), GStrArg, GPtrArg, GNameArg);
+		BIND_CALL_CHECK_DELEGATE_RETVAL(GStrRet, FTestDynamicDelegateRetStr1,               0, &UDynamicDelegateTestObject::ConstFuncRetStr4, (GIntArg), GStrArg, GPtrArg, GNameArg);
+#endif
+
+		// 4 arg payload
+		BIND_CALL_CHECK_DELEGATE       (         FTestDynamicMulticastDelegate0::FDelegate, 0, &UDynamicDelegateTestObject::Func4,            (), GIntArg, GStrArg, GPtrArg, GNameArg);
+		BIND_CALL_CHECK_DELEGATE       (         FTestDynamicDelegate0,                     0, &UDynamicDelegateTestObject::Func4,            (), GIntArg, GStrArg, GPtrArg, GNameArg);
+		BIND_CALL_CHECK_DELEGATE       (         FTestDynamicMulticastDelegate0,            1, &UDynamicDelegateTestObject::Func4,            (), GIntArg, GStrArg, GPtrArg, GNameArg);
+		BIND_CALL_CHECK_DELEGATE       (         FTestDynamicMulticastDelegate0::FDelegate, 0, &UDynamicDelegateTestObject::ConstFunc4,       (), GIntArg, GStrArg, GPtrArg, GNameArg);
+		BIND_CALL_CHECK_DELEGATE       (         FTestDynamicDelegate0,                     0, &UDynamicDelegateTestObject::ConstFunc4,       (), GIntArg, GStrArg, GPtrArg, GNameArg);
+		BIND_CALL_CHECK_DELEGATE       (         FTestDynamicMulticastDelegate0,            1, &UDynamicDelegateTestObject::ConstFunc4,       (), GIntArg, GStrArg, GPtrArg, GNameArg);
+#if UE_USE_DYNAMIC_DELEGATE_PAYLOADS_RETVAL
+		BIND_CALL_CHECK_DELEGATE_RETVAL(GIntRet, FTestDynamicDelegateRetInt0,               0, &UDynamicDelegateTestObject::FuncRetInt4,      (), GIntArg, GStrArg, GPtrArg, GNameArg);
+		BIND_CALL_CHECK_DELEGATE_RETVAL(GIntRet, FTestDynamicDelegateRetInt0,               0, &UDynamicDelegateTestObject::ConstFuncRetInt4, (), GIntArg, GStrArg, GPtrArg, GNameArg);
+		BIND_CALL_CHECK_DELEGATE_RETVAL(GStrRet, FTestDynamicDelegateRetStr0,               0, &UDynamicDelegateTestObject::FuncRetStr4,      (), GIntArg, GStrArg, GPtrArg, GNameArg);
+		BIND_CALL_CHECK_DELEGATE_RETVAL(GStrRet, FTestDynamicDelegateRetStr0,               0, &UDynamicDelegateTestObject::ConstFuncRetStr4, (), GIntArg, GStrArg, GPtrArg, GNameArg);
+#endif
+#endif
+	}
+
+	// 5 param function
+	{
+		// 0 arg payload
+		BIND_CALL_CHECK_DELEGATE       (         FTestDynamicMulticastDelegate5::FDelegate, 0, &UDynamicDelegateTestObject::Func5,            (GIntArg, GStrArg, GPtrArg, GNameArg, bGBoolArg));
+		BIND_CALL_CHECK_DELEGATE       (         FTestDynamicDelegate5,                     0, &UDynamicDelegateTestObject::Func5,            (GIntArg, GStrArg, GPtrArg, GNameArg, bGBoolArg));
+		BIND_CALL_CHECK_DELEGATE       (         FTestDynamicMulticastDelegate5,            1, &UDynamicDelegateTestObject::Func5,            (GIntArg, GStrArg, GPtrArg, GNameArg, bGBoolArg));
+		BIND_CALL_CHECK_DELEGATE       (         FTestDynamicMulticastDelegate5::FDelegate, 0, &UDynamicDelegateTestObject::ConstFunc5,       (GIntArg, GStrArg, GPtrArg, GNameArg, bGBoolArg));
+		BIND_CALL_CHECK_DELEGATE       (         FTestDynamicDelegate5,                     0, &UDynamicDelegateTestObject::ConstFunc5,       (GIntArg, GStrArg, GPtrArg, GNameArg, bGBoolArg));
+		BIND_CALL_CHECK_DELEGATE       (         FTestDynamicMulticastDelegate5,            1, &UDynamicDelegateTestObject::ConstFunc5,       (GIntArg, GStrArg, GPtrArg, GNameArg, bGBoolArg));
+		BIND_CALL_CHECK_DELEGATE_RETVAL(GIntRet, FTestDynamicDelegateRetInt5,               0, &UDynamicDelegateTestObject::FuncRetInt5,      (GIntArg, GStrArg, GPtrArg, GNameArg, bGBoolArg));
+		BIND_CALL_CHECK_DELEGATE_RETVAL(GIntRet, FTestDynamicDelegateRetInt5,               0, &UDynamicDelegateTestObject::ConstFuncRetInt5, (GIntArg, GStrArg, GPtrArg, GNameArg, bGBoolArg));
+		BIND_CALL_CHECK_DELEGATE_RETVAL(GStrRet, FTestDynamicDelegateRetStr5,               0, &UDynamicDelegateTestObject::FuncRetStr5,      (GIntArg, GStrArg, GPtrArg, GNameArg, bGBoolArg));
+		BIND_CALL_CHECK_DELEGATE_RETVAL(GStrRet, FTestDynamicDelegateRetStr5,               0, &UDynamicDelegateTestObject::ConstFuncRetStr5, (GIntArg, GStrArg, GPtrArg, GNameArg, bGBoolArg));
+
+#if defined(UE_USE_DYNAMIC_DELEGATE_PAYLOADS) && UE_USE_DYNAMIC_DELEGATE_PAYLOADS
+		// 1 arg payload
+		BIND_CALL_CHECK_DELEGATE       (         FTestDynamicMulticastDelegate4::FDelegate, 0, &UDynamicDelegateTestObject::Func5,            (GIntArg, GStrArg, GPtrArg, GNameArg), bGBoolArg);
+		BIND_CALL_CHECK_DELEGATE       (         FTestDynamicDelegate4,                     0, &UDynamicDelegateTestObject::Func5,            (GIntArg, GStrArg, GPtrArg, GNameArg), bGBoolArg);
+		BIND_CALL_CHECK_DELEGATE       (         FTestDynamicMulticastDelegate4,            1, &UDynamicDelegateTestObject::Func5,            (GIntArg, GStrArg, GPtrArg, GNameArg), bGBoolArg);
+		BIND_CALL_CHECK_DELEGATE       (         FTestDynamicMulticastDelegate4::FDelegate, 0, &UDynamicDelegateTestObject::ConstFunc5,       (GIntArg, GStrArg, GPtrArg, GNameArg), bGBoolArg);
+		BIND_CALL_CHECK_DELEGATE       (         FTestDynamicDelegate4,                     0, &UDynamicDelegateTestObject::ConstFunc5,       (GIntArg, GStrArg, GPtrArg, GNameArg), bGBoolArg);
+		BIND_CALL_CHECK_DELEGATE       (         FTestDynamicMulticastDelegate4,            1, &UDynamicDelegateTestObject::ConstFunc5,       (GIntArg, GStrArg, GPtrArg, GNameArg), bGBoolArg);
+#if UE_USE_DYNAMIC_DELEGATE_PAYLOADS_RETVAL
+		BIND_CALL_CHECK_DELEGATE_RETVAL(GIntRet, FTestDynamicDelegateRetInt4,               0, &UDynamicDelegateTestObject::FuncRetInt5,      (GIntArg, GStrArg, GPtrArg, GNameArg), bGBoolArg);
+		BIND_CALL_CHECK_DELEGATE_RETVAL(GIntRet, FTestDynamicDelegateRetInt4,               0, &UDynamicDelegateTestObject::ConstFuncRetInt5, (GIntArg, GStrArg, GPtrArg, GNameArg), bGBoolArg);
+		BIND_CALL_CHECK_DELEGATE_RETVAL(GStrRet, FTestDynamicDelegateRetStr4,               0, &UDynamicDelegateTestObject::FuncRetStr5,      (GIntArg, GStrArg, GPtrArg, GNameArg), bGBoolArg);
+		BIND_CALL_CHECK_DELEGATE_RETVAL(GStrRet, FTestDynamicDelegateRetStr4,               0, &UDynamicDelegateTestObject::ConstFuncRetStr5, (GIntArg, GStrArg, GPtrArg, GNameArg), bGBoolArg);
+#endif
+
+		// 2 arg payload
+		BIND_CALL_CHECK_DELEGATE       (         FTestDynamicMulticastDelegate3::FDelegate, 0, &UDynamicDelegateTestObject::Func5,            (GIntArg, GStrArg, GPtrArg), GNameArg, bGBoolArg);
+		BIND_CALL_CHECK_DELEGATE       (         FTestDynamicDelegate3,                     0, &UDynamicDelegateTestObject::Func5,            (GIntArg, GStrArg, GPtrArg), GNameArg, bGBoolArg);
+		BIND_CALL_CHECK_DELEGATE       (         FTestDynamicMulticastDelegate3,            1, &UDynamicDelegateTestObject::Func5,            (GIntArg, GStrArg, GPtrArg), GNameArg, bGBoolArg);
+		BIND_CALL_CHECK_DELEGATE       (         FTestDynamicMulticastDelegate3::FDelegate, 0, &UDynamicDelegateTestObject::ConstFunc5,       (GIntArg, GStrArg, GPtrArg), GNameArg, bGBoolArg);
+		BIND_CALL_CHECK_DELEGATE       (         FTestDynamicDelegate3,                     0, &UDynamicDelegateTestObject::ConstFunc5,       (GIntArg, GStrArg, GPtrArg), GNameArg, bGBoolArg);
+		BIND_CALL_CHECK_DELEGATE       (         FTestDynamicMulticastDelegate3,            1, &UDynamicDelegateTestObject::ConstFunc5,       (GIntArg, GStrArg, GPtrArg), GNameArg, bGBoolArg);
+#if UE_USE_DYNAMIC_DELEGATE_PAYLOADS_RETVAL
+		BIND_CALL_CHECK_DELEGATE_RETVAL(GIntRet, FTestDynamicDelegateRetInt3,               0, &UDynamicDelegateTestObject::FuncRetInt5,      (GIntArg, GStrArg, GPtrArg), GNameArg, bGBoolArg);
+		BIND_CALL_CHECK_DELEGATE_RETVAL(GIntRet, FTestDynamicDelegateRetInt3,               0, &UDynamicDelegateTestObject::ConstFuncRetInt5, (GIntArg, GStrArg, GPtrArg), GNameArg, bGBoolArg);
+		BIND_CALL_CHECK_DELEGATE_RETVAL(GStrRet, FTestDynamicDelegateRetStr3,               0, &UDynamicDelegateTestObject::FuncRetStr5,      (GIntArg, GStrArg, GPtrArg), GNameArg, bGBoolArg);
+		BIND_CALL_CHECK_DELEGATE_RETVAL(GStrRet, FTestDynamicDelegateRetStr3,               0, &UDynamicDelegateTestObject::ConstFuncRetStr5, (GIntArg, GStrArg, GPtrArg), GNameArg, bGBoolArg);
+#endif
+
+		// 3 arg payload
+		BIND_CALL_CHECK_DELEGATE       (         FTestDynamicMulticastDelegate2::FDelegate, 0, &UDynamicDelegateTestObject::Func5,            (GIntArg, GStrArg), GPtrArg, GNameArg, bGBoolArg);
+		BIND_CALL_CHECK_DELEGATE       (         FTestDynamicDelegate2,                     0, &UDynamicDelegateTestObject::Func5,            (GIntArg, GStrArg), GPtrArg, GNameArg, bGBoolArg);
+		BIND_CALL_CHECK_DELEGATE       (         FTestDynamicMulticastDelegate2,            1, &UDynamicDelegateTestObject::Func5,            (GIntArg, GStrArg), GPtrArg, GNameArg, bGBoolArg);
+		BIND_CALL_CHECK_DELEGATE       (         FTestDynamicMulticastDelegate2::FDelegate, 0, &UDynamicDelegateTestObject::ConstFunc5,       (GIntArg, GStrArg), GPtrArg, GNameArg, bGBoolArg);
+		BIND_CALL_CHECK_DELEGATE       (         FTestDynamicDelegate2,                     0, &UDynamicDelegateTestObject::ConstFunc5,       (GIntArg, GStrArg), GPtrArg, GNameArg, bGBoolArg);
+		BIND_CALL_CHECK_DELEGATE       (         FTestDynamicMulticastDelegate2,            1, &UDynamicDelegateTestObject::ConstFunc5,       (GIntArg, GStrArg), GPtrArg, GNameArg, bGBoolArg);
+#if UE_USE_DYNAMIC_DELEGATE_PAYLOADS_RETVAL
+		BIND_CALL_CHECK_DELEGATE_RETVAL(GIntRet, FTestDynamicDelegateRetInt2,               0, &UDynamicDelegateTestObject::FuncRetInt5,      (GIntArg, GStrArg), GPtrArg, GNameArg, bGBoolArg);
+		BIND_CALL_CHECK_DELEGATE_RETVAL(GIntRet, FTestDynamicDelegateRetInt2,               0, &UDynamicDelegateTestObject::ConstFuncRetInt5, (GIntArg, GStrArg), GPtrArg, GNameArg, bGBoolArg);
+		BIND_CALL_CHECK_DELEGATE_RETVAL(GStrRet, FTestDynamicDelegateRetStr2,               0, &UDynamicDelegateTestObject::FuncRetStr5,      (GIntArg, GStrArg), GPtrArg, GNameArg, bGBoolArg);
+		BIND_CALL_CHECK_DELEGATE_RETVAL(GStrRet, FTestDynamicDelegateRetStr2,               0, &UDynamicDelegateTestObject::ConstFuncRetStr5, (GIntArg, GStrArg), GPtrArg, GNameArg, bGBoolArg);
+#endif
+
+		// 4 arg payload
+		BIND_CALL_CHECK_DELEGATE       (         FTestDynamicMulticastDelegate1::FDelegate, 0, &UDynamicDelegateTestObject::Func5,            (GIntArg), GStrArg, GPtrArg, GNameArg, bGBoolArg);
+		BIND_CALL_CHECK_DELEGATE       (         FTestDynamicDelegate1,                     0, &UDynamicDelegateTestObject::Func5,            (GIntArg), GStrArg, GPtrArg, GNameArg, bGBoolArg);
+		BIND_CALL_CHECK_DELEGATE       (         FTestDynamicMulticastDelegate1,            1, &UDynamicDelegateTestObject::Func5,            (GIntArg), GStrArg, GPtrArg, GNameArg, bGBoolArg);
+		BIND_CALL_CHECK_DELEGATE       (         FTestDynamicMulticastDelegate1::FDelegate, 0, &UDynamicDelegateTestObject::ConstFunc5,       (GIntArg), GStrArg, GPtrArg, GNameArg, bGBoolArg);
+		BIND_CALL_CHECK_DELEGATE       (         FTestDynamicDelegate1,                     0, &UDynamicDelegateTestObject::ConstFunc5,       (GIntArg), GStrArg, GPtrArg, GNameArg, bGBoolArg);
+		BIND_CALL_CHECK_DELEGATE       (         FTestDynamicMulticastDelegate1,            1, &UDynamicDelegateTestObject::ConstFunc5,       (GIntArg), GStrArg, GPtrArg, GNameArg, bGBoolArg);
+#if UE_USE_DYNAMIC_DELEGATE_PAYLOADS_RETVAL
+		BIND_CALL_CHECK_DELEGATE_RETVAL(GIntRet, FTestDynamicDelegateRetInt1,               0, &UDynamicDelegateTestObject::FuncRetInt5,      (GIntArg), GStrArg, GPtrArg, GNameArg, bGBoolArg);
+		BIND_CALL_CHECK_DELEGATE_RETVAL(GIntRet, FTestDynamicDelegateRetInt1,               0, &UDynamicDelegateTestObject::ConstFuncRetInt5, (GIntArg), GStrArg, GPtrArg, GNameArg, bGBoolArg);
+		BIND_CALL_CHECK_DELEGATE_RETVAL(GStrRet, FTestDynamicDelegateRetStr1,               0, &UDynamicDelegateTestObject::FuncRetStr5,      (GIntArg), GStrArg, GPtrArg, GNameArg, bGBoolArg);
+		BIND_CALL_CHECK_DELEGATE_RETVAL(GStrRet, FTestDynamicDelegateRetStr1,               0, &UDynamicDelegateTestObject::ConstFuncRetStr5, (GIntArg), GStrArg, GPtrArg, GNameArg, bGBoolArg);
+#endif
+
+		// 5 arg payload
+		BIND_CALL_CHECK_DELEGATE       (         FTestDynamicMulticastDelegate0::FDelegate, 0, &UDynamicDelegateTestObject::Func5,            (), GIntArg, GStrArg, GPtrArg, GNameArg, bGBoolArg);
+		BIND_CALL_CHECK_DELEGATE       (         FTestDynamicDelegate0,                     0, &UDynamicDelegateTestObject::Func5,            (), GIntArg, GStrArg, GPtrArg, GNameArg, bGBoolArg);
+		BIND_CALL_CHECK_DELEGATE       (         FTestDynamicMulticastDelegate0,            1, &UDynamicDelegateTestObject::Func5,            (), GIntArg, GStrArg, GPtrArg, GNameArg, bGBoolArg);
+		BIND_CALL_CHECK_DELEGATE       (         FTestDynamicMulticastDelegate0::FDelegate, 0, &UDynamicDelegateTestObject::ConstFunc5,       (), GIntArg, GStrArg, GPtrArg, GNameArg, bGBoolArg);
+		BIND_CALL_CHECK_DELEGATE       (         FTestDynamicDelegate0,                     0, &UDynamicDelegateTestObject::ConstFunc5,       (), GIntArg, GStrArg, GPtrArg, GNameArg, bGBoolArg);
+		BIND_CALL_CHECK_DELEGATE       (         FTestDynamicMulticastDelegate0,            1, &UDynamicDelegateTestObject::ConstFunc5,       (), GIntArg, GStrArg, GPtrArg, GNameArg, bGBoolArg);
+#if UE_USE_DYNAMIC_DELEGATE_PAYLOADS_RETVAL
+		BIND_CALL_CHECK_DELEGATE_RETVAL(GIntRet, FTestDynamicDelegateRetInt0,               0, &UDynamicDelegateTestObject::FuncRetInt5,      (), GIntArg, GStrArg, GPtrArg, GNameArg, bGBoolArg);
+		BIND_CALL_CHECK_DELEGATE_RETVAL(GIntRet, FTestDynamicDelegateRetInt0,               0, &UDynamicDelegateTestObject::ConstFuncRetInt5, (), GIntArg, GStrArg, GPtrArg, GNameArg, bGBoolArg);
+		BIND_CALL_CHECK_DELEGATE_RETVAL(GStrRet, FTestDynamicDelegateRetStr0,               0, &UDynamicDelegateTestObject::FuncRetStr5,      (), GIntArg, GStrArg, GPtrArg, GNameArg, bGBoolArg);
+		BIND_CALL_CHECK_DELEGATE_RETVAL(GStrRet, FTestDynamicDelegateRetStr0,               0, &UDynamicDelegateTestObject::ConstFuncRetStr5, (), GIntArg, GStrArg, GPtrArg, GNameArg, bGBoolArg);
+#endif
+#endif
+	}
+
+	// 6 param function
+	{
+		// 0 arg payload
+		BIND_CALL_CHECK_DELEGATE       (         FTestDynamicMulticastDelegate6::FDelegate, 0, &UDynamicDelegateTestObject::Func6,            (GIntArg, GStrArg, GPtrArg, GNameArg, bGBoolArg, GWeakObjPtrArg));
+		BIND_CALL_CHECK_DELEGATE       (         FTestDynamicDelegate6,                     0, &UDynamicDelegateTestObject::Func6,            (GIntArg, GStrArg, GPtrArg, GNameArg, bGBoolArg, GWeakObjPtrArg));
+		BIND_CALL_CHECK_DELEGATE       (         FTestDynamicMulticastDelegate6,            1, &UDynamicDelegateTestObject::Func6,            (GIntArg, GStrArg, GPtrArg, GNameArg, bGBoolArg, GWeakObjPtrArg));
+		BIND_CALL_CHECK_DELEGATE       (         FTestDynamicMulticastDelegate6::FDelegate, 0, &UDynamicDelegateTestObject::ConstFunc6,       (GIntArg, GStrArg, GPtrArg, GNameArg, bGBoolArg, GWeakObjPtrArg));
+		BIND_CALL_CHECK_DELEGATE       (         FTestDynamicDelegate6,                     0, &UDynamicDelegateTestObject::ConstFunc6,       (GIntArg, GStrArg, GPtrArg, GNameArg, bGBoolArg, GWeakObjPtrArg));
+		BIND_CALL_CHECK_DELEGATE       (         FTestDynamicMulticastDelegate6,            1, &UDynamicDelegateTestObject::ConstFunc6,       (GIntArg, GStrArg, GPtrArg, GNameArg, bGBoolArg, GWeakObjPtrArg));
+		BIND_CALL_CHECK_DELEGATE_RETVAL(GIntRet, FTestDynamicDelegateRetInt6,               0, &UDynamicDelegateTestObject::FuncRetInt6,      (GIntArg, GStrArg, GPtrArg, GNameArg, bGBoolArg, GWeakObjPtrArg));
+		BIND_CALL_CHECK_DELEGATE_RETVAL(GIntRet, FTestDynamicDelegateRetInt6,               0, &UDynamicDelegateTestObject::ConstFuncRetInt6, (GIntArg, GStrArg, GPtrArg, GNameArg, bGBoolArg, GWeakObjPtrArg));
+		BIND_CALL_CHECK_DELEGATE_RETVAL(GStrRet, FTestDynamicDelegateRetStr6,               0, &UDynamicDelegateTestObject::FuncRetStr6,      (GIntArg, GStrArg, GPtrArg, GNameArg, bGBoolArg, GWeakObjPtrArg));
+		BIND_CALL_CHECK_DELEGATE_RETVAL(GStrRet, FTestDynamicDelegateRetStr6,               0, &UDynamicDelegateTestObject::ConstFuncRetStr6, (GIntArg, GStrArg, GPtrArg, GNameArg, bGBoolArg, GWeakObjPtrArg));
+
+#if defined(UE_USE_DYNAMIC_DELEGATE_PAYLOADS) && UE_USE_DYNAMIC_DELEGATE_PAYLOADS
+		// 1 arg payload
+		BIND_CALL_CHECK_DELEGATE       (         FTestDynamicMulticastDelegate5::FDelegate, 0, &UDynamicDelegateTestObject::Func6,            (GIntArg, GStrArg, GPtrArg, GNameArg, bGBoolArg), GWeakObjPtrArg);
+		BIND_CALL_CHECK_DELEGATE       (         FTestDynamicDelegate5,                     0, &UDynamicDelegateTestObject::Func6,            (GIntArg, GStrArg, GPtrArg, GNameArg, bGBoolArg), GWeakObjPtrArg);
+		BIND_CALL_CHECK_DELEGATE       (         FTestDynamicMulticastDelegate5,            1, &UDynamicDelegateTestObject::Func6,            (GIntArg, GStrArg, GPtrArg, GNameArg, bGBoolArg), GWeakObjPtrArg);
+		BIND_CALL_CHECK_DELEGATE       (         FTestDynamicMulticastDelegate5::FDelegate, 0, &UDynamicDelegateTestObject::ConstFunc6,       (GIntArg, GStrArg, GPtrArg, GNameArg, bGBoolArg), GWeakObjPtrArg);
+		BIND_CALL_CHECK_DELEGATE       (         FTestDynamicDelegate5,                     0, &UDynamicDelegateTestObject::ConstFunc6,       (GIntArg, GStrArg, GPtrArg, GNameArg, bGBoolArg), GWeakObjPtrArg);
+		BIND_CALL_CHECK_DELEGATE       (         FTestDynamicMulticastDelegate5,            1, &UDynamicDelegateTestObject::ConstFunc6,       (GIntArg, GStrArg, GPtrArg, GNameArg, bGBoolArg), GWeakObjPtrArg);
+#if UE_USE_DYNAMIC_DELEGATE_PAYLOADS_RETVAL
+		BIND_CALL_CHECK_DELEGATE_RETVAL(GIntRet, FTestDynamicDelegateRetInt5,               0, &UDynamicDelegateTestObject::FuncRetInt6,      (GIntArg, GStrArg, GPtrArg, GNameArg, bGBoolArg), GWeakObjPtrArg);
+		BIND_CALL_CHECK_DELEGATE_RETVAL(GIntRet, FTestDynamicDelegateRetInt5,               0, &UDynamicDelegateTestObject::ConstFuncRetInt6, (GIntArg, GStrArg, GPtrArg, GNameArg, bGBoolArg), GWeakObjPtrArg);
+		BIND_CALL_CHECK_DELEGATE_RETVAL(GStrRet, FTestDynamicDelegateRetStr5,               0, &UDynamicDelegateTestObject::FuncRetStr6,      (GIntArg, GStrArg, GPtrArg, GNameArg, bGBoolArg), GWeakObjPtrArg);
+		BIND_CALL_CHECK_DELEGATE_RETVAL(GStrRet, FTestDynamicDelegateRetStr5,               0, &UDynamicDelegateTestObject::ConstFuncRetStr6, (GIntArg, GStrArg, GPtrArg, GNameArg, bGBoolArg), GWeakObjPtrArg);
+#endif
+
+		// 2 arg payload
+		BIND_CALL_CHECK_DELEGATE       (         FTestDynamicMulticastDelegate4::FDelegate, 0, &UDynamicDelegateTestObject::Func6,            (GIntArg, GStrArg, GPtrArg, GNameArg), bGBoolArg, GWeakObjPtrArg);
+		BIND_CALL_CHECK_DELEGATE       (         FTestDynamicDelegate4,                     0, &UDynamicDelegateTestObject::Func6,            (GIntArg, GStrArg, GPtrArg, GNameArg), bGBoolArg, GWeakObjPtrArg);
+		BIND_CALL_CHECK_DELEGATE       (         FTestDynamicMulticastDelegate4,            1, &UDynamicDelegateTestObject::Func6,            (GIntArg, GStrArg, GPtrArg, GNameArg), bGBoolArg, GWeakObjPtrArg);
+		BIND_CALL_CHECK_DELEGATE       (         FTestDynamicMulticastDelegate4::FDelegate, 0, &UDynamicDelegateTestObject::ConstFunc6,       (GIntArg, GStrArg, GPtrArg, GNameArg), bGBoolArg, GWeakObjPtrArg);
+		BIND_CALL_CHECK_DELEGATE       (         FTestDynamicDelegate4,                     0, &UDynamicDelegateTestObject::ConstFunc6,       (GIntArg, GStrArg, GPtrArg, GNameArg), bGBoolArg, GWeakObjPtrArg);
+		BIND_CALL_CHECK_DELEGATE       (         FTestDynamicMulticastDelegate4,            1, &UDynamicDelegateTestObject::ConstFunc6,       (GIntArg, GStrArg, GPtrArg, GNameArg), bGBoolArg, GWeakObjPtrArg);
+#if UE_USE_DYNAMIC_DELEGATE_PAYLOADS_RETVAL
+		BIND_CALL_CHECK_DELEGATE_RETVAL(GIntRet, FTestDynamicDelegateRetInt4,               0, &UDynamicDelegateTestObject::FuncRetInt6,      (GIntArg, GStrArg, GPtrArg, GNameArg), bGBoolArg, GWeakObjPtrArg);
+		BIND_CALL_CHECK_DELEGATE_RETVAL(GIntRet, FTestDynamicDelegateRetInt4,               0, &UDynamicDelegateTestObject::ConstFuncRetInt6, (GIntArg, GStrArg, GPtrArg, GNameArg), bGBoolArg, GWeakObjPtrArg);
+		BIND_CALL_CHECK_DELEGATE_RETVAL(GStrRet, FTestDynamicDelegateRetStr4,               0, &UDynamicDelegateTestObject::FuncRetStr6,      (GIntArg, GStrArg, GPtrArg, GNameArg), bGBoolArg, GWeakObjPtrArg);
+		BIND_CALL_CHECK_DELEGATE_RETVAL(GStrRet, FTestDynamicDelegateRetStr4,               0, &UDynamicDelegateTestObject::ConstFuncRetStr6, (GIntArg, GStrArg, GPtrArg, GNameArg), bGBoolArg, GWeakObjPtrArg);
+#endif
+
+		// 3 arg payload
+		BIND_CALL_CHECK_DELEGATE       (         FTestDynamicMulticastDelegate3::FDelegate, 0, &UDynamicDelegateTestObject::Func6,            (GIntArg, GStrArg, GPtrArg), GNameArg, bGBoolArg, GWeakObjPtrArg);
+		BIND_CALL_CHECK_DELEGATE       (         FTestDynamicDelegate3,                     0, &UDynamicDelegateTestObject::Func6,            (GIntArg, GStrArg, GPtrArg), GNameArg, bGBoolArg, GWeakObjPtrArg);
+		BIND_CALL_CHECK_DELEGATE       (         FTestDynamicMulticastDelegate3,            1, &UDynamicDelegateTestObject::Func6,            (GIntArg, GStrArg, GPtrArg), GNameArg, bGBoolArg, GWeakObjPtrArg);
+		BIND_CALL_CHECK_DELEGATE       (         FTestDynamicMulticastDelegate3::FDelegate, 0, &UDynamicDelegateTestObject::ConstFunc6,       (GIntArg, GStrArg, GPtrArg), GNameArg, bGBoolArg, GWeakObjPtrArg);
+		BIND_CALL_CHECK_DELEGATE       (         FTestDynamicDelegate3,                     0, &UDynamicDelegateTestObject::ConstFunc6,       (GIntArg, GStrArg, GPtrArg), GNameArg, bGBoolArg, GWeakObjPtrArg);
+		BIND_CALL_CHECK_DELEGATE       (         FTestDynamicMulticastDelegate3,            1, &UDynamicDelegateTestObject::ConstFunc6,       (GIntArg, GStrArg, GPtrArg), GNameArg, bGBoolArg, GWeakObjPtrArg);
+#if UE_USE_DYNAMIC_DELEGATE_PAYLOADS_RETVAL
+		BIND_CALL_CHECK_DELEGATE_RETVAL(GIntRet, FTestDynamicDelegateRetInt3,               0, &UDynamicDelegateTestObject::FuncRetInt6,      (GIntArg, GStrArg, GPtrArg), GNameArg, bGBoolArg, GWeakObjPtrArg);
+		BIND_CALL_CHECK_DELEGATE_RETVAL(GIntRet, FTestDynamicDelegateRetInt3,               0, &UDynamicDelegateTestObject::ConstFuncRetInt6, (GIntArg, GStrArg, GPtrArg), GNameArg, bGBoolArg, GWeakObjPtrArg);
+		BIND_CALL_CHECK_DELEGATE_RETVAL(GStrRet, FTestDynamicDelegateRetStr3,               0, &UDynamicDelegateTestObject::FuncRetStr6,      (GIntArg, GStrArg, GPtrArg), GNameArg, bGBoolArg, GWeakObjPtrArg);
+		BIND_CALL_CHECK_DELEGATE_RETVAL(GStrRet, FTestDynamicDelegateRetStr3,               0, &UDynamicDelegateTestObject::ConstFuncRetStr6, (GIntArg, GStrArg, GPtrArg), GNameArg, bGBoolArg, GWeakObjPtrArg);
+#endif
+
+		// 4 arg payload
+		BIND_CALL_CHECK_DELEGATE       (         FTestDynamicMulticastDelegate2::FDelegate, 0, &UDynamicDelegateTestObject::Func6,            (GIntArg, GStrArg), GPtrArg, GNameArg, bGBoolArg, GWeakObjPtrArg);
+		BIND_CALL_CHECK_DELEGATE       (         FTestDynamicDelegate2,                     0, &UDynamicDelegateTestObject::Func6,            (GIntArg, GStrArg), GPtrArg, GNameArg, bGBoolArg, GWeakObjPtrArg);
+		BIND_CALL_CHECK_DELEGATE       (         FTestDynamicMulticastDelegate2,            1, &UDynamicDelegateTestObject::Func6,            (GIntArg, GStrArg), GPtrArg, GNameArg, bGBoolArg, GWeakObjPtrArg);
+		BIND_CALL_CHECK_DELEGATE       (         FTestDynamicMulticastDelegate2::FDelegate, 0, &UDynamicDelegateTestObject::ConstFunc6,       (GIntArg, GStrArg), GPtrArg, GNameArg, bGBoolArg, GWeakObjPtrArg);
+		BIND_CALL_CHECK_DELEGATE       (         FTestDynamicDelegate2,                     0, &UDynamicDelegateTestObject::ConstFunc6,       (GIntArg, GStrArg), GPtrArg, GNameArg, bGBoolArg, GWeakObjPtrArg);
+		BIND_CALL_CHECK_DELEGATE       (         FTestDynamicMulticastDelegate2,            1, &UDynamicDelegateTestObject::ConstFunc6,       (GIntArg, GStrArg), GPtrArg, GNameArg, bGBoolArg, GWeakObjPtrArg);
+#if UE_USE_DYNAMIC_DELEGATE_PAYLOADS_RETVAL
+		BIND_CALL_CHECK_DELEGATE_RETVAL(GIntRet, FTestDynamicDelegateRetInt2,               0, &UDynamicDelegateTestObject::FuncRetInt6,      (GIntArg, GStrArg), GPtrArg, GNameArg, bGBoolArg, GWeakObjPtrArg);
+		BIND_CALL_CHECK_DELEGATE_RETVAL(GIntRet, FTestDynamicDelegateRetInt2,               0, &UDynamicDelegateTestObject::ConstFuncRetInt6, (GIntArg, GStrArg), GPtrArg, GNameArg, bGBoolArg, GWeakObjPtrArg);
+		BIND_CALL_CHECK_DELEGATE_RETVAL(GStrRet, FTestDynamicDelegateRetStr2,               0, &UDynamicDelegateTestObject::FuncRetStr6,      (GIntArg, GStrArg), GPtrArg, GNameArg, bGBoolArg, GWeakObjPtrArg);
+		BIND_CALL_CHECK_DELEGATE_RETVAL(GStrRet, FTestDynamicDelegateRetStr2,               0, &UDynamicDelegateTestObject::ConstFuncRetStr6, (GIntArg, GStrArg), GPtrArg, GNameArg, bGBoolArg, GWeakObjPtrArg);
+#endif
+
+		// 5 arg payload
+		BIND_CALL_CHECK_DELEGATE       (         FTestDynamicMulticastDelegate1::FDelegate, 0, &UDynamicDelegateTestObject::Func6,            (GIntArg), GStrArg, GPtrArg, GNameArg, bGBoolArg, GWeakObjPtrArg);
+		BIND_CALL_CHECK_DELEGATE       (         FTestDynamicDelegate1,                     0, &UDynamicDelegateTestObject::Func6,            (GIntArg), GStrArg, GPtrArg, GNameArg, bGBoolArg, GWeakObjPtrArg);
+		BIND_CALL_CHECK_DELEGATE       (         FTestDynamicMulticastDelegate1,            1, &UDynamicDelegateTestObject::Func6,            (GIntArg), GStrArg, GPtrArg, GNameArg, bGBoolArg, GWeakObjPtrArg);
+		BIND_CALL_CHECK_DELEGATE       (         FTestDynamicMulticastDelegate1::FDelegate, 0, &UDynamicDelegateTestObject::ConstFunc6,       (GIntArg), GStrArg, GPtrArg, GNameArg, bGBoolArg, GWeakObjPtrArg);
+		BIND_CALL_CHECK_DELEGATE       (         FTestDynamicDelegate1,                     0, &UDynamicDelegateTestObject::ConstFunc6,       (GIntArg), GStrArg, GPtrArg, GNameArg, bGBoolArg, GWeakObjPtrArg);
+		BIND_CALL_CHECK_DELEGATE       (         FTestDynamicMulticastDelegate1,            1, &UDynamicDelegateTestObject::ConstFunc6,       (GIntArg), GStrArg, GPtrArg, GNameArg, bGBoolArg, GWeakObjPtrArg);
+#if UE_USE_DYNAMIC_DELEGATE_PAYLOADS_RETVAL
+		BIND_CALL_CHECK_DELEGATE_RETVAL(GIntRet, FTestDynamicDelegateRetInt1,               0, &UDynamicDelegateTestObject::FuncRetInt6,      (GIntArg), GStrArg, GPtrArg, GNameArg, bGBoolArg, GWeakObjPtrArg);
+		BIND_CALL_CHECK_DELEGATE_RETVAL(GIntRet, FTestDynamicDelegateRetInt1,               0, &UDynamicDelegateTestObject::ConstFuncRetInt6, (GIntArg), GStrArg, GPtrArg, GNameArg, bGBoolArg, GWeakObjPtrArg);
+		BIND_CALL_CHECK_DELEGATE_RETVAL(GStrRet, FTestDynamicDelegateRetStr1,               0, &UDynamicDelegateTestObject::FuncRetStr6,      (GIntArg), GStrArg, GPtrArg, GNameArg, bGBoolArg, GWeakObjPtrArg);
+		BIND_CALL_CHECK_DELEGATE_RETVAL(GStrRet, FTestDynamicDelegateRetStr1,               0, &UDynamicDelegateTestObject::ConstFuncRetStr6, (GIntArg), GStrArg, GPtrArg, GNameArg, bGBoolArg, GWeakObjPtrArg);
+#endif
+
+		// 6 arg payload
+		BIND_CALL_CHECK_DELEGATE       (         FTestDynamicMulticastDelegate0::FDelegate, 0, &UDynamicDelegateTestObject::Func6,            (), GIntArg, GStrArg, GPtrArg, GNameArg, bGBoolArg, GWeakObjPtrArg);
+		BIND_CALL_CHECK_DELEGATE       (         FTestDynamicDelegate0,                     0, &UDynamicDelegateTestObject::Func6,            (), GIntArg, GStrArg, GPtrArg, GNameArg, bGBoolArg, GWeakObjPtrArg);
+		BIND_CALL_CHECK_DELEGATE       (         FTestDynamicMulticastDelegate0,            1, &UDynamicDelegateTestObject::Func6,            (), GIntArg, GStrArg, GPtrArg, GNameArg, bGBoolArg, GWeakObjPtrArg);
+		BIND_CALL_CHECK_DELEGATE       (         FTestDynamicMulticastDelegate0::FDelegate, 0, &UDynamicDelegateTestObject::ConstFunc6,       (), GIntArg, GStrArg, GPtrArg, GNameArg, bGBoolArg, GWeakObjPtrArg);
+		BIND_CALL_CHECK_DELEGATE       (         FTestDynamicDelegate0,                     0, &UDynamicDelegateTestObject::ConstFunc6,       (), GIntArg, GStrArg, GPtrArg, GNameArg, bGBoolArg, GWeakObjPtrArg);
+		BIND_CALL_CHECK_DELEGATE       (         FTestDynamicMulticastDelegate0,            1, &UDynamicDelegateTestObject::ConstFunc6,       (), GIntArg, GStrArg, GPtrArg, GNameArg, bGBoolArg, GWeakObjPtrArg);
+#if UE_USE_DYNAMIC_DELEGATE_PAYLOADS_RETVAL
+		BIND_CALL_CHECK_DELEGATE_RETVAL(GIntRet, FTestDynamicDelegateRetInt0,               0, &UDynamicDelegateTestObject::FuncRetInt6,      (), GIntArg, GStrArg, GPtrArg, GNameArg, bGBoolArg, GWeakObjPtrArg);
+		BIND_CALL_CHECK_DELEGATE_RETVAL(GIntRet, FTestDynamicDelegateRetInt0,               0, &UDynamicDelegateTestObject::ConstFuncRetInt6, (), GIntArg, GStrArg, GPtrArg, GNameArg, bGBoolArg, GWeakObjPtrArg);
+		BIND_CALL_CHECK_DELEGATE_RETVAL(GStrRet, FTestDynamicDelegateRetStr0,               0, &UDynamicDelegateTestObject::FuncRetStr6,      (), GIntArg, GStrArg, GPtrArg, GNameArg, bGBoolArg, GWeakObjPtrArg);
+		BIND_CALL_CHECK_DELEGATE_RETVAL(GStrRet, FTestDynamicDelegateRetStr0,               0, &UDynamicDelegateTestObject::ConstFuncRetStr6, (), GIntArg, GStrArg, GPtrArg, GNameArg, bGBoolArg, GWeakObjPtrArg);
+#endif
+#endif
+	}
+}
+
+TEST_CASE_NAMED(DynamicDelegateTest, "UE::CoreUObject::DynamicDelegateTest", "[ApplicationContextMask][SmokeFilter]")
+{
+	UDynamicDelegateTestObject* TestObj = NewObject<UDynamicDelegateTestObject>();
+	TestObj->RunTests();
+}
+
+#undef BIND_CALL_CHECK_DELEGATE_RETVAL
+#undef BIND_CALL_CHECK_DELEGATE
+
+#endif // #if WITH_TESTS
diff --git a/Engine/Source/Runtime/CoreUObject/Public/UObject/DynamicDelegateTestObject.h b/Engine/Source/Runtime/CoreUObject/Public/UObject/DynamicDelegateTestObject.h
new file mode 100644
--- /dev/null
+++ b/Engine/Source/Runtime/CoreUObject/Public/UObject/DynamicDelegateTestObject.h
@@ -0,0 +1,106 @@
+// Copyright Epic Games, Inc. All Rights Reserved.
+
+#pragma once
+
+#if WITH_TESTS
+
+#include "UObject/Object.h"
+#include "UObject/ScriptMacros.h"
+
+#include "DynamicDelegateTestObject.generated.h"
+
+DECLARE_DYNAMIC_DELEGATE            (FTestDynamicDelegate0);
+DECLARE_DYNAMIC_DELEGATE_OneParam   (FTestDynamicDelegate1, int, Int);
+DECLARE_DYNAMIC_DELEGATE_TwoParams  (FTestDynamicDelegate2, int, Int, FString, Str);
+DECLARE_DYNAMIC_DELEGATE_ThreeParams(FTestDynamicDelegate3, int, Int, FString, Str, const UDynamicDelegateTestObject*, Ptr);
+DECLARE_DYNAMIC_DELEGATE_FourParams (FTestDynamicDelegate4, int, Int, FString, Str, const UDynamicDelegateTestObject*, Ptr, FName, Name);
+DECLARE_DYNAMIC_DELEGATE_FiveParams (FTestDynamicDelegate5, int, Int, FString, Str, const UDynamicDelegateTestObject*, Ptr, FName, Name, bool, bBool);
+DECLARE_DYNAMIC_DELEGATE_SixParams  (FTestDynamicDelegate6, int, Int, FString, Str, const UDynamicDelegateTestObject*, Ptr, FName, Name, bool, bBool, TWeakObjectPtr<UDynamicDelegateTestObject>, WeakPtr);
+
+DECLARE_DYNAMIC_MULTICAST_DELEGATE            (FTestDynamicMulticastDelegate0);
+DECLARE_DYNAMIC_MULTICAST_DELEGATE_OneParam   (FTestDynamicMulticastDelegate1, int, Int);
+DECLARE_DYNAMIC_MULTICAST_DELEGATE_TwoParams  (FTestDynamicMulticastDelegate2, int, Int, FString, Str);
+DECLARE_DYNAMIC_MULTICAST_DELEGATE_ThreeParams(FTestDynamicMulticastDelegate3, int, Int, FString, Str, const UDynamicDelegateTestObject*, Ptr);
+DECLARE_DYNAMIC_MULTICAST_DELEGATE_FourParams (FTestDynamicMulticastDelegate4, int, Int, FString, Str, const UDynamicDelegateTestObject*, Ptr, FName, Name);
+DECLARE_DYNAMIC_MULTICAST_DELEGATE_FiveParams (FTestDynamicMulticastDelegate5, int, Int, FString, Str, const UDynamicDelegateTestObject*, Ptr, FName, Name, bool, bBool);
+DECLARE_DYNAMIC_MULTICAST_DELEGATE_SixParams  (FTestDynamicMulticastDelegate6, int, Int, FString, Str, const UDynamicDelegateTestObject*, Ptr, FName, Name, bool, bBool, TWeakObjectPtr<UDynamicDelegateTestObject>, WeakPtr);
+
+DECLARE_DYNAMIC_DELEGATE_RetVal            (int, FTestDynamicDelegateRetInt0);
+DECLARE_DYNAMIC_DELEGATE_RetVal_OneParam   (int, FTestDynamicDelegateRetInt1, int, Int);
+DECLARE_DYNAMIC_DELEGATE_RetVal_TwoParams  (int, FTestDynamicDelegateRetInt2, int, Int, FString, Str);
+DECLARE_DYNAMIC_DELEGATE_RetVal_ThreeParams(int, FTestDynamicDelegateRetInt3, int, Int, FString, Str, const UDynamicDelegateTestObject*, Ptr);
+DECLARE_DYNAMIC_DELEGATE_RetVal_FourParams (int, FTestDynamicDelegateRetInt4, int, Int, FString, Str, const UDynamicDelegateTestObject*, Ptr, FName, Name);
+DECLARE_DYNAMIC_DELEGATE_RetVal_FiveParams (int, FTestDynamicDelegateRetInt5, int, Int, FString, Str, const UDynamicDelegateTestObject*, Ptr, FName, Name, bool, bBool);
+DECLARE_DYNAMIC_DELEGATE_RetVal_SixParams  (int, FTestDynamicDelegateRetInt6, int, Int, FString, Str, const UDynamicDelegateTestObject*, Ptr, FName, Name, bool, bBool, TWeakObjectPtr<UDynamicDelegateTestObject>, WeakPtr);
+
+DECLARE_DYNAMIC_DELEGATE_RetVal            (FString, FTestDynamicDelegateRetStr0);
+DECLARE_DYNAMIC_DELEGATE_RetVal_OneParam   (FString, FTestDynamicDelegateRetStr1, int, Int);
+DECLARE_DYNAMIC_DELEGATE_RetVal_TwoParams  (FString, FTestDynamicDelegateRetStr2, int, Int, FString, Str);
+DECLARE_DYNAMIC_DELEGATE_RetVal_ThreeParams(FString, FTestDynamicDelegateRetStr3, int, Int, FString, Str, const UDynamicDelegateTestObject*, Ptr);
+DECLARE_DYNAMIC_DELEGATE_RetVal_FourParams (FString, FTestDynamicDelegateRetStr4, int, Int, FString, Str, const UDynamicDelegateTestObject*, Ptr, FName, Name);
+DECLARE_DYNAMIC_DELEGATE_RetVal_FiveParams (FString, FTestDynamicDelegateRetStr5, int, Int, FString, Str, const UDynamicDelegateTestObject*, Ptr, FName, Name, bool, bBool);
+DECLARE_DYNAMIC_DELEGATE_RetVal_SixParams  (FString, FTestDynamicDelegateRetStr6, int, Int, FString, Str, const UDynamicDelegateTestObject*, Ptr, FName, Name, bool, bBool, TWeakObjectPtr<UDynamicDelegateTestObject>, WeakPtr);
+
+UCLASS()
+class UDynamicDelegateTestObject : public UObject
+{
+	GENERATED_BODY()
+
+public:
+	UDynamicDelegateTestObject();
+
+	UFUNCTION() void Func0();
+	UFUNCTION() void Func1(int Int);
+	UFUNCTION() void Func2(int Int, FString Str);
+	UFUNCTION() void Func3(int Int, FString Str, const UDynamicDelegateTestObject* Ptr);
+	UFUNCTION() void Func4(int Int, FString Str, const UDynamicDelegateTestObject* Ptr, FName Name);
+	UFUNCTION() void Func5(int Int, FString Str, const UDynamicDelegateTestObject* Ptr, FName Name, bool bBool);
+	UFUNCTION() void Func6(int Int, FString Str, const UDynamicDelegateTestObject* Ptr, FName Name, bool bBool, TWeakObjectPtr<UDynamicDelegateTestObject> WeakPtr);
+
+	UFUNCTION() void ConstFunc0() const;
+	UFUNCTION() void ConstFunc1(int Int) const;
+	UFUNCTION() void ConstFunc2(int Int, FString Str) const;
+	UFUNCTION() void ConstFunc3(int Int, FString Str, const UDynamicDelegateTestObject* Ptr) const;
+	UFUNCTION() void ConstFunc4(int Int, FString Str, const UDynamicDelegateTestObject* Ptr, FName Name) const;
+	UFUNCTION() void ConstFunc5(int Int, FString Str, const UDynamicDelegateTestObject* Ptr, FName Name, bool bBool) const;
+	UFUNCTION() void ConstFunc6(int Int, FString Str, const UDynamicDelegateTestObject* Ptr, FName Name, bool bBool, TWeakObjectPtr<UDynamicDelegateTestObject> WeakPtr) const;
+
+	UFUNCTION() int FuncRetInt0();
+	UFUNCTION() int FuncRetInt1(int Int);
+	UFUNCTION() int FuncRetInt2(int Int, FString Str);
+	UFUNCTION() int FuncRetInt3(int Int, FString Str, const UDynamicDelegateTestObject* Ptr);
+	UFUNCTION() int FuncRetInt4(int Int, FString Str, const UDynamicDelegateTestObject* Ptr, FName Name);
+	UFUNCTION() int FuncRetInt5(int Int, FString Str, const UDynamicDelegateTestObject* Ptr, FName Name, bool bBool);
+	UFUNCTION() int FuncRetInt6(int Int, FString Str, const UDynamicDelegateTestObject* Ptr, FName Name, bool bBool, TWeakObjectPtr<UDynamicDelegateTestObject> WeakPtr);
+
+	UFUNCTION() FString FuncRetStr0();
+	UFUNCTION() FString FuncRetStr1(int Int);
+	UFUNCTION() FString FuncRetStr2(int Int, FString Str);
+	UFUNCTION() FString FuncRetStr3(int Int, FString Str, const UDynamicDelegateTestObject* Ptr);
+	UFUNCTION() FString FuncRetStr4(int Int, FString Str, const UDynamicDelegateTestObject* Ptr, FName Name);
+	UFUNCTION() FString FuncRetStr5(int Int, FString Str, const UDynamicDelegateTestObject* Ptr, FName Name, bool bBool);
+	UFUNCTION() FString FuncRetStr6(int Int, FString Str, const UDynamicDelegateTestObject* Ptr, FName Name, bool bBool, TWeakObjectPtr<UDynamicDelegateTestObject> WeakPtr);
+
+	UFUNCTION() int ConstFuncRetInt0() const;
+	UFUNCTION() int ConstFuncRetInt1(int Int) const;
+	UFUNCTION() int ConstFuncRetInt2(int Int, FString Str) const;
+	UFUNCTION() int ConstFuncRetInt3(int Int, FString Str, const UDynamicDelegateTestObject* Ptr) const;
+	UFUNCTION() int ConstFuncRetInt4(int Int, FString Str, const UDynamicDelegateTestObject* Ptr, FName Name) const;
+	UFUNCTION() int ConstFuncRetInt5(int Int, FString Str, const UDynamicDelegateTestObject* Ptr, FName Name, bool bBool) const;
+	UFUNCTION() int ConstFuncRetInt6(int Int, FString Str, const UDynamicDelegateTestObject* Ptr, FName Name, bool bBool, TWeakObjectPtr<UDynamicDelegateTestObject> WeakPtr) const;
+
+	UFUNCTION() FString ConstFuncRetStr0() const;
+	UFUNCTION() FString ConstFuncRetStr1(int Int) const;
+	UFUNCTION() FString ConstFuncRetStr2(int Int, FString Str) const;
+	UFUNCTION() FString ConstFuncRetStr3(int Int, FString Str, const UDynamicDelegateTestObject* Ptr) const;
+	UFUNCTION() FString ConstFuncRetStr4(int Int, FString Str, const UDynamicDelegateTestObject* Ptr, FName Name) const;
+	UFUNCTION() FString ConstFuncRetStr5(int Int, FString Str, const UDynamicDelegateTestObject* Ptr, FName Name, bool bBool) const;
+	UFUNCTION() FString ConstFuncRetStr6(int Int, FString Str, const UDynamicDelegateTestObject* Ptr, FName Name, bool bBool, TWeakObjectPtr<UDynamicDelegateTestObject> WeakPtr) const;
+
+	void RunTests();
+};
+
+#endif // WITH_TESTS
+
+// Forward declaration still needed for .gen.cpp when building without tests
+class UDynamicDelegateTestObject;
