[Core] Adjust weakmap serialization to align better with sets/maps so we don't make changes that breaks hashing/lookup

#tests Local testing and try to get some smoke testing done before commit
#Synced-CL 49715754

diff --git a/Engine/Source/Runtime/CoreUObject/Private/Serialization/ArchiveUObject.cpp b/Engine/Source/Runtime/CoreUObject/Private/Serialization/ArchiveUObject.cpp
--- a/Engine/Source/Runtime/CoreUObject/Private/Serialization/ArchiveUObject.cpp
+++ b/Engine/Source/Runtime/CoreUObject/Private/Serialization/ArchiveUObject.cpp
@@ -143,7 +143,7 @@ FArchive& FArchiveUObject::SerializeWeakObjectPtr(FArchive& Ar, FWeakObjectPtr&
 	
 		Ar << Object;
 	
-		if (Ar.IsLoading() || Ar.IsModifyingWeakAndStrongReferences())
+		if (Ar.IsLoading() || (Ar.IsModifyingWeakAndStrongReferences() && !Ar.IsSaving()))
 		{
 			Value = Object;
 		}
