- Skip calls to LoadObject when unloading a GFP

#Synced-CL 49643405

diff --git a/Engine/Plugins/Runtime/GameFeatures/Source/GameFeatures/Private/GameFeaturePluginStateMachine.cpp b/Engine/Plugins/Runtime/GameFeatures/Source/GameFeatures/Private/GameFeaturePluginStateMachine.cpp
--- a/Engine/Plugins/Runtime/GameFeatures/Source/GameFeatures/Private/GameFeaturePluginStateMachine.cpp
+++ b/Engine/Plugins/Runtime/GameFeatures/Source/GameFeatures/Private/GameFeaturePluginStateMachine.cpp
@@ -3786,6 +3786,8 @@ struct FGameFeaturePluginState_Deactivating : public FGameFeaturePluginState
 			FName Tag = *StateProperties.PluginName;
 			UE::DynamicConfig::PerformDynamicConfig(Tag, [Tag](FConfigModificationTracker* ChangeTracker)
 			{
+				// note that we are unloading the GFP so we don't call LoadObject iuuselessly
+				ChangeTracker->bIsUnloadingPlugin = true;
 				FConfigCacheIni::RemoveTagFromAllBranches(Tag, ChangeTracker);
 				IConsoleManager::Get().UnsetAllConsoleVariablesWithTag(Tag);
 			});
diff --git a/Engine/Source/Runtime/Core/Public/Misc/ConfigCacheIni.h b/Engine/Source/Runtime/Core/Public/Misc/ConfigCacheIni.h
--- a/Engine/Source/Runtime/Core/Public/Misc/ConfigCacheIni.h
+++ b/Engine/Source/Runtime/Core/Public/Misc/ConfigCacheIni.h
@@ -515,6 +515,7 @@ public:
 	// input
 	bool bTrackModifiedSections = true;
 	bool bTrackLoadedFiles = false;
+	bool bIsUnloadingPlugin = false;
 	// if this is true and bTrackLoadedFiles is set, then we ignore files in LoadedFiles array
 	bool bSkipLoadingAlreadyLoadedFiles = false;
 
diff --git a/Engine/Source/Runtime/CoreUObject/Private/UObject/UObjectBaseUtility.cpp b/Engine/Source/Runtime/CoreUObject/Private/UObject/UObjectBaseUtility.cpp
--- a/Engine/Source/Runtime/CoreUObject/Private/UObject/UObjectBaseUtility.cpp
+++ b/Engine/Source/Runtime/CoreUObject/Private/UObject/UObjectBaseUtility.cpp
@@ -545,7 +545,10 @@ void UObjectBaseUtility::ReloadObjectsFromModifiedConfigSections(const FConfigMo
 {	
 	TArray<const UClass*> ClassesToReload;
 	TArray<UObject*> ObjectsToReload;
-	
+
+	static const FString NativeClassPrefix("/Script/");
+	static const FString BPClassPrefix("/Game/");
+
 	// Reload configs so objects get the changes
 	for (const TPair<FName, TSet<FString>>& Pair : ChangeTracker->ModifiedSectionsPerBranch)
 	{
@@ -564,7 +567,7 @@ void UObjectBaseUtility::ReloadObjectsFromModifiedConfigSections(const FConfigMo
 				const FString ClassName = SectionName.Mid(PerObjConfigDelimIdx + 1);
 				
 				// TryFindTypeSlow will throw a warning/callstack for short pathnames, so use the function it calls internally
-				UClass* ObjClass = (UClass*)StaticFindFirstObject(UClass::StaticClass(), *ClassName, EFindFirstObjectOptions::EnsureIfAmbiguous | EFindFirstObjectOptions::NativeFirst, ELogVerbosity::Error, TEXT("UGameFeatureData::ReloadConfigs"));
+				UClass* ObjClass = (UClass*)StaticFindFirstObject(UClass::StaticClass(), *ClassName, EFindFirstObjectOptions::EnsureIfAmbiguous | EFindFirstObjectOptions::NativeFirst, ELogVerbosity::Error, TEXT("UObjectBaseUtility::ReloadObjectsFromModifiedConfigSections"));
 				if (ObjClass && ObjClass->HasAnyClassFlags(CLASS_Config))
 				{
 					// Now try to actually find the object it's referencing specifically and update it
@@ -583,12 +586,10 @@ void UObjectBaseUtility::ReloadObjectsFromModifiedConfigSections(const FConfigMo
 			// Standard INI section case
 			else
 			{
-				const FString NativeClassPrefix("/Script/");
-				const FString BPClassPrefix("/Game/");
-				
 				UClass* ObjClass = nullptr;
 				// handle /Script and /Game (/Game is for BP classes)
-				if (SectionName.StartsWith(NativeClassPrefix))
+				// if we are unloading, neverf call LoadObject as it can cause problems to load as we unload - likely we never need to call LoadObject, but minimizing side effects
+				if (SectionName.StartsWith(NativeClassPrefix) || (ChangeTracker->bIsUnloadingPlugin && SectionName.StartsWith(BPClassPrefix)))
 				{
 					ObjClass = FindFirstObject<UClass>(*SectionName, EFindFirstObjectOptions::ExactClass | EFindFirstObjectOptions::EnsureIfAmbiguous | EFindFirstObjectOptions::NativeFirst);
 				}
