[ChunkDependency] Move UChunkDependencyInfo from UnrealEd to CoreUObject to make it accessible in IoStoreUtilities for shaders.

#rnx
#Synced-CL 49742993

diff --git a/Engine/Plugins/Runtime/GameFeatures/Source/GameFeatures/Private/GameFeatureAction_AddChunkOverride.cpp b/Engine/Plugins/Runtime/GameFeatures/Source/GameFeatures/Private/GameFeatureAction_AddChunkOverride.cpp
--- a/Engine/Plugins/Runtime/GameFeatures/Source/GameFeatures/Private/GameFeatureAction_AddChunkOverride.cpp
+++ b/Engine/Plugins/Runtime/GameFeatures/Source/GameFeatures/Private/GameFeatureAction_AddChunkOverride.cpp
@@ -8,7 +8,7 @@
 #include "Misc/PathViews.h"
 
 #if WITH_EDITOR
-#include "Commandlets/ChunkDependencyInfo.h"
+#include "Misc/ChunkDependencyInfo.h"
 #endif
 
 #include UE_INLINE_GENERATED_CPP_BY_NAME(GameFeatureAction_AddChunkOverride)
diff --git a/Engine/Source/Editor/UnrealEd/Private/Commandlets/AssetRegistryGenerator.cpp b/Engine/Source/Editor/UnrealEd/Private/Commandlets/AssetRegistryGenerator.cpp
--- a/Engine/Source/Editor/UnrealEd/Private/Commandlets/AssetRegistryGenerator.cpp
+++ b/Engine/Source/Editor/UnrealEd/Private/Commandlets/AssetRegistryGenerator.cpp
@@ -15,7 +15,6 @@
 #include "Cooker/CookWorkerClient.h"
 #include "Cooker/MPCollector.h"
 #include "CookMetadata.h"
-#include "Commandlets/ChunkDependencyInfo.h"
 #include "Commandlets/IChunkDataGenerator.h"
 #include "Containers/RingBuffer.h"
 #include "Engine/AssetManager.h"
@@ -31,6 +30,7 @@
 #include "Interfaces/ITargetPlatform.h"
 #include "Misc/App.h"
 #include "Misc/AsciiSet.h"
+#include "Misc/ChunkDependencyInfo.h"
 #include "Misc/ConfigCacheIni.h"
 #include "Misc/DataDrivenPlatformInfoRegistry.h"
 #include "Misc/FileHelper.h"
diff --git a/Engine/Source/Editor/UnrealEd/Private/Commandlets/ChunkDependencyInfo.cpp b/Engine/Source/Editor/UnrealEd/Private/Commandlets/ChunkDependencyInfo.cpp
deleted file mode 100644
diff --git a/Engine/Source/Editor/UnrealEd/Private/Tests/ChunkDependencyTests.cpp b/Engine/Source/Editor/UnrealEd/Private/Tests/ChunkDependencyTests.cpp
--- a/Engine/Source/Editor/UnrealEd/Private/Tests/ChunkDependencyTests.cpp
+++ b/Engine/Source/Editor/UnrealEd/Private/Tests/ChunkDependencyTests.cpp
@@ -4,7 +4,7 @@
 
 #if WITH_AUTOMATION_TESTS
 
-#include "Commandlets/ChunkDependencyInfo.h"
+#include "Misc/ChunkDependencyInfo.h"
 
 IMPLEMENT_SIMPLE_AUTOMATION_TEST(FChunkDependencyHighestSharedTest, "System.Core.ChunkDependency.HighestSharedChunk",
 	EAutomationTestFlags::EditorContext | EAutomationTestFlags::EngineFilter)
diff --git a/Engine/Source/Editor/UnrealEd/Public/Commandlets/ChunkDependencyInfo.h b/Engine/Source/Editor/UnrealEd/Public/Commandlets/ChunkDependencyInfo.h
--- a/Engine/Source/Editor/UnrealEd/Public/Commandlets/ChunkDependencyInfo.h
+++ b/Engine/Source/Editor/UnrealEd/Public/Commandlets/ChunkDependencyInfo.h
@@ -2,100 +2,7 @@
 
 #pragma once
 
-#include "Containers/Array.h"
-#include "Containers/Map.h"
-#include "Containers/Set.h"
-#include "CoreMinimal.h"
-#include "HAL/Platform.h"
-#include "HAL/PlatformCrt.h"
-#include "UObject/Object.h"
-#include "UObject/ObjectMacros.h"
-#include "UObject/UObjectGlobals.h"
+UE_DEPRECATED_HEADER(5.8, "Include Misc/ChunkDependencyInfo.h from the CoreUObject module instead.")
 
-#include "ChunkDependencyInfo.generated.h"
-
-/** A single dependency, read from ini file */
-USTRUCT()
-struct FChunkDependency
-{
-	GENERATED_USTRUCT_BODY()
-
-	FChunkDependency()
-	{};
-
-	FChunkDependency(int32 InChunkId, int32 InParentChunkId)
-		: ChunkID(InChunkId)
-		, ParentChunkID(InParentChunkId)
-	{};
-
-	/** The child chunk */
-	UPROPERTY(EditAnywhere, Category = ChunkInfo)
-	int32 ChunkID = 0;
-
-	/** Parent chunk, anything in both Parent and Child is only placed into Parent */
-	UPROPERTY(EditAnywhere, Category = ChunkInfo)
-	int32 ParentChunkID = 0;
-
-	bool operator== (const FChunkDependency& RHS) const
-	{
-		return ChunkID == RHS.ChunkID;
-	}
-};
-
-/** In memory structure used for dependency tree */
-struct FChunkDependencyTreeNode
-{
-	FChunkDependencyTreeNode(int32 InChunkID=0)
-		: ChunkID(InChunkID)
-	{}
-
-	int32 ChunkID;
-
-	TArray<FChunkDependencyTreeNode> ChildNodes;
-};
-
-/** This is read out of config and defines a tree of chunk dependencies */
-UCLASS(config=Engine, defaultconfig, MinimalAPI)
-class UChunkDependencyInfo : public UObject
-{
-	GENERATED_UCLASS_BODY()
-
-	/** Will return an existing dependency graph, or build if it necessary */
-	UNREALED_API const FChunkDependencyTreeNode* GetOrBuildChunkDependencyGraph(int32 HighestChunk = 0, bool bForceRebuild = false);
-
-	/** Will create a dependency tree starting with RootTreeNode. If HighestChunk is == 0 it will only add dependencies on 0 for chunks already in the dependencies list */
-	UNREALED_API const FChunkDependencyTreeNode* BuildChunkDependencyGraph(int32 HighestChunk);
-
-	/** Removes redundant chunks from a chunk list */
-	UNREALED_API void RemoveRedundantChunks(TArray<int32>& ChunkIDs) const;
-
-	/** Given a set of chunk Ids will use the dependencies to find which chunk is shared among all the given chunks.
-	All chunks should converge to the 0 startup chunk.
-	Potentially will return a chunk that is not in the given input array */
-	UNREALED_API int32 FindHighestSharedChunk(const TArray<int32>& ChunkIDs) const;
-
-	/**
-	 * Returns the chunks that the in chunk depends upon through the parent rules.
-	 */
-	UNREALED_API void GetChunkDependencies(const int32 InChunk, TSet<int32>& OutChunkDependencies) const;
-
-	/** List of dependencies used to remove redundant chunks */
-	UPROPERTY(config)
-	TArray<FChunkDependency> DependencyArray;
-
-private:
-	/** Fills out the Dependency Tree starting with Node */
-	void AddChildrenRecursive(FChunkDependencyTreeNode& Node, TArray<FChunkDependency>& DepInfo, TSet<int32> Parents);
-
-	/** Root of tree, this will be valid after calling BuildChunkDependencyGraph */
-	FChunkDependencyTreeNode RootTreeNode;
-
-	/** Map of child chunks to all parent chunks, computed in BuildChunkDependencyGraph */
-	TMap<int32, TSet<int32>> ChildToParentMap;
-
-	/** Cached array of topologically sorted chunk. */
-	TArray<int32> TopologicallySortedChunks;
-
-	/** Cached value of HighestChunk at time of building */
-	int32 CachedHighestChunk;
-};
+// This file is deprecated since UE 5.8. Use Misc/ChunkDependencyInfo.h instead!
+#include "Misc/ChunkDependencyInfo.h"
diff --git a/Engine/Source/Runtime/CoreUObject/Private/Misc/ChunkDependencyInfo.cpp b/Engine/Source/Runtime/CoreUObject/Private/Misc/ChunkDependencyInfo.cpp
new file mode 100644
--- /dev/null
+++ b/Engine/Source/Runtime/CoreUObject/Private/Misc/ChunkDependencyInfo.cpp
@@ -0,0 +1,221 @@
+// Copyright Epic Games, Inc. All Rights Reserved.
+
+#include "Misc/ChunkDependencyInfo.h"
+#include "Algo/Sort.h"
+#include "Algo/Unique.h"
+#include "Containers/Queue.h"
+
+#include UE_INLINE_GENERATED_CPP_BY_NAME(ChunkDependencyInfo)
+
+DEFINE_LOG_CATEGORY_STATIC(LogChunkDependencyInfo, Log, All);
+
+namespace ChunkDependencyInfo
+{
+	// breadth first traversal
+	static TArray<int32> BuildTopologicallySortedArray(const FChunkDependencyTreeNode* RootNode)
+	{
+		TArray<int32> OutArray;
+		if (!RootNode)
+		{
+			return OutArray;
+		}
+		TQueue<const FChunkDependencyTreeNode*> Nodes;
+		Nodes.Enqueue(RootNode);
+
+		TSet<const FChunkDependencyTreeNode*> ProcessedNodes;
+		const FChunkDependencyTreeNode* CurrentNode = nullptr;
+		while (!Nodes.IsEmpty())
+		{
+			Nodes.Dequeue(CurrentNode);
+			if (ProcessedNodes.Contains(CurrentNode))
+			{
+				continue;
+			}
+			ProcessedNodes.Add(CurrentNode);
+			OutArray.Add(CurrentNode->ChunkID);
+			for (const FChunkDependencyTreeNode& ChildNode : CurrentNode->ChildNodes)
+			{
+				Nodes.Enqueue(&ChildNode);
+			}
+		}
+
+		return OutArray;
+	}
+}
+
+UChunkDependencyInfo::UChunkDependencyInfo(const FObjectInitializer& ObjectInitializer)
+	: Super(ObjectInitializer)
+{
+	CachedHighestChunk = -1;
+}
+
+const FChunkDependencyTreeNode* UChunkDependencyInfo::GetOrBuildChunkDependencyGraph(int32 HighestChunk, bool bForceRebuild)
+{
+	if (HighestChunk > CachedHighestChunk)
+	{
+		return BuildChunkDependencyGraph(HighestChunk);
+	}
+	else if (bForceRebuild)
+	{
+		return BuildChunkDependencyGraph(CachedHighestChunk);
+	}
+	return &RootTreeNode;
+}
+
+const FChunkDependencyTreeNode* UChunkDependencyInfo::BuildChunkDependencyGraph(int32 HighestChunk)
+{
+	// Reset any current tree
+	RootTreeNode.ChunkID = 0;
+	RootTreeNode.ChildNodes.Reset(0);
+
+	ChildToParentMap.Reset();
+	ChildToParentMap.Add(0);
+	CachedHighestChunk = HighestChunk;
+
+	// Ensure the DependencyArray is OK to work with.
+	for (int32 DepIndex = DependencyArray.Num() - 1; DepIndex >= 0; DepIndex --)
+	{
+		const FChunkDependency& Dep = DependencyArray[DepIndex];
+		if (Dep.ChunkID > HighestChunk)
+		{
+			HighestChunk = Dep.ChunkID;
+		}
+		if (Dep.ParentChunkID > HighestChunk)
+		{
+			HighestChunk = Dep.ParentChunkID;
+		}
+		if (Dep.ChunkID == Dep.ParentChunkID)
+		{
+			// Remove cycles
+			DependencyArray.RemoveAtSwap(DepIndex);
+		}
+	}
+	// Add missing links (assumes they parent to chunk zero)
+	for (int32 i = 1; i <= HighestChunk; ++i)
+	{
+		if (!DependencyArray.FindByPredicate([=](const FChunkDependency& RHS){ return i == RHS.ChunkID; }))
+		{
+			FChunkDependency Dep;
+			Dep.ChunkID = i;
+			Dep.ParentChunkID = 0;
+			DependencyArray.Add(Dep);
+		}
+	}
+	// Remove duplicates
+	DependencyArray.StableSort([](const FChunkDependency& LHS, const FChunkDependency& RHS) { return LHS.ChunkID < RHS.ChunkID; });
+	for (int32 i = 0; i < DependencyArray.Num() - 1;)
+	{
+		if (DependencyArray[i] == DependencyArray[i + 1])
+		{
+			DependencyArray.RemoveAt(i + 1);
+		}
+		else
+		{
+			++i;
+		}
+	}
+
+	AddChildrenRecursive(RootTreeNode, DependencyArray, TSet<int32>());
+	TopologicallySortedChunks = ChunkDependencyInfo::BuildTopologicallySortedArray(&RootTreeNode);
+	return &RootTreeNode;
+}
+
+void UChunkDependencyInfo::AddChildrenRecursive(FChunkDependencyTreeNode& Node, TArray<FChunkDependency>& DepInfo, TSet<int32> Parents)
+{
+	if (Parents.Num() > 0)
+	{
+		ChildToParentMap.FindOrAdd(Node.ChunkID).Append(Parents);
+	}
+
+	Parents.Add(Node.ChunkID);
+	auto ChildNodeIndices = DepInfo.FilterByPredicate(
+		[&](const FChunkDependency& RHS) 
+		{
+			return Node.ChunkID == RHS.ParentChunkID;
+		});
+	for (const auto& ChildIndex : ChildNodeIndices)
+	{
+		Node.ChildNodes.Add(FChunkDependencyTreeNode(ChildIndex.ChunkID));
+	}
+	for (auto& Child : Node.ChildNodes) 
+	{
+		AddChildrenRecursive(Child, DepInfo, Parents);
+	}
+}
+
+void UChunkDependencyInfo::RemoveRedundantChunks(TArray<int32>& ChunkIDs) const
+{
+	for (int32 ChunkIndex = ChunkIDs.Num() - 1; ChunkIndex >= 0; ChunkIndex--)
+	{
+		const TSet<int32>* FoundParents = ChildToParentMap.Find(ChunkIDs[ChunkIndex]);
+
+		if (FoundParents)
+		{
+			for (int32 ParentChunk : *FoundParents)
+			{
+				if (ChunkIDs.Contains(ParentChunk))
+				{
+					ChunkIDs.RemoveAt(ChunkIndex);
+					break;
+				}
+			}
+		}
+	}
+}
+
+int32 UChunkDependencyInfo::FindHighestSharedChunk(const TArray<int32>& ChunkIDs) const
+{
+	TArray<int32> TestChunkIds;
+	TestChunkIds.Append(ChunkIDs);
+	Algo::Sort(TestChunkIds);
+	TestChunkIds.SetNum(Algo::Unique(TestChunkIds));
+
+	for (int32 ChunkId : TestChunkIds)
+	{
+		if (!ChildToParentMap.Contains(ChunkId))
+		{
+			return INDEX_NONE;
+		}
+	}
+
+	if (TestChunkIds.Num() == 0)
+	{
+		return INDEX_NONE;
+	}
+	if (TestChunkIds.Num() == 1)
+	{
+		return TestChunkIds[0];
+	}
+
+	TSet<int32> CommonParentSet;
+	CommonParentSet.Append(ChildToParentMap[TestChunkIds[0]]);
+	CommonParentSet.Add(TestChunkIds[0]);
+
+	for (int32 Index = 1; Index < TestChunkIds.Num(); Index++)
+	{
+		TSet<int32> OtherSet;
+		OtherSet.Append(ChildToParentMap[TestChunkIds[Index]]);
+		OtherSet.Add(TestChunkIds[Index]);
+		CommonParentSet = CommonParentSet.Intersect(OtherSet);
+	}
+
+	int32 HighestIndex = INDEX_NONE;
+	for (int32 ChunkId : CommonParentSet)
+	{
+		HighestIndex = FMath::Max(HighestIndex, TopologicallySortedChunks.IndexOfByKey(ChunkId));
+	}
+	if (HighestIndex == INDEX_NONE)
+	{
+		UE_LOG(LogChunkDependencyInfo, Error, TEXT("Unable to find parent."));
+		return 0;
+	}
+	return TopologicallySortedChunks[HighestIndex];
+}
+
+void UChunkDependencyInfo::GetChunkDependencies(const int32 InChunk, TSet<int32>& OutChunkDependencies) const
+{
+	if (const TSet<int32>* Parents = ChildToParentMap.Find(InChunk))
+	{
+		OutChunkDependencies.Append(*Parents);
+	}
+}
diff --git a/Engine/Source/Runtime/CoreUObject/Public/Misc/ChunkDependencyInfo.h b/Engine/Source/Runtime/CoreUObject/Public/Misc/ChunkDependencyInfo.h
new file mode 100644
--- /dev/null
+++ b/Engine/Source/Runtime/CoreUObject/Public/Misc/ChunkDependencyInfo.h
@@ -0,0 +1,101 @@
+// Copyright Epic Games, Inc. All Rights Reserved.
+
+#pragma once
+
+#include "Containers/Array.h"
+#include "Containers/Map.h"
+#include "Containers/Set.h"
+#include "CoreMinimal.h"
+#include "HAL/Platform.h"
+#include "HAL/PlatformCrt.h"
+#include "UObject/Object.h"
+#include "UObject/ObjectMacros.h"
+#include "UObject/UObjectGlobals.h"
+
+#include "ChunkDependencyInfo.generated.h"
+
+/** A single dependency, read from ini file */
+USTRUCT()
+struct FChunkDependency
+{
+	GENERATED_USTRUCT_BODY()
+
+	FChunkDependency()
+	{};
+
+	FChunkDependency(int32 InChunkId, int32 InParentChunkId)
+		: ChunkID(InChunkId)
+		, ParentChunkID(InParentChunkId)
+	{};
+
+	/** The child chunk */
+	UPROPERTY(EditAnywhere, Category = ChunkInfo)
+	int32 ChunkID = 0;
+
+	/** Parent chunk, anything in both Parent and Child is only placed into Parent */
+	UPROPERTY(EditAnywhere, Category = ChunkInfo)
+	int32 ParentChunkID = 0;
+
+	bool operator== (const FChunkDependency& RHS) const
+	{
+		return ChunkID == RHS.ChunkID;
+	}
+};
+
+/** In memory structure used for dependency tree */
+struct FChunkDependencyTreeNode
+{
+	FChunkDependencyTreeNode(int32 InChunkID=0)
+		: ChunkID(InChunkID)
+	{}
+
+	int32 ChunkID;
+
+	TArray<FChunkDependencyTreeNode> ChildNodes;
+};
+
+/** This is read out of config and defines a tree of chunk dependencies */
+UCLASS(config=Engine, defaultconfig, MinimalAPI)
+class UChunkDependencyInfo : public UObject
+{
+	GENERATED_UCLASS_BODY()
+
+	/** Will return an existing dependency graph, or build if it necessary */
+	COREUOBJECT_API const FChunkDependencyTreeNode* GetOrBuildChunkDependencyGraph(int32 HighestChunk = 0, bool bForceRebuild = false);
+
+	/** Will create a dependency tree starting with RootTreeNode. If HighestChunk is == 0 it will only add dependencies on 0 for chunks already in the dependencies list */
+	COREUOBJECT_API const FChunkDependencyTreeNode* BuildChunkDependencyGraph(int32 HighestChunk);
+
+	/** Removes redundant chunks from a chunk list */
+	COREUOBJECT_API void RemoveRedundantChunks(TArray<int32>& ChunkIDs) const;
+
+	/** Given a set of chunk Ids will use the dependencies to find which chunk is shared among all the given chunks.
+	All chunks should converge to the 0 startup chunk.
+	Potentially will return a chunk that is not in the given input array */
+	COREUOBJECT_API int32 FindHighestSharedChunk(const TArray<int32>& ChunkIDs) const;
+
+	/**
+	 * Returns the chunks that the in chunk depends upon through the parent rules.
+	 */
+	COREUOBJECT_API void GetChunkDependencies(const int32 InChunk, TSet<int32>& OutChunkDependencies) const;
+
+	/** List of dependencies used to remove redundant chunks */
+	UPROPERTY(config)
+	TArray<FChunkDependency> DependencyArray;
+
+private:
+	/** Fills out the Dependency Tree starting with Node */
+	void AddChildrenRecursive(FChunkDependencyTreeNode& Node, TArray<FChunkDependency>& DepInfo, TSet<int32> Parents);
+
+	/** Root of tree, this will be valid after calling BuildChunkDependencyGraph */
+	FChunkDependencyTreeNode RootTreeNode;
+
+	/** Map of child chunks to all parent chunks, computed in BuildChunkDependencyGraph */
+	TMap<int32, TSet<int32>> ChildToParentMap;
+
+	/** Cached array of topologically sorted chunk. */
+	TArray<int32> TopologicallySortedChunks;
+
+	/** Cached value of HighestChunk at time of building */
+	int32 CachedHighestChunk;
+};
diff --git a/Engine/Source/Runtime/Engine/Private/AssetManager.cpp b/Engine/Source/Runtime/Engine/Private/AssetManager.cpp
--- a/Engine/Source/Runtime/Engine/Private/AssetManager.cpp
+++ b/Engine/Source/Runtime/Engine/Private/AssetManager.cpp
@@ -45,7 +45,7 @@
 
 #if WITH_EDITOR
 #include "Editor.h"
-#include "Commandlets/ChunkDependencyInfo.h"
+#include "Misc/ChunkDependencyInfo.h"
 #include "Framework/Notifications/NotificationManager.h"
 #include "Settings/ProjectPackagingSettings.h"
 #include "Widgets/Notifications/SNotificationList.h"
