Add a cvar to suppress the CoreRedirects warning about iterative redirects.
#Synced-CL 49725521

diff --git a/Engine/Source/Runtime/CoreUObject/Private/UObject/CoreRedirects.cpp b/Engine/Source/Runtime/CoreUObject/Private/UObject/CoreRedirects.cpp
--- a/Engine/Source/Runtime/CoreUObject/Private/UObject/CoreRedirects.cpp
+++ b/Engine/Source/Runtime/CoreUObject/Private/UObject/CoreRedirects.cpp
@@ -2689,6 +2689,7 @@ bool FCoreRedirects::ReadRedirectsFromIni(const FString& IniName)
 				bool bRemoved = false;
 				bool bMatchSubstring = false;
 				bool bMatchWildcard = false;
+				bool bSuppressIterateWarning = false;
 
 				const FString& ValueString = It.Value().GetValue();
 
@@ -2697,6 +2698,7 @@ bool FCoreRedirects::ReadRedirectsFromIni(const FString& IniName)
 				FParse::Bool(*ValueString, TEXT("Removed="), bRemoved);
 				FParse::Bool(*ValueString, TEXT("MatchSubstring="), bMatchSubstring);
 				FParse::Bool(*ValueString, TEXT("MatchWildcard="), bMatchWildcard);
+				FParse::Bool(*ValueString, TEXT("SuppressIterateWarning="), bSuppressIterateWarning);
 
 				FParse::Value(*ValueString, TEXT("OldName="), OldName);
 				FParse::Value(*ValueString, TEXT("NewName="), NewName);
@@ -2724,9 +2726,12 @@ bool FCoreRedirects::ReadRedirectsFromIni(const FString& IniName)
 					{
 						NewFlags |= ECoreRedirectFlags::Option_Iterate;
 
-						//UE_LOG(LogCoreRedirects, Warning,
-						//	TEXT("ReadRedirectsFromIni(%s) Redirect %s is iterative. Iterative redirects are expensive, and should be used with caution."),
-						//	*IniName, *ValueString);
+						if (!bSuppressIterateWarning)
+						{
+							UE_LOG(LogCoreRedirects, Warning,
+								TEXT("ReadRedirectsFromIni(%s) Redirect %s is iterative. Iterative redirects are expensive, and should be used with caution."),
+								*IniName, *ValueString);
+						}
 					}
 
 					if (bMatchWildcard || bMatchSubstring)
