[Core] Add missing comment in FArchiveUObject::SerializeWeakObjectPtr

#Synced-CL 49780466

diff --git a/Engine/Source/Runtime/CoreUObject/Private/Serialization/ArchiveUObject.cpp b/Engine/Source/Runtime/CoreUObject/Private/Serialization/ArchiveUObject.cpp
--- a/Engine/Source/Runtime/CoreUObject/Private/Serialization/ArchiveUObject.cpp
+++ b/Engine/Source/Runtime/CoreUObject/Private/Serialization/ArchiveUObject.cpp
@@ -143,6 +143,7 @@ FArchive& FArchiveUObject::SerializeWeakObjectPtr(FArchive& Ar, FWeakObjectPtr&
 	
 		Ar << Object;
 	
+		// Only commit the value using the same checks that we have in map/set to align with rehashing, otherwise those containers can become inoperable and lead to crashes
 		if (Ar.IsLoading() || (Ar.IsModifyingWeakAndStrongReferences() && !Ar.IsSaving()))
 		{
 			Value = Object;
