PR #14258: [CoreUObject] Fix unintended string concat in list of Object property flag names

#rnx

#Synced-CL 49717872

diff --git a/Engine/Source/Runtime/CoreUObject/Private/UObject/Obj.cpp b/Engine/Source/Runtime/CoreUObject/Private/UObject/Obj.cpp
--- a/Engine/Source/Runtime/CoreUObject/Private/UObject/Obj.cpp
+++ b/Engine/Source/Runtime/CoreUObject/Private/UObject/Obj.cpp
@@ -4684,7 +4684,7 @@ COREUOBJECT_API TArray<const TCHAR*> ParsePropertyFlags(EPropertyFlags InFlags)
 		TEXT("CPF_NativeAccessSpecifierProtected"),
 		TEXT("CPF_NativeAccessSpecifierPrivate"),
 		TEXT("CPF_SkipSerialization"),
-		TEXT("CPF_TObjectPtr")
+		TEXT("CPF_TObjectPtr"),
 		TEXT("CPF_ExperimentalOverridableLogic"),
 		TEXT("CPF_ExperimentalAlwaysOverriden"),
 		TEXT("CPF_ExperimentalNeverOverriden"),
