LLM: Remove references to deprecated LLM defines which are now always enabled.
#rnx

#Synced-CL 49797803

#Synced-CL 49797803

diff --git a/Engine/Source/Developer/TraceInsights/Private/Insights/MemoryProfiler/Widgets/SMemTagTreeView.cpp b/Engine/Source/Developer/TraceInsights/Private/Insights/MemoryProfiler/Widgets/SMemTagTreeView.cpp
--- a/Engine/Source/Developer/TraceInsights/Private/Insights/MemoryProfiler/Widgets/SMemTagTreeView.cpp
+++ b/Engine/Source/Developer/TraceInsights/Private/Insights/MemoryProfiler/Widgets/SMemTagTreeView.cpp
@@ -278,7 +278,7 @@ void SMemTagTreeView::UpdateAvailableTagSets()
 		}
 
 		FText TagSetToolTip;
-		const FText TagSetToolTipFmt = LOCTEXT("TagSetToolTipFmt", "'{0}' Tag Set\n\nThe tree view shows memory tags from the selected tag set.\n\nNote:\n\tTrace data is emitted by the Low-Level Memory Tracker system (LLM).\n\tSee documentation about LLM_ALLOW_ASSETS_TAGS macro and\n\tabout \"-llm -llmtagsets=assets,assetclasses\" command line parameters.");
+		const FText TagSetToolTipFmt = LOCTEXT("TagSetToolTipFmt", "'{0}' Tag Set\n\nThe tree view shows memory tags from the selected tag set.\n\nNote:\n\tTrace data is emitted by the Low-Level Memory Tracker system (LLM).\n\tSee documentation \n\tabout \"-llm -llmtagsets=assets,assetclasses\" command line parameters.");
 		TagSetToolTip = FText::Format(TagSetToolTipFmt, TagSetName);
 
 		TSharedRef<SWidget> Widget = SNew(SHorizontalBox)
diff --git a/Engine/Source/Developer/TraceInsights/Private/Insights/MemoryProfiler/Widgets/SMemoryProfilerWindow.cpp b/Engine/Source/Developer/TraceInsights/Private/Insights/MemoryProfiler/Widgets/SMemoryProfilerWindow.cpp
--- a/Engine/Source/Developer/TraceInsights/Private/Insights/MemoryProfiler/Widgets/SMemoryProfilerWindow.cpp
+++ b/Engine/Source/Developer/TraceInsights/Private/Insights/MemoryProfiler/Widgets/SMemoryProfilerWindow.cpp
@@ -434,7 +434,7 @@ void SMemoryProfilerWindow::RegisterTabSpawners()
 		TabManagerPtr->RegisterTabSpawner(FMemoryProfilerTabs::MemTagTreeViewID, FOnSpawnTab::CreateRaw(this, &SMemoryProfilerWindow::SpawnTab_MemTagTreeView))
 			.SetDisplayName(LOCTEXT("MemTagTreeViewTabTitle", "Memory Tags"))
 			.SetIcon(FSlateIcon(FInsightsStyle::GetStyleSetName(), "Icons.MemTagTreeView"))
-			.SetTooltipText(LOCTEXT("MemTagTreeViewTabToolTip", "Opens the Memory Tags (LLM) panel.\n\nIt analyses the per-frame trace data emitted by the Low-Level Memory Tracker system (LLM).\nSee documentation about LLM_ALLOW_ASSETS_TAGS macro and\nabout \"-llm -llmtagsets=assets,assetclasses\" command line parameters."))
+			.SetTooltipText(LOCTEXT("MemTagTreeViewTabToolTip", "Opens the Memory Tags (LLM) panel.\n\nIt analyses the per-frame trace data emitted by the Low-Level Memory Tracker system (LLM).\nSee documentation \nabout \"-llm -llmtagsets=assets,assetclasses\" command line parameters."))
 			.SetGroup(Group);
 	}
 
diff --git a/Engine/Source/Runtime/Core/Public/HAL/LowLevelMemTracker.h b/Engine/Source/Runtime/Core/Public/HAL/LowLevelMemTracker.h
--- a/Engine/Source/Runtime/Core/Public/HAL/LowLevelMemTracker.h
+++ b/Engine/Source/Runtime/Core/Public/HAL/LowLevelMemTracker.h
@@ -20,15 +20,15 @@
 
 // Deprecated defines no longer used
 #ifndef LLM_ALLOW_ASSETS_TAGS
-#define LLM_ALLOW_ASSETS_TAGS ENABLE_LOW_LEVEL_MEM_TRACKER
+#define LLM_ALLOW_ASSETS_TAGS UE_DEPRECATED_MACRO(5.8, "LLM_ALLOW_ASSETS_TAGS is deprecated and should always be considered 1 when ENABLE_LOW_LEVEL_MEM_TRACKER is 1.") ENABLE_LOW_LEVEL_MEM_TRACKER
 #endif
-#ifndef  LLM_ALLOW_UOBJECTCLASSES_TAGS
-#define LLM_ALLOW_UOBJECTCLASSES_TAGS ENABLE_LOW_LEVEL_MEM_TRACKER
+#ifndef LLM_ALLOW_UOBJECTCLASSES_TAGS
+#define LLM_ALLOW_UOBJECTCLASSES_TAGS UE_DEPRECATED_MACRO(5.8, "LLM_ALLOW_UOBJECTCLASSES_TAGS is deprecated and should always be considered 1 when ENABLE_LOW_LEVEL_MEM_TRACKER is 1.") ENABLE_LOW_LEVEL_MEM_TRACKER
 #endif
 #ifndef LLM_ALLOW_STATS
-#define LLM_ALLOW_STATS ENABLE_LOW_LEVEL_MEM_TRACKER
+#define LLM_ALLOW_STATS UE_DEPRECATED_MACRO(5.8, "LLM_ALLOW_STATS is deprecated and should always be considered 1 when ENABLE_LOW_LEVEL_MEM_TRACKER is 1.") ENABLE_LOW_LEVEL_MEM_TRACKER
 #endif
-#define LLM_ENABLED_STAT_TAGS ENABLE_LOW_LEVEL_MEM_TRACKER
+#define LLM_ENABLED_STAT_TAGS UE_DEPRECATED_MACRO(5.8, "LLM_ENABLED_STAT_TAGS is deprecated and should always be considered 1 when ENABLE_LOW_LEVEL_MEM_TRACKER is 1.") ENABLE_LOW_LEVEL_MEM_TRACKER
 
 // Define LLM macros to refer to the implementation class if LLM is compiled in, otherwise define them to NOOP.
 #if ENABLE_LOW_LEVEL_MEM_TRACKER
diff --git a/Engine/Source/Runtime/CoreUObject/Private/Serialization/AsyncLoading2.cpp b/Engine/Source/Runtime/CoreUObject/Private/Serialization/AsyncLoading2.cpp
--- a/Engine/Source/Runtime/CoreUObject/Private/Serialization/AsyncLoading2.cpp
+++ b/Engine/Source/Runtime/CoreUObject/Private/Serialization/AsyncLoading2.cpp
@@ -12471,10 +12471,8 @@ void FAsyncLoadingThread2::FlushLoading(TConstArrayView<int32> RequestIDs)
 	if (IsAsyncLoadingPackages())
 	{
 		LLM_SCOPE(ELLMTag::AsyncLoading);
-#if LLM_ALLOW_ASSETS_TAGS
 		LLM_TAGSET_SCOPE_CLEAR(ELLMTagSet::Assets);
 		LLM_TAGSET_SCOPE_CLEAR(ELLMTagSet::AssetClasses);
-#endif // LLM_ALLOW_ASSETS_TAGS
 		UE_TRACE_METADATA_CLEAR_SCOPE();
 		// Prevent objects still being loaded from being accessible from the game thread.
 		TGuardValue GuardVisibilityFilter(FUObjectThreadContext::Get().AsyncVisibilityFilter, EInternalObjectFlags::AsyncLoadingPhase1);
diff --git a/Engine/Source/Runtime/CoreUObject/Private/UObject/UObjectArray.cpp b/Engine/Source/Runtime/CoreUObject/Private/UObject/UObjectArray.cpp
--- a/Engine/Source/Runtime/CoreUObject/Private/UObject/UObjectArray.cpp
+++ b/Engine/Source/Runtime/CoreUObject/Private/UObject/UObjectArray.cpp
@@ -232,10 +232,8 @@ void FUObjectArray::AllocateUObjectIndex(UObjectBase* Object, EInternalObjectFla
 	LLM_SCOPE(ELLMTag::UObject);
 	LLM_SCOPE_BYTAG(UObject_UObjectArray);
 	// Clear asset scopes
-#if LLM_ALLOW_ASSETS_TAGS
 	LLM_TAGSET_SCOPE_CLEAR(ELLMTagSet::Assets);
 	LLM_TAGSET_SCOPE_CLEAR(ELLMTagSet::AssetClasses);
-#endif // LLM_ALLOW_ASSETS_TAGS
 	UE_TRACE_METADATA_CLEAR_SCOPE();
 
 	int32 Index = INDEX_NONE;
diff --git a/Engine/Source/Runtime/CoreUObject/Private/UObject/UObjectGlobals.cpp b/Engine/Source/Runtime/CoreUObject/Private/UObject/UObjectGlobals.cpp
--- a/Engine/Source/Runtime/CoreUObject/Private/UObject/UObjectGlobals.cpp
+++ b/Engine/Source/Runtime/CoreUObject/Private/UObject/UObjectGlobals.cpp
@@ -3461,11 +3461,15 @@ UObject* StaticAllocateObject
 	LLM_SCOPE_BYTAG(UObject_StaticAllocateObject);
 	
 	// report detailed UObject classes allocations
-#if ENABLE_LOW_LEVEL_MEM_TRACKER && LLM_ALLOW_UOBJECTCLASSES_TAGS
-	FString ClassNameString = FString::Printf(TEXT("UObject/%s"), *InClass->GetName());
-	FName LLMScope_Name(*ClassNameString);
-
-	LLM_SCOPE_DYNAMIC(LLMScope_Name, ELLMTracker::Default, ELLMTagSet::UObjectClasses, FLLMDynamicTagConstructorStatString(ClassNameString));
+#if ENABLE_LOW_LEVEL_MEM_TRACKER
+	FName LLMScope_Name;
+	FString LLMClassNameString;
+	if (FLowLevelMemTracker::Get().IsTagSetActive(ELLMTagSet::UObjectClasses))
+	{
+		LLMClassNameString = FString::Printf(TEXT("UObject/%s"), *InClass->GetName());
+		LLMScope_Name = *LLMClassNameString;
+	}
+	LLM_SCOPE_DYNAMIC(LLMScope_Name, ELLMTracker::Default, ELLMTagSet::UObjectClasses, FLLMDynamicTagConstructorStatString(LLMClassNameString));
 #endif // ENABLE_LOW_LEVEL_MEM_TRACKER
 
 	SCOPE_CYCLE_COUNTER(STAT_AllocateObject);
diff --git a/Engine/Source/Runtime/CoreUObject/Private/UObject/UObjectHash.cpp b/Engine/Source/Runtime/CoreUObject/Private/UObject/UObjectHash.cpp
--- a/Engine/Source/Runtime/CoreUObject/Private/UObject/UObjectHash.cpp
+++ b/Engine/Source/Runtime/CoreUObject/Private/UObject/UObjectHash.cpp
@@ -1523,10 +1523,8 @@ FORCEINLINE static void RemoveFromRemoteIdMap(FUObjectHashTables& ThreadHash, UO
 void ShrinkUObjectHashTables()
 {
 	LLM_SCOPE_BYTAG(UObjectHash);
-#if LLM_ALLOW_ASSETS_TAGS
 	LLM_TAGSET_SCOPE_CLEAR(ELLMTagSet::Assets);
 	LLM_TAGSET_SCOPE_CLEAR(ELLMTagSet::AssetClasses);
-#endif // LLM_ALLOW_ASSETS_TAGS
 	UE_TRACE_METADATA_CLEAR_SCOPE();
 	TRACE_CPUPROFILER_EVENT_SCOPE(ShrinkUObjectHashTables);
 	FUObjectHashTables& ThreadHash = FUObjectHashTables::Get();
@@ -2011,10 +2009,8 @@ void HashObject(UObjectBase* Object)
 	if (!Name.IsNone())
 	{
 		LLM_SCOPE_BYTAG(UObjectHash);
-#if LLM_ALLOW_ASSETS_TAGS
 		LLM_TAGSET_SCOPE_CLEAR(ELLMTagSet::Assets);
 		LLM_TAGSET_SCOPE_CLEAR(ELLMTagSet::AssetClasses);
-#endif // LLM_ALLOW_ASSETS_TAGS
 		UE_TRACE_METADATA_CLEAR_SCOPE();
 #if !UE_BUILD_TEST && !UE_BUILD_SHIPPING
 		SCOPE_CYCLE_COUNTER(STAT_Hash_HashObject);
@@ -2068,10 +2064,8 @@ void UnhashObject(UObjectBase* Object)
 	if (!Name.IsNone())
 	{
 		LLM_SCOPE_BYTAG(UObjectHash);
-#if LLM_ALLOW_ASSETS_TAGS
 		LLM_TAGSET_SCOPE_CLEAR(ELLMTagSet::Assets);
 		LLM_TAGSET_SCOPE_CLEAR(ELLMTagSet::AssetClasses);
-#endif // LLM_ALLOW_ASSETS_TAGS
 		UE_TRACE_METADATA_CLEAR_SCOPE();
 #if !UE_BUILD_TEST && !UE_BUILD_SHIPPING
 		SCOPE_CYCLE_COUNTER(STAT_Hash_UnhashObject);
@@ -2114,10 +2108,8 @@ void UnhashObject(UObjectBase* Object)
 void HashObjectExternalPackage(UObjectBase* Object, UPackage* Package)
 {
 	LLM_SCOPE_BYTAG(UObjectHash);
-#if LLM_ALLOW_ASSETS_TAGS
 	LLM_TAGSET_SCOPE_CLEAR(ELLMTagSet::Assets);
 	LLM_TAGSET_SCOPE_CLEAR(ELLMTagSet::AssetClasses);
-#endif // LLM_ALLOW_ASSETS_TAGS
 	UE_TRACE_METADATA_CLEAR_SCOPE();
 	if (Package)
 	{
@@ -2143,10 +2135,8 @@ void HashObjectExternalPackage(UObjectBase* Object, UPackage* Package)
 void UnhashObjectExternalPackage(class UObjectBase* Object)
 {
 	LLM_SCOPE_BYTAG(UObjectHash);
-#if LLM_ALLOW_ASSETS_TAGS
 	LLM_TAGSET_SCOPE_CLEAR(ELLMTagSet::Assets);
 	LLM_TAGSET_SCOPE_CLEAR(ELLMTagSet::AssetClasses);
-#endif // LLM_ALLOW_ASSETS_TAGS
 	UE_TRACE_METADATA_CLEAR_SCOPE();
 	FUObjectHashTables& ThreadHash = FUObjectHashTables::Get();
 	FHashTableLock LockHash(ThreadHash);
diff --git a/Engine/Source/Runtime/Renderer/Private/RendererScene.cpp b/Engine/Source/Runtime/Renderer/Private/RendererScene.cpp
--- a/Engine/Source/Runtime/Renderer/Private/RendererScene.cpp
+++ b/Engine/Source/Runtime/Renderer/Private/RendererScene.cpp
@@ -1360,7 +1360,7 @@ void FScene::BatchAddPrimitivesInternal(TArrayView<T*> InPrimitives)
 {
 	check(InPrimitives.Num() > 0);
 
-#if ENABLE_LOW_LEVEL_MEM_TRACKER && LLM_ENABLED_STAT_TAGS
+#if ENABLE_LOW_LEVEL_MEM_TRACKER
 	// If detailed per-tag asset memory stats are active, don't batch primitives, so the memory tags can be independent
 	if (FLowLevelMemTracker::Get().IsTagSetActive(ELLMTagSet::Assets) && InPrimitives.Num() > 1)
 	{
