Adding an include needed for UE_WITH_REMOTE_OBJECT_HANDLE case.

#rnx

#Synced-CL 49820573

#Synced-CL 49820573

diff --git a/Engine/Source/Runtime/CoreUObject/Private/UObject/ObjectPtr.cpp b/Engine/Source/Runtime/CoreUObject/Private/UObject/ObjectPtr.cpp
--- a/Engine/Source/Runtime/CoreUObject/Private/UObject/ObjectPtr.cpp
+++ b/Engine/Source/Runtime/CoreUObject/Private/UObject/ObjectPtr.cpp
@@ -4,6 +4,7 @@
 #include "UObject/Class.h"
 #include "UObject/RemoteObjectPrivate.h"
 #include "UObject/RemoteObjectPathName.h"
+#include "UObject/UObjectHash.h"
 
 #if UE_WITH_OBJECT_HANDLE_LATE_RESOLVE
 
