Fix PVS null pointer reference warnings
#rnx

diff --git a/Engine/Source/Runtime/CoreUObject/Public/AssetRegistry/CookTagList.h b/Engine/Source/Runtime/CoreUObject/Public/AssetRegistry/CookTagList.h
--- a/Engine/Source/Runtime/CoreUObject/Public/AssetRegistry/CookTagList.h
+++ b/Engine/Source/Runtime/CoreUObject/Public/AssetRegistry/CookTagList.h
@@ -46,7 +46,7 @@ struct FCookTagList
 	*/
 	void Add(UObject* InObject, FName InTagName, FString&& InTagValue)
 	{
-		if (!Package || !ensure(InObject->GetOutermost() == Package))
+		if (!Package || !InObject || !ensure(InObject->GetOutermost() == Package))
 		{
 			return;
 		}
diff --git a/Engine/Source/Runtime/Engine/Private/TextureDerivedData.cpp b/Engine/Source/Runtime/Engine/Private/TextureDerivedData.cpp
--- a/Engine/Source/Runtime/Engine/Private/TextureDerivedData.cpp
+++ b/Engine/Source/Runtime/Engine/Private/TextureDerivedData.cpp
@@ -3061,6 +3061,7 @@ static void SerializePlatformData(
 			CookTags->Add(Texture, "Size", MoveTemp(DimensionsStr));
 		}
 
+		check(Texture);
 		CookTags->Add(Texture, "Format", FString(GPixelFormats[PlatformData->PixelFormat].Name));
 
 		// Add in diff keys for change detection/blame.
