Add UPackage::IsInstanced function to determine if a package is instanced or not.

#Synced-CL 49830329

diff --git a/Engine/Source/Runtime/CoreUObject/Public/UObject/Package.h b/Engine/Source/Runtime/CoreUObject/Public/UObject/Package.h
--- a/Engine/Source/Runtime/CoreUObject/Public/UObject/Package.h
+++ b/Engine/Source/Runtime/CoreUObject/Public/UObject/Package.h
@@ -469,6 +469,12 @@ public:
 	/** Packages are never assets */
 	virtual bool IsAsset() const override { return false; }
 
+	bool IsInstanced() const
+	{
+		const FName LoadedPathName = LoadedPath.GetPackageFName();
+		return !LoadedPathName.IsNone() && (LoadedPathName != GetFName());
+	}
+
 #if WITH_EDITOR
 	/**
 	 *  Functionality to support SoftGC in the cooker. When activated by the cooker, objects in
