[ZenLoader]
- Fix potential race condition on AddConstructedObjects depending on the thread it is called on and the state of the package

#Synced-CL 49847011

diff --git a/Engine/Source/Runtime/CoreUObject/Private/Serialization/AsyncLoading2.cpp b/Engine/Source/Runtime/CoreUObject/Private/Serialization/AsyncLoading2.cpp
--- a/Engine/Source/Runtime/CoreUObject/Private/Serialization/AsyncLoading2.cpp
+++ b/Engine/Source/Runtime/CoreUObject/Private/Serialization/AsyncLoading2.cpp
@@ -11659,6 +11659,31 @@ void FAsyncLoadingThread2::NotifyConstructedDuringAsyncLoading(UObject* Object,
 		return;
 	}
 
+	auto IsTargetPackageSafeForAdd =
+		[this](FAsyncPackage2* TargetPackage)
+		{
+			// if the loader isn't multithreaded, it's fine to add constructed packages as long as the package
+			// is not done with its deferred postload phase.
+			if (!IsMultithreaded())
+			{
+				return TargetPackage->AsyncPackageLoadingState <= EAsyncPackageLoadingState2::DeferredPostLoad;
+			}
+
+			// To avoid race conditions, ALT packages should only be added to by ALT threads. 
+			if (TargetPackage->AsyncPackageLoadingState < EAsyncPackageLoadingState2::DeferredPostLoad)
+			{
+				return !IsInGameThread();
+			}
+
+			// To avoid race conditions, GT packages should only be added to by GT threads.
+			if (TargetPackage->AsyncPackageLoadingState == EAsyncPackageLoadingState2::DeferredPostLoad)
+			{
+				return IsInGameThread();
+			}
+
+			return false;
+		};
+
 #if WITH_EDITOR
 	// In editor, objects from other packages might be constructed with the wrong package currently on stack (i.e. reinstantiation).
 	// It's a lot cleaner and easier to properly dispatch them to their respective package than adding scopes everywhere.
@@ -11666,7 +11691,7 @@ void FAsyncLoadingThread2::NotifyConstructedDuringAsyncLoading(UObject* Object,
 	{
 		if (FAsyncPackage2* AsyncPackage = (FAsyncPackage2*)LinkerLoad->AsyncRoot)
 		{
-			if (AsyncPackage->AsyncPackageLoadingState <= EAsyncPackageLoadingState2::DeferredPostLoad)
+			if (IsTargetPackageSafeForAdd(AsyncPackage))
 			{
 				check(AsyncPackage->Desc.UPackageId == ObjectPackage->GetPackageId());
 				AsyncPackage->AddConstructedObject(Object, bSubObjectThatAlreadyExists);
@@ -11682,7 +11707,7 @@ void FAsyncLoadingThread2::NotifyConstructedDuringAsyncLoading(UObject* Object,
 		AsyncPackageLookup.FindAndApply(ObjectPackageId,
 			[&](FAsyncPackage2* AsyncPackage)
 			{
-				if (AsyncPackage->AsyncPackageLoadingState < EAsyncPackageLoadingState2::DeferredPostLoad)
+				if (IsTargetPackageSafeForAdd(AsyncPackage))
 				{
 					AsyncPackage->AddConstructedObject(Object, bSubObjectThatAlreadyExists);
 					bReturn = true;
@@ -11697,7 +11722,7 @@ void FAsyncLoadingThread2::NotifyConstructedDuringAsyncLoading(UObject* Object,
 	}
 #endif
 
-	if (ContextAsyncPackage->AsyncPackageLoadingState <= EAsyncPackageLoadingState2::DeferredPostLoad)
+	if (IsTargetPackageSafeForAdd(ContextAsyncPackage))
 	{
 		UE_CLOG(ObjectPackage->GetPackageId().IsValid() && ContextAsyncPackage->Desc.UPackageId != ObjectPackage->GetPackageId(), LogStreaming, VeryVerbose,
 			TEXT("Constructed object '%s' is part of package '%s' which is no longer postloading objects. This object will be postloaded with the FUObjectThreadContext AsyncPackage '%s'(%s) instead."),
