Allow FindOrGetImportCell to return null

This function was always intended to be fallible, for its use by ConditionalCreateCellImport. Its null checks make more sense in its other caller, the operator<< overload for VCell.

#rnx

#Synced-CL 49755474

diff --git a/Engine/Source/Runtime/CoreUObject/Private/Serialization/AsyncLoading2.cpp b/Engine/Source/Runtime/CoreUObject/Private/Serialization/AsyncLoading2.cpp
--- a/Engine/Source/Runtime/CoreUObject/Private/Serialization/AsyncLoading2.cpp
+++ b/Engine/Source/Runtime/CoreUObject/Private/Serialization/AsyncLoading2.cpp
@@ -2879,12 +2879,10 @@ struct FPackageImportStore
 		if (GlobalIndex.IsScriptImport())
 		{
 			Cell = GlobalImportStore.FindScriptImportCell(GlobalIndex);
-			check(Cell);
 		}
 		else if (GlobalIndex.IsPackageImport())
 		{
 			Cell = GlobalImportStore.FindPublicExportCell(FPublicExportKey::FromPackageImport(GlobalIndex, Header.ImportedPackageIds, Header.ImportedPublicExportHashes));
-			check(Cell);
 		}
 		else
 		{
@@ -3159,6 +3157,7 @@ public:
 			{
 				FPackageObjectIndex GlobalIndex = HeaderData->CellImportMap[CellImportIndex];
 				Cell = ImportStore->FindOrGetImportCell(*HeaderData, GlobalIndex);
+				check(Cell);
 			}
 			else
 			{
