[UE] Add an asset log to note when exports get created as IDO placeholder types during editor asset loads.

Notes:
- Excluded from internal placeholder mode build configs (since it universally loads all non-native types as placeholder).
- Also added a log to note when references to placeholder-typed exports are serialized as NULL (due to IDO feature CVar being disabled).
- Intent is to add visibility around IDO placeholders when debugging e.g. crash reports/cook job failures related to package serialization in build environments where IDO is enabled.

#Synced-CL 49835867

diff --git a/Engine/Source/Runtime/CoreUObject/Private/UObject/LinkerLoad.cpp b/Engine/Source/Runtime/CoreUObject/Private/UObject/LinkerLoad.cpp
--- a/Engine/Source/Runtime/CoreUObject/Private/UObject/LinkerLoad.cpp
+++ b/Engine/Source/Runtime/CoreUObject/Private/UObject/LinkerLoad.cpp
@@ -6481,6 +6481,13 @@ UObject* FLinkerLoad::CreateExport( int32 Index )
 				CurrentLoadContext->IncrementForcedExportCount();
 				FLinkerManager::Get().AddLoaderWithForcedExports(this);
 			}
+
+#if WITH_EDITOR && !UE_INTERNAL_UNIVERSAL_PLACEHOLDERS
+			if (UE::FPropertyBagRepository::IsPropertyBagPlaceholderObject(Export.Object))
+			{
+				UE_ASSET_LOG(LogLinker, Log, PackagePath, TEXT("Export %s created with placeholder type: %s"), *Export.Object->GetName(), *LoadClass->GetFullName());
+			}
+#endif
 		}
 	}
 	return Export.bExportLoadFailed ? nullptr : Export.Object;
@@ -6939,6 +6946,7 @@ FArchive& FLinkerLoad::operator<<( UObject*& Object )
 		}
 		else
 		{
+			UE_LOG(LogLinker, Log, TEXT("Serializing reference to \"%s\" as NULL because it has a placeholder type, and reference loading is currently disabled."), *Object->GetPathName());
 			Object = nullptr;
 		}
 	}
@@ -6983,6 +6991,7 @@ FArchive& FLinkerLoad::operator<<(FObjectPtr& ObjectPtr)
 			}
 			else
 			{
+				UE_LOG(LogLinker, Log, TEXT("Serializing reference to \"%s\" as NULL because it has a placeholder type, and reference loading is currently disabled."), *ResolvedObject->GetPathName());
 				ResolvedObject = nullptr;
 			}
 		}
