[UE] Don't construct a default value for non-nullable references to placeholder-typed exports if they are serialized as NULL during the IDO serialization pass on editor load.

TL;DR:
- IDOs do not construct their own unique inner object graphs, but instead mirror their instanced counterparts.
- A default value is thus not required to be constructed for non-nullable values in the IDO load pass, because it is either already in sync with the serialized data via IDO creation, or will otherwise be synced up at instancing time.

Notes:
- IDOs "shadow-serialize" references to exports that are not otherwise explicitly overridden at save time. As with overridden values, schema changes can lead to load failures when the IDO attempts to serialize these references back in on load.
- IDO placeholder types mitigate these failures by allowing the exports to be created and serialized on load. However, references to placeholder-typed exports are not currently allowed to persist beyond the load phase and are currently being explicitly serialized as NULL.
- If the reference is shadow-serialized, then serializing the NULL value applies *only* during the IDO's load pass.
- The non-nullable failsafe path otherwise creates a temporary IDO state in this case, that can unintentionally lead to e.g. cook job failures due to errors in the log (e.g. - attempting to construct a default value of an abstract type). These failures can be resolved with a re-save, but the failure itself is unnecessary for IDOs.
- Furthermore, a re-save is not possible if IDOs are disabled in the interactive editor context, but enabled on the content service side, for example, and can lead to downstream failures that the end user cannot resolve.

#Synced-CL 49835867

diff --git a/Engine/Source/Runtime/CoreUObject/Private/UObject/PropertyObject.cpp b/Engine/Source/Runtime/CoreUObject/Private/UObject/PropertyObject.cpp
--- a/Engine/Source/Runtime/CoreUObject/Private/UObject/PropertyObject.cpp
+++ b/Engine/Source/Runtime/CoreUObject/Private/UObject/PropertyObject.cpp
@@ -287,13 +287,25 @@ void FObjectProperty::SerializeItem(FStructuredArchive::FSlot Slot, void* Value,
 
 void FObjectProperty::PostSerializeObjectItem(FArchive& SerializingArchive, void* Value, UObject* CurrentValue, UObject* ObjectValue, EObjectPropertyOptions Options /*= EObjectPropertyOptions::None*/, const void* Defaults /*= nullptr*/) const
 {
+	// Query the thread context to see if we're serializing an IDO.
+	auto IsSerializingInstanceDataObject = []()
+	{
+#if WITH_EDITORONLY_DATA
+		const FUObjectSerializeContext* SerializeContext = FUObjectThreadContext::Get().GetSerializeContext();
+		return SerializeContext ? SerializeContext->bImpersonateProperties : false;
+#else
+		return false;
+#endif
+	};
+
 	// Make sure non-nullable properties don't end up with null values.
 	if (!(Options & EObjectPropertyOptions::AllowNullValuesOnNonNullableProperty) &&
 		!ObjectValue && HasAnyPropertyFlags(CPF_NonNullable) &&
 		!SerializingArchive.IsSerializingDefaults() && // null values when Serializing CDOs are allowed, they will be fixed up later
 		!SerializingArchive.IsSaving() && // Constructing new objects when saving may confuse package saving code (new import/export created after import/export collection pass)
 		!SerializingArchive.IsTransacting() && // Don't create new objects when loading from the transaction buffer
-		!SerializingArchive.IsCountingMemory())
+		!SerializingArchive.IsCountingMemory() &&
+		!IsSerializingInstanceDataObject()) // Don't create new objects when serializing IDOs, they instead mirror the instanced object's graph
 	{
 		UObject* DefaultValue = nullptr;
 
