Deserialization and PostMigrate steps will now be performed inside of a transaction when receiving migrated remote objects

+ deprecated all IsRemote(..) functions and methods and replaced them with GetResidence(...) functions which now return an enum which represents three states of remote objects

#Synced-CL 49877584

diff --git a/Engine/Source/Runtime/Core/Private/UObject/RemoteObjectCore.cpp b/Engine/Source/Runtime/Core/Private/UObject/RemoteObjectCore.cpp
new file mode 100644
--- /dev/null
+++ b/Engine/Source/Runtime/Core/Private/UObject/RemoteObjectCore.cpp
@@ -0,0 +1,26 @@
+// Copyright Epic Games, Inc. All Rights Reserved.
+
+/*=============================================================================
+	RemoteObjectCore.cpp: Core module Remote Object utility functions.
+=============================================================================*/
+
+#include "UObject/RemoteObjectCore.h"
+
+#include "Misc/AssertionMacros.h"
+
+const TCHAR* EnumToString(EResidence InResidence)
+{
+	switch (InResidence)
+	{
+	case EResidence::Local:
+		return TEXT("Local");
+	case EResidence::LocalNotReady:
+		return TEXT("LocalNotReady");
+	case EResidence::Remote:
+		return TEXT("Remote");
+	default:
+		checkf(false, TEXT("EnumToString: Unknown EResidence"));
+		break;
+	};
+	return TEXT("");
+}
\ No newline at end of file
diff --git a/Engine/Source/Runtime/Core/Public/UObject/RemoteObjectCore.h b/Engine/Source/Runtime/Core/Public/UObject/RemoteObjectCore.h
new file mode 100644
--- /dev/null
+++ b/Engine/Source/Runtime/Core/Public/UObject/RemoteObjectCore.h
@@ -0,0 +1,14 @@
+// Copyright Epic Games, Inc. All Rights Reserved.
+
+#pragma once
+
+#include "CoreTypes.h"
+
+enum class EResidence : uint8
+{
+	Local /*UE_EXPERIMENTAL(5.8, "This enum type and its value names may change in the future")*/ = 0,			// The object is local and fully resolved (PostMigrated)
+	LocalNotReady /*UE_EXPERIMENTAL(5.8, "This enum type and its value names may change in the future")*/ = 1,	// The object's memory is local but requires resolving (deserialization and PostMigrate)
+	Remote /*UE_EXPERIMENTAL(5.8, "This enum type and its value names may change in the future")*/ = 2,			// The object is remote
+};
+
+CORE_API const TCHAR* EnumToString(EResidence InResidence);
\ No newline at end of file
diff --git a/Engine/Source/Runtime/Core/Public/UObject/WeakObjectPtrTemplates.h b/Engine/Source/Runtime/Core/Public/UObject/WeakObjectPtrTemplates.h
--- a/Engine/Source/Runtime/Core/Public/UObject/WeakObjectPtrTemplates.h
+++ b/Engine/Source/Runtime/Core/Public/UObject/WeakObjectPtrTemplates.h
@@ -9,6 +9,7 @@
 #include "Misc/UEOps.h"
 #include "UObject/WeakObjectPtrTemplatesFwd.h"
 #include "UObject/StrongObjectPtrTemplatesFwd.h"
+#include "UObject/RemoteObjectCore.h"
 
 #include <type_traits>
 
@@ -300,9 +301,15 @@ public:
 	}
 #endif // UE_WITH_REMOTE_OBJECT_HANDLE
 
+	UE_DEPRECATED(5.8, "Use GetResidence() instead.")
 	[[nodiscard]] FORCEINLINE bool IsRemote() const
 	{
-		return WeakPtr.IsRemote();
+		return WeakPtr.GetResidence() == EResidence::Remote;
+	}
+
+	[[nodiscard]] FORCEINLINE EResidence GetResidence() const
+	{
+		return WeakPtr.GetResidence();
 	}
 
 	/**
diff --git a/Engine/Source/Runtime/CoreUObject/Internal/UObject/UObjectMigrationContext.h b/Engine/Source/Runtime/CoreUObject/Internal/UObject/UObjectMigrationContext.h
--- a/Engine/Source/Runtime/CoreUObject/Internal/UObject/UObjectMigrationContext.h
+++ b/Engine/Source/Runtime/CoreUObject/Internal/UObject/UObjectMigrationContext.h
@@ -60,6 +60,12 @@ UE_API const TCHAR* ToString(EObjectMigrationSendType Value);
 /** Structure that holds context for what we're intending to do during migration */
 struct FUObjectMigrationContext final
 {
+	struct FMigratedObjectStatus
+	{
+		bool bIsOwned = false;
+		EResidence Residence = EResidence::Local;
+	};
+
 	/** The Object that is being migrated, it may not be the top-level UObject */
 	FRemoteObjectId ObjectId;
 
@@ -75,9 +81,25 @@ struct FUObjectMigrationContext final
 	/** If valid, the request id of the multi-server commit */
 	FRemoteTransactionId MultiServerCommitRequestId;
 
+	/** Maps remote object Ids to the ownership status on the current server (true means the object was owned by the MigrationSide at the moment the object's data was migrated) */
+	mutable TMap<FRemoteObjectId, FMigratedObjectStatus> CachedObjectStatus;
+
 	/** Are the values of this context valid? */
 	UE_API bool IsValid() const;
 
+	UE_API void CacheObjectStatus(const struct FRemoteObjectData& ObjectData);
+	UE_API const FMigratedObjectStatus& CacheObjectStatus(FRemoteObjectId InObjectId) const;
+
+	inline bool IsOwned(FRemoteObjectId InObjectId) const
+	{
+		return CachedObjectStatus.FindChecked(InObjectId).bIsOwned;
+	}
+
+	inline EResidence GetResidence(FRemoteObjectId InObjectId) const
+	{
+		return CachedObjectStatus.FindChecked(InObjectId).Residence;
+	}
+
 	/**
 	* During migration, figure out the role of the object being received
 	*/
diff --git a/Engine/Source/Runtime/CoreUObject/Private/UObject/PropertyBaseObject.cpp b/Engine/Source/Runtime/CoreUObject/Private/UObject/PropertyBaseObject.cpp
--- a/Engine/Source/Runtime/CoreUObject/Private/UObject/PropertyBaseObject.cpp
+++ b/Engine/Source/Runtime/CoreUObject/Private/UObject/PropertyBaseObject.cpp
@@ -890,7 +890,7 @@ void FObjectPropertyBase::CheckValidObject(void* ValueAddress, TObjectPtr<UObjec
 #if UE_WITH_REMOTE_OBJECT_HANDLE
 	// Verifying the object class will result in resolving remote objects and we don't want that.
 	// However, since remote objects are guarnteed to have compatible binaries it's not required to validate objects.
-	if (Object.IsRemote())
+	if (Object.GetResidence() != EResidence::Local)
 	{
 		return;
 	}
diff --git a/Engine/Source/Runtime/CoreUObject/Private/UObject/RemoteExecutor.cpp b/Engine/Source/Runtime/CoreUObject/Private/UObject/RemoteExecutor.cpp
--- a/Engine/Source/Runtime/CoreUObject/Private/UObject/RemoteExecutor.cpp
+++ b/Engine/Source/Runtime/CoreUObject/Private/UObject/RemoteExecutor.cpp
@@ -43,6 +43,24 @@ TMulticastDelegate<void(FRemoteTransactionId)> OnTransactionReleasedDelegate;
 
 TDelegate<void(FName /*Text*/)> OnRegionBeginDelegate;
 TDelegate<void(const FString& /*Text*/)> OnRegionEndDelegate;
+
+
+const TCHAR* EnumToString(ERemoteExecutorAbortReason AbortReason)
+{
+	switch (AbortReason)
+	{
+	case ERemoteExecutorAbortReason::Unspecified:
+		return TEXT("Unspecified");
+	case ERemoteExecutorAbortReason::RequiresDependencies:
+		return TEXT("RequiresDependencies");
+	case ERemoteExecutorAbortReason::AbandonWork:
+		return TEXT("AbandonWork");
+	default:
+		checkf(false, TEXT("EnumToString: Unknown abort reason"));
+		break;
+	};
+	return TEXT("");
+}
 }
 
 FArchive& operator<<(FArchive& Ar, FRemoteTransactionId& Id)
@@ -104,6 +122,7 @@ class FRemoteExecutor
 	uint32 NextTransactionRequestId = 0;
 
 	TArray<FRemoteExecutorWork> PendingWorks;
+	TArray<FRemoteExecutorWork> PendingMigrationWorks;
 	TArray<FRemoteExecutorWorkList> PendingWorkLists;
 
 	// GRemoteExecutorForcedNetPumpIterationCount controls how many iterations
@@ -162,6 +181,19 @@ public:
 		}
 	}
 
+	void EnqueueMigrationWork(FName InWorkName, FRemoteWorkPriority InWorkPriority, int32 InWorkListIndex, const TFunction<void(void)>& InWork)
+	{
+		FRemoteExecutorWork& NewWork = PendingMigrationWorks.Emplace_GetRef();
+		NewWork.Name = InWorkName;
+		NewWork.Priority = InWorkPriority;
+		NewWork.Work = InWork;
+		NewWork.bIsTransactional = true;
+		NewWork.RequestId = GenerateNextTransactionId();
+		NewWork.WorkListIndex = InWorkListIndex;
+
+		UE::RemoteExecutor::OnTransactionQueuedDelegate.Broadcast(NewWork.RequestId, NewWork.Name);
+	}
+
 	void EnqueueWorkList(FName WorkName, FRemoteWorkPriority InWorkPriority, bool bIsTransactional, const TArray<TFunction<void(void)>>& InWorkList)
 	{
 		FRemoteExecutorWorkList& NewWorkList = PendingWorkLists.Emplace_GetRef();
@@ -176,9 +208,9 @@ public:
 	{
 		check(ExecutingWork == nullptr);
 
-		if ((PendingWorks.Num() > 0) || (PendingWorkLists.Num() > 0))
+		if ((PendingWorks.Num() > 0) || (PendingWorkLists.Num() > 0) || (PendingMigrationWorks.Num() > 0))
 		{
-			UE_LOG(LogRemoteExec, VeryVerbose, TEXT("ExecutePendingWork BEGIN with %d work items and %d work lists"), PendingWorks.Num(), PendingWorkLists.Num());
+			UE_LOG(LogRemoteExec, VeryVerbose, TEXT("ExecutePendingWork BEGIN with %d work items, %d migration items and %d work lists"), PendingWorks.Num(), PendingMigrationWorks.Num(), PendingWorkLists.Num());
 		}
 
 		auto EnqueueFirstWorkFromList = [this](int32 WorkListIndex) -> void
@@ -301,6 +333,12 @@ public:
 				continue;
 			}
 
+			if (PendingMigrationWorks.Num())
+			{
+				PendingMigrationWorks += PendingWorks;
+				PendingWorks = MoveTemp(PendingMigrationWorks);
+			}
+
 			// enqueue the first work item from any new work lists
 			for (int32 WorkListIndex = 0; WorkListIndex < PendingWorkLists.Num(); WorkListIndex++)
 			{
@@ -452,8 +490,10 @@ public:
 								// You can set a breakpoint inside the abort handler below to see what caused the abort,
 								// or you can rely on the log output which only catches livelocks
 								const uint32 ExecutionAttempts = ExecutingWork->ExecutionAttempts;
-								UE_AUTORTFM_ONABORT(ExecutionAttempts)
+								UE_AUTORTFM_ONABORT(ExecutionAttempts, this)
 								{
+									UE_LOG(LogRemoteExec, VeryVerbose, TEXT("Transaction %s Aborted: %s: %s"), *ExecutingWork->Name.ToString(), EnumToString(AbortReason), *AbortReasonDescription);
+
 									constexpr uint32 MaxExecutionAttemptsBeforeLivelockWarning = 500;
 									if ((ExecutionAttempts >= MaxExecutionAttemptsBeforeLivelockWarning) && (ExecutionAttempts % MaxExecutionAttemptsBeforeLivelockWarning) == 0)
 									{
@@ -1144,6 +1184,17 @@ void EnqueueWork(FName WorkName, bool bIsTransactional, const TFunction<void(voi
 	EnqueueWorkWithExplicitPriority(WorkName, RootWorkPriority, bIsTransactional, InWork);
 }
 
+void EnqueueMigrationWork(FName WorkName, const TFunction<void(void)>& InWork)
+{
+#if UE_WITH_REMOTE_OBJECT_HANDLE && UE_AUTORTFM
+	FRemoteWorkPriority RootWorkPriority = FRemoteWorkPriority::CreateRootWorkPriority(FRemoteServerId::GetLocalServerId(), GRemoteExecutor.GenerateNextTransactionId());
+	GRemoteExecutor.EnqueueMigrationWork(WorkName, RootWorkPriority, INDEX_NONE, InWork);
+#else
+	InWork();
+#endif
+}
+
+
 void EnqueueWorkListWithExplicitPriority(FName WorkName, FRemoteWorkPriority WorkPriority, bool bIsTransactional, const TArray<TFunction<void(void)>>& InWorkList)
 {
 #if UE_WITH_REMOTE_OBJECT_HANDLE
diff --git a/Engine/Source/Runtime/CoreUObject/Private/UObject/RemoteObject.cpp b/Engine/Source/Runtime/CoreUObject/Private/UObject/RemoteObject.cpp
--- a/Engine/Source/Runtime/CoreUObject/Private/UObject/RemoteObject.cpp
+++ b/Engine/Source/Runtime/CoreUObject/Private/UObject/RemoteObject.cpp
@@ -79,6 +79,7 @@ std::atomic<uint64> AssetObjectSerialNumber(1);
 int32 UnsafeToMigrateObjects = 0; // This should go into TLS
 int32 GForceReturnObjectHandles = 0; // This should go into TLS
 int32 GInitRemoteServerIdBeforeUObjectInit = 0; // Initializes remote server id before UObject system is initialized
+thread_local bool GIsTransactionallyPostMigratingObjects = false;
 
 class FRemoteObjectStubMap : private TMap<FRemoteObjectId, UE::RemoteObject::Handle::FRemoteObjectStub*>
 {
@@ -207,6 +208,24 @@ public:
 		FRemoteObjectPathName** PathName = AssetPaths.Find(ObjectId);
 		return PathName ? *PathName : nullptr;
 	}
+
+	void ForEachStubWithMigratedData(TFunctionRef<bool(UE::RemoteObject::Handle::FRemoteObjectStub*)> Operation)
+	{
+		using namespace UE::RemoteObject::Handle;
+
+		UE::TScopeLock MapLock(ObjectMapCritical);
+		for (TPair<FRemoteObjectId,FRemoteObjectStub*>& Pair : RemoteObjects)
+		{
+			FRemoteObjectStub* Stub = Pair.Value;
+			if (Stub->MigratedData && !Stub->MigratedData->bCompleted)
+			{
+				if (!Operation(Stub))
+				{
+					return;
+				}
+			}
+		}
+	}
 };
 FRemoteObjectMaps* ObjectMaps = nullptr;
 
@@ -451,6 +470,20 @@ FScopedForceReturnObjectHandles::~FScopedForceReturnObjectHandles()
 	check(Private::GForceReturnObjectHandles >= 0);
 }
 
+void ForEachStubWithMigratedData(TFunctionRef<bool(UE::RemoteObject::Handle::FRemoteObjectStub*)> Operation)
+{
+	if (ObjectMaps)
+	{
+		ObjectMaps->ForEachStubWithMigratedData(Operation);
+	}
+}
+
+void SetTransactionallyPostMigratingObjects(bool bIsPostMigrating)
+{
+	checkf(IsInGameThread(), TEXT("Starting to transactionally post migrate objects outside of the game thread. This is not supported."))
+	GIsTransactionallyPostMigratingObjects = bIsPostMigrating;
+}
+
 } // namespace UE::RemoteObject::Private
 
 namespace UE::RemoteObject::Handle
@@ -506,58 +539,84 @@ FRemoteObjectStub::FRemoteObjectStub(UObject* Object)
 	}
 }
 
-bool IsRemote(FRemoteObjectId ObjectId)
+EResidence GetResidence(FRemoteObjectId ObjectId)
 {
 	using namespace UE::RemoteObject::Private;
 
-	if (!ObjectId.IsValid())
+	if (!ObjectId.IsValid() || !UE::RemoteObject::Private::ObjectMaps)
 	{
-		return false;
+		return EResidence::Local;
 	}
 
 	if (UObject* Object = StaticFindObjectFastInternal(ObjectId))
 	{
-		return IsRemote(Object);
+		return GetResidence(Object);
 	}
 
 	// if a stub exists but we don't have the object's memory (the UObject doesn't exist) then the object must be remote
-	if (FindRemoteObjectStub(ObjectId))
+	if (UE::RemoteObject::Handle::FRemoteObjectStub* Stub = FindRemoteObjectStub(ObjectId))
 	{
-		return true;
+		return Stub->MigratedData ? EResidence::LocalNotReady : EResidence::Remote;
 	}
 
 	const FRemoteServerId ServerId = ObjectId.GetServerId();
 	// Invalid server Id means local native classes which are created before a server has a chance to have an id assigned
-	return ServerId.IsValid() && !ServerId.IsLocal() && !ServerId.IsAsset();
+	return (ServerId.IsValid() && !ServerId.IsLocal() && !ServerId.IsAsset()) ? EResidence::Remote : EResidence::Local;
 }
 
-bool IsRemote(const FRemoteObjectStub* Stub)
+EResidence GetResidence(const FRemoteObjectStub* Stub)
 {
 	if (!Stub)
 	{
-		return false;
+		return EResidence::Local;
 	}
 
 	if (UObject* Object = StaticFindObjectFastInternal(Stub->Id))
 	{
-		return IsRemote(Object);
+		return GetResidence(Object);
 	}
 
-	// Similar to the above IsRemote(ObjectId): if a stub exists but we don't have the object's memory (the UObject doesn't exist) then the object must be remote
-	return true;
+	// Similar to the above GetResidence(ObjectId): if a stub exists but we don't have the object's memory (the UObject doesn't exist) then the object must be remote
+	return Stub->MigratedData ? EResidence::LocalNotReady : EResidence::Remote;
 }
 
-bool IsRemote(const UObject* Object)
+EResidence GetResidence(const UObject* Object)
 {
+	using namespace UE::RemoteObject::Private;
+
 #if UE_WITH_REMOTE_OBJECT_HANDLE
+	if (!Object || !UE::RemoteObject::Private::ObjectMaps)
+	{
+		return EResidence::Local;
+	}
 	const int32 InternalIndex = GUObjectArray.ObjectToIndex(Object);
 	const bool bIsRemote = InternalIndex >= 0 && GUObjectArray.IndexToObject(InternalIndex)->HasAnyFlags(EInternalObjectFlags::Remote);
-	return bIsRemote;
+	if (!bIsRemote)
+	{
+		UE::RemoteObject::Handle::FRemoteObjectStub* Stub = FindRemoteObjectStub(FRemoteObjectId(Object));
+		return Stub && Stub->MigratedData ? EResidence::LocalNotReady : EResidence::Local;
+	}
+	return EResidence::Remote;
 #else
-	return false;
+	return EResidence::Local;
 #endif
 }
 
+bool IsRemote(FRemoteObjectId ObjectId)
+{
+	return GetResidence(ObjectId) == EResidence::Remote;
+}
+
+bool IsRemote(const FRemoteObjectStub* Stub)
+{
+	return GetResidence(Stub) == EResidence::Remote;
+}
+
+bool IsRemote(const UObject* Object)
+{
+	return GetResidence(Object) == EResidence::Remote;
+}
+
 bool IsOwned(const UObject* Object)
 {
 	using namespace UE::CoreUObject::Private;
@@ -656,7 +715,7 @@ FRemoteObjectId GetRemoteObjectId(TFunctionRef<const UObject*()> CodeThatReturns
 #endif
 }
 
-UObject* ResolveObject(const FRemoteObjectStub* Stub, ERemoteReferenceType RefType /*= ERemoteReferenceType::Strong*/)
+UObject* ResolveObject(FRemoteObjectStub* Stub, ERemoteReferenceType RefType /*= ERemoteReferenceType::Strong*/)
 {
 	using namespace UE::RemoteExecutor;
 	using namespace UE::RemoteObject::Transfer;
@@ -683,7 +742,24 @@ UObject* ResolveObject(const FRemoteObjectStub* Stub, ERemoteReferenceType RefTy
 		return Object;
 	}
 
-	bool bRemoteObject = !Object || IsRemote(Object);
+	if (Stub->MigratedData)
+	{
+#if UE_WITH_REMOTE_OBJECT_HANDLE
+		if (IsForceReturnObjectHandles())
+		{
+			// Re-encode this as the FObjectHandle (which has bitfield rules)
+			return reinterpret_cast<UObject*>(FObjectHandle(Stub).Payload);
+		}
+#endif
+
+		TransactionallyMigrateObjects(Stub->MigratedData);
+	
+		Object = StaticFindObjectFastInternal(Stub->Id);
+		checkf(Object, TEXT("Tried to migrate %s from migrated data but the object does not exist"), *Stub->Id.ToString());
+		return Object;
+	}
+
+	bool bRemoteObject = !Object || GetResidence(Object) == EResidence::Remote;
 	if (bRemoteObject)
 	{
 		#if UE_WITH_REMOTE_OBJECT_HANDLE
@@ -703,7 +779,7 @@ UObject* ResolveObject(const FRemoteObjectStub* Stub, ERemoteReferenceType RefTy
 
 		// if running non-transactionally, object migrated immediately, so we just re-resolve
 		Object = StaticFindObjectFastInternal(Stub->Id);
-		bRemoteObject = !Object || IsRemote(Object);
+		bRemoteObject = !Object || GetResidence(Object) == EResidence::Remote;
 		checkf(IsSpeculationMode() || !bRemoteObject, TEXT("Failed to resolve remote object %s, either this code is not running in a transaction and should be, or the transaction failed to abort"), *Stub->Id.ToString());
 	}
 
@@ -759,6 +835,11 @@ bool CanResolveObject(FRemoteObjectId ObjectId, ECanResolveObjectBehavior CanRes
 
 	if (FRemoteObjectStub* Stub = ObjectMaps->FindStub(ObjectId))
 	{
+		if (Stub->MigratedData)
+		{
+			return true;
+		}
+
 		// A stub exists so the object memory is not local but we can (attempt to) migrate it back if:
 		// 1. It's not already Garbage (unless we're allowing a migration of Garbage objects)
 		if (CanResolveObjectBehavior == ECanResolveObjectBehavior::CannotResolveGarbageObjects && Stub->bWasGarbage)
@@ -807,6 +888,15 @@ UClass* GetClass(FRemoteObjectId ObjectId, ERemoteObjectGetClassBehavior GetClas
 	return nullptr;
 }
 
+FRemoteObjectStub* DebugFindRemoteObjectStub(FRemoteObjectId ObjectId)
+{
+	return UE::RemoteObject::Private::FindRemoteObjectStub(ObjectId);
+}
+
+bool IsTransactionallyPostMigratingObjects()
+{
+	return UE::RemoteObject::Private::GIsTransactionallyPostMigratingObjects;
+}
 
 } // namespace UE::RemoteObject::Handle
 
@@ -943,7 +1033,7 @@ namespace UE::CoreUObject::Private
 		return GetRemoteId() == RHS.GetRemoteId();
 	}
 
-	const UE::RemoteObject::Handle::FRemoteObjectStub* UE::CoreUObject::Private::FRemoteObjectHandlePrivate::ToStub() const
+	UE::RemoteObject::Handle::FRemoteObjectStub* UE::CoreUObject::Private::FRemoteObjectHandlePrivate::ToStub() const
 	{
 		using namespace UE::RemoteObject::Handle;
 
@@ -1196,9 +1286,10 @@ FString FRemoteObjectId::ToString(ERemoteIdToStringVerbosity InVerbosityOverride
 
 		if (Verbosity >= (int32)ERemoteIdToStringVerbosity::FullNameAttributes)
 		{
-			if (IsRemote(*this))
+			EResidence Residence = GetResidence(*this);
+			if (Residence != EResidence::Local)
 			{
-				AdditionalInfo += TEXT(" (remote)");
+				AdditionalInfo += FString::Printf(TEXT(" (%s)"), EnumToString(Residence));
 			}
 			if (IsOwned(*this))
 			{
diff --git a/Engine/Source/Runtime/CoreUObject/Private/UObject/RemoteObjectPrivate.h b/Engine/Source/Runtime/CoreUObject/Private/UObject/RemoteObjectPrivate.h
--- a/Engine/Source/Runtime/CoreUObject/Private/UObject/RemoteObjectPrivate.h
+++ b/Engine/Source/Runtime/CoreUObject/Private/UObject/RemoteObjectPrivate.h
@@ -7,11 +7,24 @@
 #include "UObject/NameTypes.h"
 #include "UObject/ObjectMacros.h"
 #include "UObject/Object.h"
+#include "UObject/UObjectMigrationContext.h"
+#include "UObject/RemoteObjectTransfer.h"
 
 DECLARE_LOG_CATEGORY_EXTERN(LogRemoteObject, Display, All);
 
 struct FRemoteObjectPathName;
 
+/** Structure that contains object data and migration context for deserializing and PostMigrating objects inside of a transaction */
+struct FTransactionalMigrationData : public FRefCountedObject
+{
+	/** Serialized data of the migrated objects */
+	FRemoteObjectData Data;
+	/** The context for the migrated objects */
+	FUObjectMigrationContext Context;
+	/** If true, the data has been fully deserialized and should no longer be used */
+	bool bCompleted = false;
+};
+
 namespace UE::RemoteObject::Private
 {
 	/**
@@ -140,6 +153,16 @@ namespace UE::RemoteObject::Private
 		}
 		return false;
 	}
+
+	/**
+	* Calls the provided function for each of the internally stored stubs that references valid migration data
+	*/
+	void ForEachStubWithMigratedData(TFunctionRef<bool(UE::RemoteObject::Handle::FRemoteObjectStub*)> Operation);
+
+	/**
+	* Sets an internal flag indicating if the current thread is transactionally deserializing and PostMigrating objects
+	*/
+	void SetTransactionallyPostMigratingObjects(bool bIsPostMigrating);
 }
 
 namespace UE::RemoteObject::Transfer::Private
diff --git a/Engine/Source/Runtime/CoreUObject/Private/UObject/RemoteObjectSerialization.cpp b/Engine/Source/Runtime/CoreUObject/Private/UObject/RemoteObjectSerialization.cpp
--- a/Engine/Source/Runtime/CoreUObject/Private/UObject/RemoteObjectSerialization.cpp
+++ b/Engine/Source/Runtime/CoreUObject/Private/UObject/RemoteObjectSerialization.cpp
@@ -410,12 +410,13 @@ public:
 	static FRemoteObjectReferenceInfo GetReferenceInfo(const FObjectPtr& ObjPtr, UObject* InRootObject)
 	{
 		using namespace UE::RemoteObject;
+		using namespace UE::RemoteObject::Handle;
 
 		FRemoteObjectReferenceInfo Info;
 		Info.Id = ObjPtr.GetRemoteId();
 		if (ObjPtr)
 		{
-			if (!ObjPtr.IsRemote())
+			if (ObjPtr.GetResidence() == EResidence::Local)
 			{
 				Info.Object = ObjPtr.Get();
 				Info.bIsSubobject = Info.Object->IsIn(InRootObject);
@@ -491,11 +492,12 @@ protected:
 	{
 		using namespace UE::CoreUObject::Private;
 		using namespace UE::RemoteObject::Private;
+		using namespace UE::RemoteObject::Handle;
 
 		FRemoteObjectReferenceInfo Info = GetReferenceInfo(ObjPtr);
 		WriteObjectReference(Info);
 
-		if (ObjPtr && !ObjPtr.IsRemote())
+		if (ObjPtr && ObjPtr.GetResidence() == EResidence::Local)
 		{
 			UObject* Object = ObjPtr.Get();
 
@@ -619,10 +621,11 @@ public:
 	virtual FArchive& operator<<(FSoftObjectPath& Value) override { return FArchiveUObject::SerializeSoftObjectPath(*this, Value); }
 	virtual FArchive& operator<<(FWeakObjectPtr& Value) override
 	{
+		using namespace UE::RemoteObject::Handle;
 		FObjectPtr Ptr;
 #if UE_WITH_REMOTE_OBJECT_HANDLE
 		const FRemoteObjectId WeakPtrId = Value.GetRemoteId();
-		if (UE::RemoteObject::Handle::IsRemote(WeakPtrId))
+		if (GetResidence(WeakPtrId) != EResidence::Local)
 		{
 			Ptr = FObjectPtr(WeakPtrId);
 		}
@@ -774,6 +777,7 @@ class FArchiveRemoteObjectReader : public TArchiveRemoteObjectBase<FMemoryReader
 	ERemoteReferenceType ReadObjectReference(FObjectPtr& Value)
 	{
 		using namespace UE::RemoteObject::Private;
+		using namespace UE::RemoteObject::Handle;
 
 		ERemoteReferenceType Type;		
 		*this << Type;
@@ -796,7 +800,7 @@ class FArchiveRemoteObjectReader : public TArchiveRemoteObjectBase<FMemoryReader
 				// In some situations (like resetting an object to its archetype state) we might want to preserve
 				// references to remote objects because we might end up migrating them mid-deserialization 
 				// Overwriting them could also potentially discard any changes made to them on another server
-				if (EnumHasAnyFlags(DeserializeFlags, ERemoteObjectSerializationFlags::PreserveRemoteReferences) && Value.IsRemote())
+				if (EnumHasAnyFlags(DeserializeFlags, ERemoteObjectSerializationFlags::PreserveRemoteReferences) && Value.GetResidence() != EResidence::Local)
 				{
 					bNeedsResolvingWithId = false;
 				}
@@ -817,7 +821,7 @@ class FArchiveRemoteObjectReader : public TArchiveRemoteObjectBase<FMemoryReader
 				// If the serialized reference was not found in memory or if the reference was serialized as id-only 
 				// keep it as an unresolved ObjectPtr and store a pointer to it so that we can try to resolve it after all objects have been deserialize
 #if UE_WITH_REMOTE_OBJECT_HANDLE
-				if (!EnumHasAnyFlags(DeserializeFlags, ERemoteObjectSerializationFlags::PreserveRemoteReferences) || !Value.IsRemote())
+				if (!EnumHasAnyFlags(DeserializeFlags, ERemoteObjectSerializationFlags::PreserveRemoteReferences) || Value.GetResidence() == EResidence::Local)
 				{
 					FObjectHandle Handle = FObjectHandle::FromIdNoResolve(ObjId);
 					Value = FObjectPtr(Handle);
@@ -853,7 +857,7 @@ public:
 	{
 		FObjectPtr Value;
 		ReadObjectReference(Value);
-		check(!Value.IsRemote());
+		check(Value.GetResidence() == EResidence::Local);
 		Obj = Value.Get();
 		return *this;
 	}
@@ -877,10 +881,12 @@ public:
 	virtual FArchive& operator<<(FSoftObjectPath& Value) override { return FArchiveUObject::SerializeSoftObjectPath(*this, Value); }
 	virtual FArchive& operator<<(FWeakObjectPtr& Value) override 
 	{
+		using namespace UE::RemoteObject::Handle;
+
 		FObjectPtr Ptr;
 		ReadObjectReference(Ptr);
 
-		if (!Ptr.IsRemote())
+		if (Ptr.GetResidence() == EResidence::Local)
 		{
 			Value = Ptr.Get();
 		}
@@ -937,7 +943,7 @@ public:
 
 	virtual FArchive& operator<<(FObjectPtr& Value) override
 	{
-		if (Value.IsRemote())
+		if (Value.GetResidence() == EResidence::Remote)
 		{
 			PatchObjectReference(Value);
 		}
@@ -998,7 +1004,7 @@ FRemoteObjectConstructionOverrides::FRemoteObjectConstructionOverrides(const FRe
 	}
 }
 
-UObject* ConstructRemoteObject(const FRemoteObjectData& ObjectData, const FRemoteObjectHeader& Header, ERemoteObjectSerializationFlags DeserializeFlags)
+UObject* ConstructRemoteObject(const FRemoteObjectData& ObjectData, const FRemoteObjectHeader& Header, ERemoteObjectSerializationFlags DeserializeFlags, const FUObjectMigrationContext* MigrationContext)
 {
 	using namespace UE::RemoteObject::Handle;
 	using namespace UE::RemoteObject::Private;
@@ -1006,7 +1012,7 @@ UObject* ConstructRemoteObject(const FRemoteObjectData& ObjectData, const FRemot
 	UClass* Class = CastChecked<UClass>(Header.Class.Get());
 	UObject* Outer = nullptr;
 #if UE_WITH_REMOTE_OBJECT_HANDLE
-	if (!Header.Outer.IsRemote())
+	if (Header.Outer.GetResidence() == EResidence::Local)
 	{
 		Outer = Header.Outer.Get();
 	}
@@ -1040,10 +1046,13 @@ UObject* ConstructRemoteObject(const FRemoteObjectData& ObjectData, const FRemot
 		}
 	}
 
-	// If not or the object is marked as remote (which means we brought it back before it was GC'd) (re)construct it
+	// Note that although we expect all deserialized objects to be remote or to not exist, we might be deserializing a default subobject (component) created in it's owner (Outer) constructor 
+	// in which case it would not be marked as remote. In this case we don't need to re-construct an object that was created in one of the previous calls to ConstructRemoteObject.
+	const bool bIsRemote = (MigrationContext ? MigrationContext->GetResidence(SerializedObjectId) : GetResidence(Object)) == EResidence::Remote;
+	// If not or the object is marked as remote (which means we brought it back before it was GC'd) (re)construct it.	
 	// Unless we explcitly want to re-use existing (valid / not marked as garbage) objects and skip re-construction to avoid side-effects
 	// Note that even then an object may not exist on this server (it could've been constructed on a different server when its owner was migrated or was simply GC'd)
-	const bool bSkipConstruction = IsValid(Object) && (!IsRemote(Object) || EnumHasAnyFlags(DeserializeFlags, ERemoteObjectSerializationFlags::UseExistingObjects));
+	const bool bSkipConstruction = IsValid(Object) && (!bIsRemote || EnumHasAnyFlags(DeserializeFlags, ERemoteObjectSerializationFlags::UseExistingObjects));
 	if (!bSkipConstruction)
 	{
 		FStaticConstructObjectParameters Params(Class);
@@ -1063,6 +1072,15 @@ UObject* ConstructRemoteObject(const FRemoteObjectData& ObjectData, const FRemot
 			Object = StaticConstructObject_Internal(Params);
 		}
 	}
+	else if (MigrationContext)
+	{
+		UE_CLOG(MigrationContext->GetObjectMigrationRecvType(SerializedObjectId) == EObjectMigrationRecvType::Borrowed,
+			LogRemoteSerialization, Warning, TEXT("Received a borrowed object that was not reconstructed %s (IsValid: %s, Residence: %s, UseExistingObjects: %s)"), *SerializedObjectId.ToString(),
+			IsValid(Object) ? TEXT("yes") : TEXT("no"),
+			EnumToString(MigrationContext->GetResidence(SerializedObjectId)),
+			EnumHasAnyFlags(DeserializeFlags, ERemoteObjectSerializationFlags::UseExistingObjects) ? TEXT("yes") : TEXT("no")
+		);
+	}
 
 	if (!Object)
 	{
@@ -1072,7 +1090,11 @@ UObject* ConstructRemoteObject(const FRemoteObjectData& ObjectData, const FRemot
 	checkf(FRemoteObjectId(Object) == SerializedObjectId, TEXT("Created an object with a different ID:%s than requested:%s"), *FRemoteObjectId(Object).ToString(), *SerializedObjectId.ToString());
 
 	MarkAsLocal(Object);
-
+	if (UE::RemoteObject::Handle::FRemoteObjectStub* Stub = UE::RemoteObject::Private::FindRemoteObjectStub(SerializedObjectId))
+	{
+		// If we don't reset the migrated data immediatetly we will enter an infinite recursion when trying to resolve this object's inner object's outer (this Object)
+		Stub->MigratedData = nullptr;
+	}
 	// Update internal flags on the migrated object. It's possible the object being migrated already existed in memory on this server and had the EInternalObjectFlags::Garbage flag set.
 	// Unless the migrated version also had this flag set we need to clear it.
 	// It's also possible that the local object didn't have this flag set but the migrated one has so we need to set it on this server too (it's not impossible to migrate objects marked as garbage)
@@ -1193,7 +1215,7 @@ FRemoteObjectData SerializeObjectData(UObject* InObject, TSet<UObject*>& OutObje
 		Ar.SetMigratingRemoteObjects(true);
 		SerializeObjectDataInternal(Ar, InObject, RequestedObjectId, OutObjects);
 	}
-
+	checkf(ObjectData.SerializedObjectHeaders.Num(), TEXT("Serialized object data for object %s is empty"), *RequestedObjectId.ToString());
 	return MoveTemp(ObjectData);
 }
 
@@ -1311,7 +1333,7 @@ int32 DeserializeObjectData(FRemoteObjectData& ObjectData,
 			}
 		}
 
-		const bool bReturningBorrowedObject = GResetBorrowedObjects && IsOwned(ObjectData.GetRootSerializedObjectId());
+		const bool bReturningBorrowedObject = GResetBorrowedObjects && (MigrationContext && MigrationContext->IsOwned(ObjectData.GetRootSerializedObjectId()));
 		if (bReturningBorrowedObject)
 		{
 			// Root object is owned by this server so we're receiving an object that was borrowed by another server.
@@ -1332,7 +1354,7 @@ int32 DeserializeObjectData(FRemoteObjectData& ObjectData,
 			{
 				FRemoteObjectId SerializedObjectId = ObjectData.GetRemoteObjectId(ObjectHeader.RemoteId);
 				OutObjectRemoteIds.Add(SerializedObjectId);
-				OutObjects.Add(ConstructRemoteObject(ObjectData, ObjectHeader, DeserializeFlags));
+				OutObjects.Add(ConstructRemoteObject(ObjectData, ObjectHeader, DeserializeFlags, MigrationContext));
 			}
 
 			checkf(OutObjectRemoteIds.Num() == ObjectData.SerializedObjectHeaders.Num(), TEXT("The number of constructed object ids (%d) does not match the number of serialized object headers (%d)"), OutObjectRemoteIds.Num(), ObjectData.SerializedObjectHeaders.Num());
diff --git a/Engine/Source/Runtime/CoreUObject/Private/UObject/RemoteObjectTransfer.cpp b/Engine/Source/Runtime/CoreUObject/Private/UObject/RemoteObjectTransfer.cpp
--- a/Engine/Source/Runtime/CoreUObject/Private/UObject/RemoteObjectTransfer.cpp
+++ b/Engine/Source/Runtime/CoreUObject/Private/UObject/RemoteObjectTransfer.cpp
@@ -24,7 +24,11 @@ namespace UE::RemoteObject::Transfer
 const FRemoteServerId DatabaseId(ERemoteServerIdConstants::Database);
 static TAutoConsoleVariable<bool> CVarRemoveMissingMigratedComponents(TEXT("DSTM.RemoveMissingMigratedComponents"), true, TEXT("When we receive less SubObjects (Components) than we expect, remove the extras that we have locally (assume they were removed on purpose)"));
 
-UObject* MigrateObjectFromObjectData(FRemoteObjectData& ObjectData, const FUObjectMigrationContext& MigrationContext);
+int32 GTransactionalMigration = 1;
+static FAutoConsoleVariableRef CVarTransactionalMigrationEnabled(
+	TEXT("DSTM.TransactionalMigration"),
+	GTransactionalMigration,
+	TEXT("When non-zero, migrated objects will be created, deserialized and PostMigrated inside of a transaction"));
 
 // Prevent us from having to match this definition to the declartion in the header just to initialize the delegate
 #define DEFINE_MIGRATION_DELEGATE(x) decltype(x) x
@@ -40,6 +44,28 @@ UObject* MigrateObjectFromObjectData(FRemoteObjectData& ObjectData, const FUObje
 	DEFINE_MIGRATION_DELEGATE(RestoreRemoteObjectDataDelegate);
 #undef DEFINE_MIGRATION_DELEGATE
 
+/**
+* Stores data received from a remote server and a migration context in their respective objects' stubs (possibly creating these stubs in the process)
+* Remote objects are not created or deserialized inside of this function but the resident server of the received objects is set to the local server id
+* @param ObjectData Data received from a remote server
+* @param MigrationContext Migration context
+* @returns Structure that holds the migrated data and the migration context
+*/
+TRefCountPtr<FTransactionalMigrationData> StoreMigratedObjectData(FRemoteObjectData& ObjectData, const FUObjectMigrationContext& MigrationContext);
+
+/**
+* Removes the specified migrated data from the stubs of objects serialized inside of the migrated data
+*/
+void ResetMigratedData(TRefCountPtr<FTransactionalMigrationData> MigratedData, FRemoteServerId);
+
+/**
+* Creates, deserializes and PostMigrates an object(s) from object data received from another server
+* @param ObjectData Data received from another server
+* @param MigrationContext Migration context
+* @return The object requested by the current server or pushed by the sending server (when transferring ownership)
+*/
+UObject* MigrateObjectFromObjectData(FRemoteObjectData& ObjectData, const FUObjectMigrationContext& MigrationContext);
+
 /**
  * Remote object transfer queue. 
  * Queued requests (both send and receive) are processed on the game thread since some systems (like RPCs) that are used for transfering object data are GT-only
@@ -131,12 +157,36 @@ public:
 					MigrationRequests.IndividualRequests.Num(),
 					*MigrationRequests.ObjectId.ToString());
 
+			FRemoteObjectId CanonicalRootObjectId;
+			FRemoteServerId OwnerServerId;		
+
 			UObject* Object = StaticFindObjectFastInternal(MigrationRequests.ObjectId);
+			TRefCountPtr<FTransactionalMigrationData> MigratedData;
 
-			if (Object && !Handle::IsRemote(Object))
+			if (Object && GetResidence(Object) == EResidence::Local)
 			{
-				UE_LOG(LogRemoteObject, VeryVerbose, TEXT("TickObjectMigrations obj %s (is local)"),
-					*MigrationRequests.ObjectId.ToString());
+				// find the canonical root object id to use for arbitration
+				UObject* const RootObject = FindCanonicalRootObjectForSerialization(Object);
+				CanonicalRootObjectId = UE::CoreUObject::Private::FObjectHandleUtils::GetRemoteId(RootObject);
+				OwnerServerId = GetOwnerServerId(RootObject);
+			}
+			else if (FRemoteObjectStub* Stub = FindRemoteObjectStub(MigrationRequests.ObjectId))
+			{
+				MigratedData = Stub->MigratedData;
+				if (MigratedData)
+				{
+					// Get the canonical root object id to use for arbitration
+					CanonicalRootObjectId = Stub->MigratedData->Data.GetRootSerializedObjectId();
+					FRemoteObjectStub* RootStub = Stub->Id == CanonicalRootObjectId ? Stub : FindRemoteObjectStub(CanonicalRootObjectId);
+					check(RootStub);
+					OwnerServerId = RootStub->OwningServerId;
+				}
+			}
+
+			if (CanonicalRootObjectId.IsValid())
+			{
+				UE_LOG(LogRemoteObject, VeryVerbose, TEXT("TickObjectMigrations from %s: obj %s (is local)"), *FRemoteServerId::GetLocalServerId().ToString(),
+					*MigrationRequests.ObjectId.ToString()); 
 
 				// look at the list of requests for this object and pick the one with the highest priority for arbitration
 				// (the list is sorted, so the first element is always the highest priority)
@@ -153,10 +203,6 @@ public:
 					*MigrationRequests.ObjectId.ToString(),
 					*HighestPriorityRequest->RequestPriority.ToString());
 
-				// find the canonical root object id to use for arbitration
-				UObject* const RootObject = FindCanonicalRootObjectForSerialization(Object);
-				const FRemoteObjectId CanonicalRootObjectId = UE::CoreUObject::Private::FObjectHandleUtils::GetRemoteId(RootObject);
-
 				// now look over all active transactions and see if this request is higher priority than all of them
 				bool bObjectLocked = false;
 
@@ -178,8 +224,6 @@ public:
 
 				if (!bObjectLocked)
 				{
-					FString ObjectName = Object->GetName();
-
 					static FName SendObjectRegionName("ObjectTransferQueue_SendObject");
 					UE::RemoteExecutor::OnRegionBeginDelegate.ExecuteIfBound(SendObjectRegionName);
 
@@ -189,7 +233,6 @@ public:
 						*HighestPriorityRequest->DestinationServerId.ToString(),
 						*HighestPriorityRequest->RequestPriority.ToString());
 
-					FRemoteServerId OwnerServerId = Handle::GetOwnerServerId(RootObject);
 					FUObjectMigrationContext MigrationContext {
 						.ObjectId = HighestPriorityRequest->ObjectId,
 						.RemoteServerId = HighestPriorityRequest->DestinationServerId,
@@ -197,10 +240,40 @@ public:
 						.MigrationSide = EObjectMigrationSide::Send
 					};
 					FScopedObjectMigrationContext ScopedMigrationContext(MigrationContext);
-
-					TSet<UObject*> SentObjects;
+					
+					TSet<TObjectPtr<UObject>> SentObjects;
 					TSet<UObject*> ReferencedObjects;
-					FRemoteObjectData ObjectData = SerializeObjectData(Object, SentObjects, ReferencedObjects, &MigrationContext);
+					FRemoteObjectData ObjectData;
+					if (!MigratedData)
+					{
+						TSet<UObject*> SerializedObjects;
+						ObjectData = SerializeObjectData(Object, SerializedObjects, ReferencedObjects, &MigrationContext);
+						SentObjects = ObjectPtrWrap(SerializedObjects);
+					}
+					else
+					{
+						ResetMigratedData(MigratedData, HighestPriorityRequest->DestinationServerId);
+						ObjectData = MoveTemp(MigratedData->Data);						
+						
+						for (FRemoteObjectId SerializedId : ObjectData)
+						{
+							UObject* SerializedObject = StaticFindObjectFastInternal(SerializedId);
+							if (SerializedObject)
+							{
+								SentObjects.Add(SerializedObject);
+							}
+							else
+							{
+#if UE_WITH_REMOTE_OBJECT_HANDLE
+								FRemoteObjectStub* Stub = FindRemoteObjectStub(SerializedId);
+								checkf(Stub, TEXT("Stub for %s must exist when sending the stored object data"), *SerializedId.ToString());
+								FObjectHandle Handle(Stub);
+								SentObjects.Add(TObjectPtr<UObject>(FObjectPtr(Handle)));
+#endif
+							}
+						};
+					}
+					MigrationContext.CacheObjectStatus(ObjectData);
 
 					checkf(SentObjects.Num() > 0, TEXT("unable to migrate object %s to %s"), *HighestPriorityRequest->ObjectId.ToString(), *HighestPriorityRequest->DestinationServerId.ToString());
 
@@ -209,14 +282,16 @@ public:
 						MarkAsRemoteReference(RO);
 					}
 
-					for (UObject* SO : SentObjects)
+#if UE_WITH_REMOTE_OBJECT_HANDLE
+					for (TObjectPtr<UObject> SO : SentObjects)
 					{
-						FRemoteObjectId SubObjectId = UE::CoreUObject::Private::FObjectHandleUtils::GetRemoteId(SO);
+						FRemoteObjectId SubObjectId = SO.GetRemoteId();
 						UE_LOG(LogRemoteObject, VeryVerbose, TEXT("TickObjectMigrations sending obj %s sent subobject : %s [%s]"),
 							*HighestPriorityRequest->ObjectId.ToString(),
 							*SubObjectId.ToString(),
-							*GetNameSafe(SO));
+							*SO.GetName());
 					}
+#endif // UE_WITH_REMOTE_OBJECT_HANDLE
 
 					FMigrateSendParams SendRemoteObjectParams = { .MigrationContext = MigrationContext, .ObjectData = ObjectData };
 					SendRemoteObject(SendRemoteObjectParams);
@@ -229,13 +304,27 @@ public:
 						// Notify that objects have been migrated to a remote server
 						OnObjectsSentDelegate.Broadcast(SentObjects, ObjectData, MigrationContext);
 
-						for (UObject* SentObject : SentObjects)
+						for (TObjectPtr<UObject> SentObject : SentObjects)
 						{
+							if (IsObjectHandleResolved_ForGC(SentObject.GetHandle()))
+							{
+								UObject* ObjectToSend = SentObject;
 #if UE_WITH_REMOTE_OBJECT_HANDLE
-							SentObject->PostMigrate(MigrationContext);
+								ObjectToSend->PostMigrate(MigrationContext);
+#endif								
+								MarkAsRemote(ObjectToSend, HighestPriorityRequest->DestinationServerId);
+								UE::RemoteObject::Handle::ChangeOwnerServerId(ObjectToSend, OwnerServerId);
+							}
+							else
+							{
+#if UE_WITH_REMOTE_OBJECT_HANDLE
+								FRemoteObjectStub* Stub = SentObject.GetHandle().ToStub();
+								checkf(Stub, TEXT("Stub must exist here"));
+
+								Stub->ResidentServerId = HighestPriorityRequest->DestinationServerId;
+								Stub->OwningServerId = OwnerServerId;
 #endif
-							MarkAsRemote(SentObject, HighestPriorityRequest->DestinationServerId);
-							UE::RemoteObject::Handle::ChangeOwnerServerId(SentObject, OwnerServerId);
+							}
 						}
 					}
 
@@ -265,7 +354,7 @@ public:
 						}
 					}
 
-					UE::RemoteExecutor::OnRegionEndDelegate.ExecuteIfBound(FString::Format(TEXT("{0} {1} -> {2}"), {*RootObject->GetName(), *CanonicalRootObjectId.ToString(), *MigrationContext.RemoteServerId.ToString()}));
+					UE::RemoteExecutor::OnRegionEndDelegate.ExecuteIfBound(FString::Format(TEXT("{0} -> {1}"), {*CanonicalRootObjectId.ToString(ERemoteIdToStringVerbosity::Name), *MigrationContext.RemoteServerId.ToString()}));
 
 					// all done, delete all requests for this object
 					PendingObjectMigrationRequests.RemoveAt(MigrationRequestsIndex, EAllowShrinking::No);
@@ -323,7 +412,7 @@ public:
 			FRemoteObjectId* RequiredCanonicalRootObjectId = ActiveRequest->RequiredObjectsCanonicalRootByObjectID.Find(RequiredObjectId);
 			FRemoteObjectId ObjectIdToRequest = RequiredCanonicalRootObjectId ? *RequiredCanonicalRootObjectId : RequiredObjectId;
 
-			if (Handle::IsRemote(ObjectIdToRequest))
+			if (Handle::GetResidence(ObjectIdToRequest) == EResidence::Remote)
 			{
 				// find the highest priority request that needs this object
 				FRemoteObjectRequest* HighestPriorityRequest = ActiveRequest;
@@ -401,6 +490,15 @@ public:
 					ActiveRequest->RequiredObjectsCanonicalRootByObjectID.Add(RequiredObjectId, CanonicalRootObjectId);
 					ActiveRequest->RequiredObjectsCanonicalRoots.Add(CanonicalRootObjectId);
 				}
+				else if (FRemoteObjectStub* Stub = FindRemoteObjectStub(RequiredObjectId))
+				{
+					if (Stub->MigratedData)
+					{
+						const FRemoteObjectId CanonicalRootObjectId = Stub->MigratedData->Data.GetRootSerializedObjectId();
+						ActiveRequest->RequiredObjectsCanonicalRootByObjectID.Add(RequiredObjectId, CanonicalRootObjectId);
+						ActiveRequest->RequiredObjectsCanonicalRoots.Add(CanonicalRootObjectId);
+					}
+				}
 			}
 		}
 	}
@@ -415,14 +513,19 @@ public:
 		ActiveRequest->UsedObjects.Reset();
 		ActiveRequest->NewRequiredObject = FRemoteObjectId();
 
-		UE_LOG(LogRemoteObject, VeryVerbose, TEXT("FRemoteObjectTransferQueue: TickRequest(%s) %d required objs (%d pending requests):"), *ActiveRequest->GetRequestId().ToString(), ActiveRequest->RequiredObjects.Num(), PendingObjectRequests.Num());
+		UE_CLOG(ActiveRequest->RequiredObjects.Num() || PendingObjectRequests.Num(), LogRemoteObject, VeryVerbose, TEXT("FRemoteObjectTransferQueue: TickRequest(%s) %d required objs (%d pending requests):"),
+			*ActiveRequest->GetRequestId().ToString(), ActiveRequest->RequiredObjects.Num(), PendingObjectRequests.Num());
 
 		if (UE_GET_LOG_VERBOSITY(LogRemoteObject) == ELogVerbosity::VeryVerbose)
 		{
 			for (const FRemoteObjectId& RequiredObjectId : ActiveRequest->RequiredObjects)
 			{
 				FRemoteObjectId* CanonicalObjectId = ActiveRequest->RequiredObjectsCanonicalRootByObjectID.Find(RequiredObjectId);
-				UE_LOG(LogRemoteObject, VeryVerbose, TEXT("FRemoteObjectTransferQueue: TickRequest(%s) %s [root %s] %s"), *ActiveRequest->GetRequestId().ToString(), *RequiredObjectId.ToString(), CanonicalObjectId ? *(CanonicalObjectId->ToString()) : TEXT("<unk>"), Handle::IsRemote(RequiredObjectId) ? TEXT("<remote>") : TEXT("<local>"));
+				UE_LOG(LogRemoteObject, VeryVerbose, TEXT("FRemoteObjectTransferQueue: TickRequest(%s) %s [root %s] <%s>"), 
+					*ActiveRequest->GetRequestId().ToString(), 
+					*RequiredObjectId.ToString(), 
+					CanonicalObjectId ? *(CanonicalObjectId->ToString()) : TEXT("<unk>"), 
+					EnumToString(Handle::GetResidence(RequiredObjectId)));
 			}
 		}
 
@@ -448,13 +551,13 @@ public:
 	{
 		for (const FRemoteObjectId& RequiredObjectId : ActiveRequest->RequiredObjects)
 		{
-			if (Handle::IsRemote(RequiredObjectId))
+			if (Handle::GetResidence(RequiredObjectId) == EResidence::Remote)
 			{
 				return false;
 			}
 		}
 
-		UE_LOG(LogRemoteObject, VeryVerbose, TEXT("FRemoteObjectTransferQueue: TickRequest(%s) all %d required objects are local"), *ActiveRequest->GetRequestId().ToString(), ActiveRequest->RequiredObjects.Num());
+		UE_CLOG(ActiveRequest->RequiredObjects.Num(), LogRemoteObject, VeryVerbose, TEXT("FRemoteObjectTransferQueue: TickRequest(%s) all %d required objects are local"), *ActiveRequest->GetRequestId().ToString(), ActiveRequest->RequiredObjects.Num());
 
 		return true;
 	}
@@ -482,7 +585,7 @@ public:
 			}
 
 			// if we used it this run and we're about to commit, we expect that the object is local
-			check(!Handle::IsRemote(UsedObjectId));
+			check(Handle::GetResidence(UsedObjectId) != EResidence::Remote);
 
 			UObject* UsedObject = StaticFindObjectFastInternal(UsedObjectId);
 			if (!ensure(UsedObject))
@@ -548,12 +651,13 @@ public:
 				.RemoteServerId = OwnerServerId,
 				.OwnerServerId = OwnerServerId,
 				.MigrationSide = EObjectMigrationSide::Send,
-				.MultiServerCommitRequestId = ActiveRequest->RequestId,
+				.MultiServerCommitRequestId = ActiveRequest->RequestId
 			};
 			FScopedObjectMigrationContext ScopedMigrationContext(MigrationContext);
 
 			TSet<UObject*> SentObjects;
 			FRemoteObjectData ObjectData = SerializeObjectData(Object, SentObjects, MultiServerCommitReferencedObjects, &MigrationContext);
+			MigrationContext.CacheObjectStatus(ObjectData);
 			FMigrateSendParams SendRemoteObjectParams = { .MigrationContext = MigrationContext, .ObjectData = ObjectData};
 
 			UE_AUTORTFM_OPEN
@@ -615,6 +719,7 @@ public:
 				.OwnerServerId = OwnerServerId,
 				.MigrationSide = EObjectMigrationSide::Send
 			};
+			MigrationContext.CacheObjectStatus(SentObjectId);
 			FScopedObjectMigrationContext ScopedMigrationContext(MigrationContext);
 
 			UE_LOG(LogRemoteObject, VeryVerbose, TEXT("FRemoteObjectTransferQueue: CommitMultiServerCommit postmigrate %s"), *SentObjectId.ToString());
@@ -713,13 +818,14 @@ public:
 				.OwnerServerId = DatabaseId,
 				.MigrationSide = EObjectMigrationSide::Receive
 			};
+			MigrationContext.CacheObjectStatus(ObjectId);
 			FScopedObjectMigrationContext ScopedMigrationContext(MigrationContext);
 			RestoreRemoteObjectDataDelegate.Execute(MigrationContext);
 
 #if UE_WITH_REMOTE_OBJECT_HANDLE
 			// We assume we have it restored, now forward it
 			TWeakObjectPtr WeakObjPtr { ObjectId };
-			ensureMsgf(!WeakObjPtr.IsExplicitlyNull() && WeakObjPtr.IsValid(true) && !WeakObjPtr.IsRemote(), TEXT("Expected RemoteObjectId %s to be restored from Database and thus locally owned"), *ObjectId.ToString());
+			ensureMsgf(!WeakObjPtr.IsExplicitlyNull() && WeakObjPtr.IsValid(true) && WeakObjPtr.GetResidence() != EResidence::Remote, TEXT("Expected RemoteObjectId %s to be restored from Database and thus locally owned"), *ObjectId.ToString());
 #endif
 
 			// If we didn't expect it locally, send it on to the final destination
@@ -746,9 +852,12 @@ public:
 	* @param ObjectId Id of the object that was requested
 	* @param Data Object data. Data ownership is transferred to the receive request (if it exists)
 	*/
-	void FulfillReceiveRequest(FRemoteObjectData& ObjectData, const FUObjectMigrationContext& MigrationContext)
+	FORCENOINLINE void FulfillReceiveRequest(FRemoteObjectData& ObjectData, const FUObjectMigrationContext& MigrationContext)
 	{
 		using namespace UE::RemoteObject::Serialization;
+		using namespace UE::RemoteObject::Private;
+		using namespace UE::RemoteObject::Handle;
+		using namespace UE::RemoteObject::Transfer;
 
 		check(!AutoRTFM::IsClosed());
 
@@ -766,12 +875,29 @@ public:
 
 		const bool bFoundRequest = (FoundPendingRequestIndex != INDEX_NONE);
 
-		UE_LOG(LogRemoteObject, Verbose, TEXT("FRemoteObjectTransferQueue: FulfillReceiveRequest %s (was requested: %d, owner: %s)"), *MigrationContext.ObjectId.ToString(), bFoundRequest, *MigrationContext.OwnerServerId.ToString());
+		UE_LOG(LogRemoteObject, VeryVerbose, TEXT("FRemoteObjectTransferQueue(%s): FulfillReceiveRequest %s (was requested: %s, owner server id: %s, transferring ownership: %s)"),
+			*FRemoteServerId::GetLocalServerId().ToString(),
+			*MigrationContext.ObjectId.ToString(), 
+			bFoundRequest ? TEXT("yes") : TEXT("no"),
+			*MigrationContext.OwnerServerId.ToString(),
+			(MigrationContext.GetObjectMigrationRecvType(MigrationContext.ObjectId) == EObjectMigrationRecvType::AssignedOwnership) ? TEXT("yes") : TEXT("no"));
 
-		UObject* Result = MigrateObjectFromObjectData(ObjectData, MigrationContext);
-		checkf(UE::CoreUObject::Private::FObjectHandleUtils::GetRemoteId(Result) == MigrationContext.ObjectId, TEXT("PeerId:%s requested object %s migration from PeerId:%s but received %s"),
-			   *FRemoteServerId::GetLocalServerId().ToString(), *MigrationContext.ObjectId.ToString(), *MigrationContext.RemoteServerId.ToString(), *UE::CoreUObject::Private::FObjectHandleUtils::GetRemoteId(Result).ToString());
-		checkf(MigrationContext.OwnerServerId == Handle::GetOwnerServerId(Result), TEXT("Expected owner id %s doesn't match post-deserialize owner: %s"), *MigrationContext.OwnerServerId.ToString(), *Handle::GetOwnerServerId(Result).ToString());
+#if UE_AUTORTFM
+		if (FRemoteObjectId::RemoteObjectSupportCompiledIn && GTransactionalMigration)
+		{
+			TRefCountPtr<FTransactionalMigrationData> StoredData = StoreMigratedObjectData(ObjectData, MigrationContext);
+	
+			static FName TransactionalWorkName("TransactionallyMigrateObjects");
+			UE::RemoteExecutor::EnqueueMigrationWork(TransactionalWorkName, [StoredData]()
+				{
+					TransactionallyMigrateObjects(StoredData);
+				});
+		}
+		else
+#endif
+		{
+			MigrateObjectFromObjectData(ObjectData, MigrationContext);
+		}
 
 		if (bFoundRequest)
 		{
@@ -808,6 +934,65 @@ public:
 
 FRemoteObjectTransferQueue GTransferQueue;
 
+TRefCountPtr<FTransactionalMigrationData> StoreMigratedObjectData(FRemoteObjectData& ObjectData, const FUObjectMigrationContext& MigrationContext)
+{
+	using namespace UE::RemoteObject::Serialization;
+	using namespace UE::RemoteObject::Private;
+	using namespace UE::RemoteObject::Handle;
+	using namespace UE::RemoteObject::Transfer;
+
+	FRemoteObjectStub* RootStub = FindOrAddRemoteObjectStub(MigrationContext.ObjectId, FRemoteServerId::GetLocalServerId());
+
+	TRefCountPtr<FTransactionalMigrationData> MigratedData = new FTransactionalMigrationData();
+	MigratedData->Data = ObjectData;
+	MigratedData->Context = MigrationContext;
+
+	UE_LOG(LogRemoteObject, VeryVerbose, TEXT("%s received %d objects"), *FRemoteServerId::GetLocalServerId().ToString(), MigratedData->Data.SerializedObjectHeaders.Num());
+
+	for (FSerializedRemoteObjectIterator It(MigratedData->Data); It; ++It)
+	{
+		FRemoteObjectId SerializedId = It.GetId();
+		FRemoteObjectStub* Stub = FindOrAddRemoteObjectStub(SerializedId, FRemoteServerId::GetLocalServerId());
+
+		checkf(!Stub->MigratedData, TEXT("Received remote object %s bundled with %s but the latter object's data already exists locally"), *MigratedData->Context.ObjectId.ToString(), *SerializedId.ToString());
+
+		Stub->OwningServerId = MigratedData->Context.OwnerServerId;
+		Stub->MigratedData = MigratedData;
+		Stub->ResidentServerId = FRemoteServerId::GetLocalServerId();
+		if (Stub->Name.IsNone())
+		{
+			Stub->Name = It.GetName();
+		}
+		if (!Stub->Class.IsValid())
+		{
+			Stub->Class = FRemoteObjectClass(It.GetClass());
+		}
+		if (UObject* Object = StaticFindObjectFastInternal(SerializedId))
+		{
+			MarkAsLocal(Object);
+		}
+	};
+
+	checkf(RootStub->MigratedData == MigratedData, TEXT("Data migrated for %s was not correctly stored in its stub"), *RootStub->Id.ToString());
+
+	for (int32 RequestIndex = 0; RequestIndex < GTransferQueue.GetRequestCount(); RequestIndex++)
+	{
+		FRemoteObjectRequest* ExistingRequest = GTransferQueue.GetRequestByIndex(RequestIndex);
+
+		for (FRemoteObjectId SerializedId : MigratedData->Data)
+		{
+			if (FRemoteObjectId* ExistingCanonicalRootForSubObject = ExistingRequest->RequiredObjectsCanonicalRootByObjectID.Find(SerializedId))
+			{
+				*ExistingCanonicalRootForSubObject = SerializedId;
+			}
+		}
+	}
+
+	checkf(RootStub->MigratedData == MigratedData, TEXT("Received migrated object data for %s but failed to set it on the root stub"), *MigrationContext.ObjectId.ToString());
+
+	return MigratedData;
+}
+
 UObject* MigrateObjectFromObjectData(FRemoteObjectData& ObjectData, const FUObjectMigrationContext& MigrationContext)
 {
 	using namespace UE::RemoteObject::Serialization;
@@ -822,8 +1007,6 @@ UObject* MigrateObjectFromObjectData(FRemoteObjectData& ObjectData, const FUObje
 	static FName GetObjectRegionName("ObjectTransferQueue_RecvObject");
 	UE::RemoteExecutor::OnRegionBeginDelegate.ExecuteIfBound(GetObjectRegionName);
 
-	check(!AutoRTFM::IsClosed());
-
  	FRemoteServerId LocalServerId = FRemoteServerId::GetLocalServerId();
 	TArray<FRemoteObjectId> ReceivedObjectRemoteIds;
 	TArray<UObject*> ReceivedObjects;
@@ -952,20 +1135,84 @@ UObject* MigrateObjectFromObjectData(FRemoteObjectData& ObjectData, const FUObje
 
 	UE::RemoteExecutor::OnRegionEndDelegate.ExecuteIfBound(FString::Format(TEXT("{0} {1} from {2} ({3} bytes)"), {*ReceivedObjects[0]->GetName(), *RootObjectRemoteId.ToString(), *MigrationContext.RemoteServerId.ToString(), ByteSize}));
 
-	return ReceivedObjects[RequestedObjectIndex];
+	UObject* Result = ReceivedObjects[RequestedObjectIndex];
+
+	checkf(ReceivedObjectRemoteIds.Contains(MigrationContext.ObjectId), TEXT("PeerId:%s requested object %s migration from PeerId:%s but received %s"),
+		*FRemoteServerId::GetLocalServerId().ToString(), *MigrationContext.ObjectId.ToString(), *MigrationContext.RemoteServerId.ToString(), *UE::CoreUObject::Private::FObjectHandleUtils::GetRemoteId(Result).ToString());
+	checkf(MigrationContext.OwnerServerId == Handle::GetOwnerServerId(Result), TEXT("Expected owner id %s doesn't match post-deserialize owner: %s"), *MigrationContext.OwnerServerId.ToString(), *Handle::GetOwnerServerId(Result).ToString());
+
+	return Result;
 } 
 
+void ResetMigratedData(TRefCountPtr<FTransactionalMigrationData> MigratedData, FRemoteServerId ResidentServerId)
+{
+	using namespace UE::RemoteObject::Handle;
+	using namespace UE::RemoteObject::Private;
+
+	checkf(MigratedData && MigratedData->Data.SerializedObjectHeaders.Num() > 0, TEXT("Trying to reset invalid migrated data"));
+	checkf(!MigratedData->bCompleted, TEXT("Trying to reset data that has already been fully deserialized"));
+
+	for (FRemoteObjectId SerializedId : MigratedData->Data)
+	{
+		FRemoteObjectStub* Stub = FindRemoteObjectStub(SerializedId);
+		checkf(Stub->MigratedData == MigratedData || !Stub->MigratedData, TEXT("Stub for serialized object id %s does not have the expected migrated object data"), *SerializedId.ToString());
+		Stub->ResidentServerId = ResidentServerId;
+		Stub->MigratedData = nullptr;
+	}
+	MigratedData->bCompleted = true;
+}
+
+void TransactionallyMigrateObjects(TRefCountPtr<FTransactionalMigrationData> MigratedData)
+{
+	using namespace UE::RemoteObject::Handle;
+	using namespace UE::RemoteObject::Private;
+
+	// It's possible that once object data is received and queued for transactional deserialization another transaction dereferences one of the migrated objects 
+	// and deserializes their data before the queued TransactionallyMigrateObjects transaction is executed.
+	// It's also possible that the migrated object data is sent to another server before we had a chance to deserialize it.
+	// In that case (bCompleted == true) just skip the deserialization and run some extra checks (if enabled)
+	if (!MigratedData->bCompleted)
+	{
+		SetTransactionallyPostMigratingObjects(true);
+		MigrateObjectFromObjectData(MigratedData->Data, MigratedData->Context);
+		SetTransactionallyPostMigratingObjects(false);
+
+		ResetMigratedData(MigratedData, FRemoteServerId::GetLocalServerId());
+	}
+#if DO_CHECK
+	else
+	{
+		for (FRemoteObjectId SerializedId : MigratedData->Data)
+		{
+			FRemoteObjectStub* Stub = FindRemoteObjectStub(SerializedId);
+			checkf(Stub->MigratedData != MigratedData, TEXT("Stub for serialized object %s still has its migrated data stored even though the data has been fully deserialized"), *SerializedId.ToString());
+		}
+	}
+#endif
+}
+
 void MigrateObjectFromRemoteServer(FRemoteObjectId ObjectId, FRemoteServerId CurrentOwnerServerId)
 {
+	using namespace UE::RemoteObject::Handle;
+	using namespace UE::RemoteObject::Private;
+
 	UE_LOG(LogRemoteObject, Verbose, TEXT("Aborting transaction, server %s needs obj id %s from server %s"),
 		*FRemoteServerId::GetLocalServerId().ToString(),
 		*ObjectId.ToString(),
 		*CurrentOwnerServerId.ToString());
 
+	FRemoteObjectStub* Stub = FindRemoteObjectStub(ObjectId);
+	if (Stub && Stub->MigratedData)
+	{
+		// We already received the object's data from another server so deserialize and PostMigrate it now
+		TransactionallyMigrateObjects(Stub->MigratedData);
+		return;
+	}
+
 #if !UE_AUTORTFM
 	{
 		FRemoteServerId LastKnownResidentServerId;
-		if (UE::RemoteObject::Handle::FRemoteObjectStub* Stub = UE::RemoteObject::Private::FindRemoteObjectStub(ObjectId))
+		if (Stub)
 		{
 			LastKnownResidentServerId = Stub->ResidentServerId;
 		}
@@ -1051,7 +1298,7 @@ void TouchResidentObject(UObject* Object)
 	{
 		UE_AUTORTFM_OPEN
 		{
-			const FRemoteObjectId ObjectId = UE::CoreUObject::Private::FObjectHandleUtils::GetRemoteId(Object);
+			const FRemoteObjectId ObjectId(Object);
 
 			bool bAlreadyInUsedObjectsList = false;
 			GTransferQueue.ActiveRequest->UsedObjects.Add(ObjectId, &bAlreadyInUsedObjectsList);
@@ -1061,6 +1308,7 @@ void TouchResidentObject(UObject* Object)
 				OnObjectTouchedDelegate.Broadcast(GTransferQueue.ActiveRequest->RequestId, ObjectId);
 			}
 
+
 			GTransferQueue.ActiveRequest->NewRequiredObject = ObjectId;
 		};
 	}
@@ -1099,6 +1347,7 @@ void TransferObjectOwnershipToRemoteServer(UObject* Object, FRemoteServerId Dest
 	TSet<UObject*> SentObjects;
 	TSet<UObject*> ReferencedObjects;
 	FRemoteObjectData ObjectData = SerializeObjectData(Object, SentObjects, ReferencedObjects, &MigrationContext);
+	MigrationContext.CacheObjectStatus(ObjectData);
 
 	UE_LOG(LogRemoteObject, VeryVerbose, TEXT("TransferObjectOwnershipToRemoteServer %s serialization complete (%d bytes)"),
 		*ObjectId.ToString(),
@@ -1129,7 +1378,7 @@ void TransferObjectOwnershipToRemoteServer(UObject* Object, FRemoteServerId Dest
 		FUnsafeToMigrateScope UnsafeToMigrateScope;
 
 		// Notify that objects have been migrated to a remote server
-		OnObjectsSentDelegate.Broadcast(SentObjects, ObjectData, MigrationContext);
+		OnObjectsSentDelegate.Broadcast(ObjectPtrWrap(SentObjects), ObjectData, MigrationContext);
 
 		for (UObject* SentObject : SentObjects)
 		{
@@ -1183,6 +1432,34 @@ void MigrateObjectToRemoteServerWithExplicitPriority(FRemoteWorkPriority Request
 	GTransferQueue.RequestObjectMigration(RequestPriority, ObjectId, DestinationServerId);
 }
 
+int32 ReturnObjectDataToOwnedServers()
+{
+	using namespace UE::RemoteObject::Handle;
+	using namespace UE::RemoteObject::Private;
+
+	int32 NumReturnedObjects = 0;
+
+	if (GTransactionalMigration)
+	{
+		ForEachStubWithMigratedData([&NumReturnedObjects](FRemoteObjectStub* Stub)
+			{
+				if (Stub->OwningServerId != FRemoteServerId::GetLocalServerId())
+				{
+					checkf(Stub->MigratedData, TEXT("Iterating over all stubs with migrated data but the stub for %s is not referencing any migrated data"), *Stub->Id.ToString());
+					FRemoteObjectId MigratedDataRootId = Stub->MigratedData->Data.GetRootSerializedObjectId();
+					if (Stub->Id == MigratedDataRootId)
+					{
+						MigrateObjectToRemoteServer(Stub->Id, Stub->OwningServerId);
+						NumReturnedObjects++;
+					}
+				}
+				return true;
+			});
+	}
+
+	return NumReturnedObjects;
+}
+
 void OnObjectDataReceived(FRemoteServerId OwnerServerId, FRemoteObjectId ObjectId, FRemoteServerId RemoteServerId, FRemoteObjectData& Data)
 {
 	check(!AutoRTFM::IsClosed());
@@ -1193,8 +1470,9 @@ void OnObjectDataReceived(FRemoteServerId OwnerServerId, FRemoteObjectId ObjectI
 		.OwnerServerId = OwnerServerId,
 		.MigrationSide = EObjectMigrationSide::Receive
 	};
+	MigrationContext.CacheObjectStatus(Data);
 	FScopedObjectMigrationContext ScopedMigrationContext(MigrationContext);
-	
+
 	GTransferQueue.FulfillReceiveRequest(Data, MigrationContext);
 }
 
@@ -1239,20 +1517,22 @@ void StoreObjectToDatabase(UObject* Object, FRemoteObjectId ObjectId)
 	using namespace UE::RemoteObject::Serialization;
 	using namespace UE::RemoteObject::Private;
 	using namespace UE::RemoteObject::Transfer;
+	using namespace UE::RemoteObject::Handle;
 
 	FRemoteServerId LocalServerId = FRemoteServerId::GetLocalServerId();
 	FRemoteServerId OwnerServerId = UE::RemoteObject::Handle::GetOwnerServerId(Object);
 	if (LocalServerId != OwnerServerId)
 	{
 		// We don't own this, and it's definitely remotely owned, just don't save it to the database
-		if (UE::RemoteObject::Handle::IsRemote(Object))
+		EResidence Residence = GetResidence(Object);
+		if (Residence == EResidence::Remote)
 		{
 			return;
 		}
 
 		// This ensure should never fire.  We shouldn't be able to store a non-locally-owned Object into the Database (that is the job of the Owning Server).  A scenario where I've seen this:
 		// We borrowed an Object that was Marked As Garbage and didn't return it to the Owning Server before a GC.
-		ensureMsgf(false, TEXT("%hs was about to store non-locally-owned Object %s to Server %s's Database; making it remote instead"), __func__, *GetNameSafe(Object), *LocalServerId.ToString());
+		ensureMsgf(false, TEXT("%hs was about to store non-locally-owned Object %s (residence: %d) to Server %s's Database; making it remote instead"), __func__, *GetNameSafe(Object), Residence, *LocalServerId.ToString());
 		MarkAsRemote(Object, OwnerServerId);
 		return;
 	}
@@ -1346,7 +1626,7 @@ FRemoteObjectReference::FRemoteObjectReference(const FObjectPtr& Ptr)
 	ObjectId = Ptr.GetRemoteId();
 	if (ObjectId.IsValid())
 	{
-		if (!Ptr.IsRemote())
+		if (Ptr.GetResidence() != EResidence::Remote)
 		{
 			ServerId = FRemoteServerId::GetLocalServerId();
 		}
@@ -1410,7 +1690,7 @@ bool FRemoteObjectReference::NetSerialize(FArchive& Ar, class UPackageMap* Map,
 		// so try to find the object in memory first by its id and if it doesn't exist try to find the object by pathname
 		// and if such object does exist use its remote id to resolve this reference to avoid migrating the asset from another server
 		UObject* ExistingObject = StaticFindObjectFastInternal(ObjectId);
-		if (!ExistingObject)
+		if (!ExistingObject && (ObjectId.IsAsset() || ObjectId.IsLocal()))
 		{
 			ExistingObject = PathName.Resolve();
 			if (ExistingObject)
@@ -1441,7 +1721,7 @@ void FRemoteObjectReference::NetDequantize(FRemoteObjectId InObjectId, FRemoteSe
 			// so try to find the object in memory first by its id and if it doesn't exist try to find the object by pathname
 			// and if such object does exist use its remote id to resolve this reference to avoid migrating the asset from another server
 			UObject* ExistingObject = StaticFindObjectFastInternal(ObjectId);
-			if (!ExistingObject)
+			if (!ExistingObject && (ObjectId.IsAsset() || ObjectId.IsLocal()))
 			{
 				ExistingObject = InPath.Resolve();
 				if (ExistingObject)
@@ -1468,3 +1748,15 @@ bool FRemoteObjectReference::Serialize(FArchive& Ar)
 
 	return true;
 }
+
+#if UE_WITH_REMOTE_OBJECT_HANDLE
+void UObject::PostMigrate(const FUObjectMigrationContext& MigrationContext)
+{
+	UE_LOG(LogRemoteObject, VeryVerbose, TEXT("PostMigrate %s %s"),
+		MigrationContext.MigrationSide == EObjectMigrationSide::Receive ? TEXT("Receiving") : TEXT("Sending"),
+		*FRemoteObjectId(this).ToString(ERemoteIdToStringVerbosity::PathName));
+
+	checkf(UE::RemoteObject::Transfer::GTransactionalMigration == 0 || MigrationContext.MigrationSide != EObjectMigrationSide::Receive || AutoRTFM::IsClosed(), 
+		TEXT("PostMigrate when receiving %s can only be called inside of a transaction"), *FRemoteObjectId(this).ToString(ERemoteIdToStringVerbosity::PathName));
+}
+#endif //UE_WITH_REMOTE_OBJECT_HANDLE
\ No newline at end of file
diff --git a/Engine/Source/Runtime/CoreUObject/Private/UObject/UObjectGlobals.cpp b/Engine/Source/Runtime/CoreUObject/Private/UObject/UObjectGlobals.cpp
--- a/Engine/Source/Runtime/CoreUObject/Private/UObject/UObjectGlobals.cpp
+++ b/Engine/Source/Runtime/CoreUObject/Private/UObject/UObjectGlobals.cpp
@@ -3628,7 +3628,7 @@ UObject* StaticAllocateObject
 		// Subobjects are always created in the constructor, no need to re-create them here unless their archetype != CDO or they're blueprint generated or they're remote
 		if (!bCreatingCDO && (!bCanRecycleSubobjects || !Obj->IsDefaultSubobject()
 #if UE_WITH_REMOTE_OBJECT_HANDLE
-			|| UE::RemoteObject::Handle::IsRemote(Obj)
+			|| (UE::RemoteObject::Handle::GetResidence(Obj) == EResidence::Remote)
 #endif // UE_WITH_REMOTE_OBJECT_HANDLE
 			))
 		{
diff --git a/Engine/Source/Runtime/CoreUObject/Private/UObject/UObjectMigrationContext.cpp b/Engine/Source/Runtime/CoreUObject/Private/UObject/UObjectMigrationContext.cpp
--- a/Engine/Source/Runtime/CoreUObject/Private/UObject/UObjectMigrationContext.cpp
+++ b/Engine/Source/Runtime/CoreUObject/Private/UObject/UObjectMigrationContext.cpp
@@ -16,6 +16,30 @@ EObjectMigrationRecvType FUObjectMigrationContext::GetObjectMigrationRecvType(co
 	return GetObjectMigrationRecvType(FRemoteObjectId { Object });
 }
 
+const FUObjectMigrationContext::FMigratedObjectStatus& FUObjectMigrationContext::CacheObjectStatus(FRemoteObjectId InObjectId) const
+{
+	if (const FMigratedObjectStatus* CachedStatus = CachedObjectStatus.Find(InObjectId))
+	{
+		checkf(CachedStatus->bIsOwned == UE::RemoteObject::Handle::IsOwned(InObjectId), TEXT("Ownership changed after it was first cached (was owned: %s, is owned: %s"),
+			CachedStatus->bIsOwned ? TEXT("yes") : TEXT("no"),
+			UE::RemoteObject::Handle::IsOwned(InObjectId) ? TEXT("yes") : TEXT("no"));
+		checkf(CachedStatus->Residence == UE::RemoteObject::Handle::GetResidence(InObjectId), TEXT("Residence changed after it was first cached (was: %s, now: %s)"),
+			EnumToString(CachedStatus->Residence),
+			EnumToString(UE::RemoteObject::Handle::GetResidence(InObjectId)));
+		return *CachedStatus;
+	}
+
+	return CachedObjectStatus.Add(InObjectId, { UE::RemoteObject::Handle::IsOwned(InObjectId), UE::RemoteObject::Handle::GetResidence(InObjectId) });
+}
+
+void FUObjectMigrationContext::CacheObjectStatus(const FRemoteObjectData& ObjectData)
+{
+	for (FRemoteObjectId SerializedId : ObjectData)
+	{
+		CacheObjectStatus(SerializedId);
+	}
+}
+
 EObjectMigrationRecvType FUObjectMigrationContext::GetObjectMigrationRecvType(FRemoteObjectId InObjectId) const
 {
 #if UE_WITH_REMOTE_OBJECT_HANDLE
@@ -25,7 +49,7 @@ EObjectMigrationRecvType FUObjectMigrationContext::GetObjectMigrationRecvType(FR
 	}
 
 	const FRemoteServerId LocalServerId = FRemoteServerId::GetLocalServerId();
-	const bool bAlreadyOwns = UE::RemoteObject::Handle::IsOwned(InObjectId);
+	const bool bAlreadyOwns = IsOwned(InObjectId);
 
 	const bool bAssignOwnership = !bAlreadyOwns && (OwnerServerId == LocalServerId);
 	const bool bReturnLoan = bAlreadyOwns && ensureMsgf(OwnerServerId == LocalServerId, TEXT("We are receiving object %s that we own, so the data we received should affirm we are the owner."), *InObjectId.ToString());
@@ -63,8 +87,7 @@ EObjectMigrationSendType FUObjectMigrationContext::GetObjectMigrationSendType(FR
 	}
 
 	const FRemoteServerId LocalServerId = FRemoteServerId::GetLocalServerId();
-	const bool bCurrentlyOwns = UE::RemoteObject::Handle::IsOwned(InObjectId);
-
+	const bool bCurrentlyOwns = CacheObjectStatus(InObjectId).bIsOwned;
 	const bool bOwnershipTransfer = (OwnerServerId == RemoteServerId) && bCurrentlyOwns;
 	const bool bReturnBorrowed = (OwnerServerId == RemoteServerId) && !bCurrentlyOwns;
 
diff --git a/Engine/Source/Runtime/CoreUObject/Private/UObject/WeakObjectPtr.cpp b/Engine/Source/Runtime/CoreUObject/Private/UObject/WeakObjectPtr.cpp
--- a/Engine/Source/Runtime/CoreUObject/Private/UObject/WeakObjectPtr.cpp
+++ b/Engine/Source/Runtime/CoreUObject/Private/UObject/WeakObjectPtr.cpp
@@ -35,7 +35,7 @@ void FWeakObjectPtr::operator=(FObjectPtr ObjectPtr)
 		ObjectRemoteId = ObjectPtr.GetRemoteId();
 
 		// if the object is local, fill in the index and serial number immediately
-		if (!ObjectPtr.IsRemote())
+		if (ObjectPtr.GetResidence() == EResidence::Local)
 #endif
 		{
 			const UObject* Object = ObjectPtr.Get();
@@ -237,8 +237,8 @@ FUObjectItem* FWeakObjectPtr::TryResolveRemoteObject() const
 	UObject* ResolvedObject = StaticFindObjectFastInternal(RemoteId);
 
 	if (ResolvedObject)
-	{		
-		if (UE::RemoteObject::Handle::IsRemote(ResolvedObject))
+	{
+		if (UE::RemoteObject::Handle::GetResidence(ResolvedObject) != EResidence::Local)
 		{
 			// Object memory is still on this server but it's marked as remote so we need to resolve it
 			ResolvedObject = ResolveObject(ResolvedObject, ERemoteReferenceType::Weak);
diff --git a/Engine/Source/Runtime/CoreUObject/Public/UObject/Object.h b/Engine/Source/Runtime/CoreUObject/Public/UObject/Object.h
--- a/Engine/Source/Runtime/CoreUObject/Public/UObject/Object.h
+++ b/Engine/Source/Runtime/CoreUObject/Public/UObject/Object.h
@@ -557,7 +557,7 @@ public:
 	/**
 	 * Called when an object is about to be migrated and after it's been migrated to another process.
 	 */
-	virtual void PostMigrate(const struct FUObjectMigrationContext& MigrationContext) {}
+	virtual void PostMigrate(const struct FUObjectMigrationContext& MigrationContext);
 
 	virtual bool IsMigrationRoot() const { return false; }
 #endif
diff --git a/Engine/Source/Runtime/CoreUObject/Public/UObject/ObjectHandle.h b/Engine/Source/Runtime/CoreUObject/Public/UObject/ObjectHandle.h
--- a/Engine/Source/Runtime/CoreUObject/Public/UObject/ObjectHandle.h
+++ b/Engine/Source/Runtime/CoreUObject/Public/UObject/ObjectHandle.h
@@ -261,7 +261,7 @@ namespace UE::CoreUObject::Private
 		COREUOBJECT_API void AssignFromObjectPointer(const UObject* Value);
 		COREUOBJECT_API void AssignFromRemoteId(FRemoteObjectId Value);
 
-		COREUOBJECT_API const UE::RemoteObject::Handle::FRemoteObjectStub* ToStub() const;
+		COREUOBJECT_API UE::RemoteObject::Handle::FRemoteObjectStub* ToStub() const;
 
 		COREUOBJECT_API FRemoteObjectId GetRemoteId() const;
 
@@ -339,7 +339,7 @@ inline bool IsObjectHandleResolved(FObjectHandle Handle)
 			return true;
 		}
 
-		return !UE::RemoteObject::Handle::IsRemote(ReadObjectHandlePointerNoCheck(Handle));
+		return UE::RemoteObject::Handle::GetResidence(ReadObjectHandlePointerNoCheck(Handle)) == EResidence::Local;
 	}
 
 	return false;
@@ -366,6 +366,7 @@ inline bool IsObjectHandleResolved_ForGC(FObjectHandle Handle)
  * nullptr counts as not remote
  * all handles are not remote when remote handles are not enabled
  */
+UE_DEPRECATED(5.8, "Use GetObjectHandleResidence() instead.")
 inline bool IsObjectHandleRemote(FObjectHandle Handle)
 {
 #if UE_WITH_REMOTE_OBJECT_HANDLE
@@ -388,6 +389,31 @@ inline bool IsObjectHandleRemote(FObjectHandle Handle)
 	return false;
 }
 
+/* checks the residence of the object represented by the specified handle.
+ * nullptr counts as local
+ * all handles are local when remote handles are not enabled
+ */
+inline EResidence GetObjectHandleResidence(FObjectHandle Handle)
+{
+#if UE_WITH_REMOTE_OBJECT_HANDLE
+	using namespace UE::CoreUObject::Private;
+
+	switch (Handle.ExtractHandleType())
+	{
+	case ERemoteObjectHandleType::ObjectPointer:
+		return UE::RemoteObject::Handle::GetResidence(Handle.ExtractObjectPointer());
+	case ERemoteObjectHandleType::StubPointer:
+		return UE::RemoteObject::Handle::GetResidence(Handle.ExtractStubPointer());
+	case ERemoteObjectHandleType::RemoteId:
+		return UE::RemoteObject::Handle::GetResidence(Handle.ToStub());
+	default:
+		checkf(false, TEXT("Unexpected remote object handle type"));
+		break;
+	}
+#endif
+	return EResidence::Local;
+}
+
 /* return true if a handle is type safe.
  * null and resolved handles are considered type safe
  */ 
@@ -649,6 +675,10 @@ namespace UE::CoreUObject::Private
 		else if (LocalHandle.ExtractHandleType() == ERemoteObjectHandleType::ObjectPointer)
 		{
 			ResolvedObject = const_cast<UObject*>(LocalHandle.ExtractObjectPointer());
+			if (ResolvedObject != nullptr && UE::RemoteObject::Handle::GetResidence(ResolvedObject) != EResidence::Local)
+			{
+				ResolvedObject = UE::RemoteObject::Handle::ResolveObject(ResolvedObject);
+			}
 		}
 
 		// do not overwrite the internal handle if it is on the remote heap, it must
diff --git a/Engine/Source/Runtime/CoreUObject/Public/UObject/ObjectPtr.h b/Engine/Source/Runtime/CoreUObject/Public/UObject/ObjectPtr.h
--- a/Engine/Source/Runtime/CoreUObject/Public/UObject/ObjectPtr.h
+++ b/Engine/Source/Runtime/CoreUObject/Public/UObject/ObjectPtr.h
@@ -200,8 +200,11 @@ public:
 	/**
 	 * FObjectPtr::IsRemote is only used when UE_WITH_REMOTE_OBJECT_HANDLE is true, and is mutually exclusive with other uses of unresolved FObjectPtr.
 	 * When using UE_WITH_REMOTE_OBJECT_HANDLE, unresolved FObjectPtrs mean the object is not locally available and is present on another server, see @RemoteObject.h
-	*/
-	[[nodiscard]] FORCEINLINE bool IsRemote() const { return IsObjectHandleRemote(Handle); }
+	*/	
+	[[nodiscard]] FORCEINLINE EResidence GetResidence() const { return GetObjectHandleResidence(Handle); }
+
+	UE_DEPRECATED(5.8, "Use GetResidence() instead.")
+	[[nodiscard]] FORCEINLINE bool IsRemote() const { return GetResidence() == EResidence::Remote; }
 
 	// Gets the PathName of the object without resolving the object reference.
 	// @TODO OBJPTR: Deprecate this.
@@ -746,9 +749,16 @@ public:
 	{
 		return ObjectPtr.IsResolved();
 	}
+	UE_DEPRECATED(5.8, "Use GetResidence() instead.")
 	[[nodiscard]] FORCEINLINE bool IsRemote() const 
 	{
+		PRAGMA_DISABLE_DEPRECATION_WARNINGS
 		return ObjectPtr.IsRemote();
+		PRAGMA_ENABLE_DEPRECATION_WARNINGS
+	}
+	[[nodiscard]] FORCEINLINE EResidence GetResidence() const
+	{
+		return ObjectPtr.GetResidence();
 	}
 	[[nodiscard]] FORCEINLINE FString GetPath() const 
 	{
@@ -1281,7 +1291,7 @@ template <typename T>
 [[nodiscard]] FORCEINLINE TWeakObjectPtr<T> MakeWeakObjectPtr(TObjectPtr<T> Ptr)
 {
 #if UE_WITH_REMOTE_OBJECT_HANDLE
-	if (Ptr.IsRemote())
+	if (Ptr.GetResidence() != EResidence::Local)
 	{
 		return TWeakObjectPtr<T>(Ptr.GetRemoteId());
 	}
diff --git a/Engine/Source/Runtime/CoreUObject/Public/UObject/RemoteExecutor.h b/Engine/Source/Runtime/CoreUObject/Public/UObject/RemoteExecutor.h
--- a/Engine/Source/Runtime/CoreUObject/Public/UObject/RemoteExecutor.h
+++ b/Engine/Source/Runtime/CoreUObject/Public/UObject/RemoteExecutor.h
@@ -264,6 +264,7 @@ namespace UE::RemoteExecutor
 	COREUOBJECT_API void ExecuteTransactionalWithExplicitPriority(FName WorkName, FRemoteWorkPriority WorkPriority, const TFunctionRef<void(void)>& Work);
 	COREUOBJECT_API void EnqueueWorkWithExplicitPriority(FName WorkName, FRemoteWorkPriority WorkPriority, bool bIsTransactional, const TFunction<void(void)>& Work);
 	COREUOBJECT_API void EnqueueWork(FName WorkName, bool bIsTransactional, const TFunction<void(void)>& Work);
+	COREUOBJECT_API void EnqueueMigrationWork(FName WorkName, const TFunction<void(void)>& InWork);
 	COREUOBJECT_API void EnqueueWorkListWithExplicitPriority(FName WorkName, FRemoteWorkPriority WorkPriority, bool bIsTransactional, const TArray<TFunction<void(void)>>& WorkList);
 	COREUOBJECT_API void EnqueueWorkList(FName WorkName, bool bIsTransactional, const TArray<TFunction<void(void)>>& WorkList);
 	COREUOBJECT_API void ExecutePendingWork();
diff --git a/Engine/Source/Runtime/CoreUObject/Public/UObject/RemoteObject.h b/Engine/Source/Runtime/CoreUObject/Public/UObject/RemoteObject.h
--- a/Engine/Source/Runtime/CoreUObject/Public/UObject/RemoteObject.h
+++ b/Engine/Source/Runtime/CoreUObject/Public/UObject/RemoteObject.h
@@ -3,11 +3,14 @@
 #pragma once
 
 #include "CoreTypes.h"
+#include "UObject/RemoteObjectCore.h"
 #include "UObject/RemoteObjectTypes.h"
 #include "UObject/ObjectMacros.h"
+#include "Templates/RefCounting.h"
 
 class UObject;
 class UClass;
+struct FTransactionalMigrationData;
 
 /**
  * Remote objects are unique UObjects that are referenced by a local server instance but their memory is actually owned by (exists on) another server.
@@ -56,14 +59,15 @@ namespace UE::RemoteObject::Handle
 		COREUOBJECT_API UClass* GetClass() const;
 	};
 
-	// Structure that holds basic information about a remote object
-	// This is what FObjectPtr that references a remote object actually points to after the remote object's memory has been claimed by GC
+	// Structure that holds basic information about an object that has been migrated (from or to the local server)
+	// This is what FObjectPtr that references a remote object actually points to after the remote object's memory has been claimed by GC or before the remote object has been resolved.
 	struct FRemoteObjectStub
 	{
 		FRemoteObjectId Id;
 		FRemoteObjectId OuterId;
 		FName Name;
 		FRemoteObjectClass Class;
+		TRefCountPtr<FTransactionalMigrationData> MigratedData;
 		bool bWasGarbage = false;
 
 		/** SerialNumber this object had on this server */
@@ -113,7 +117,7 @@ namespace UE::RemoteObject::Handle
 	* @param RefType Reference type that wants to resolve an object
 	* @return Resolved object
 	*/
-	COREUOBJECT_API UObject* ResolveObject(const FRemoteObjectStub* Stub, ERemoteReferenceType RefType = ERemoteReferenceType::Strong);
+	COREUOBJECT_API UObject* ResolveObject(FRemoteObjectStub* Stub, ERemoteReferenceType RefType = ERemoteReferenceType::Strong);
 
 	/**
 	* Resolves a remote object, aborting the active transaction if the object is unavailable
@@ -139,22 +143,46 @@ namespace UE::RemoteObject::Handle
 	*/
 	COREUOBJECT_API bool CanResolveObject(FRemoteObjectId ObjectId, ECanResolveObjectBehavior CanResolveObjectBehavior);
 
+	/**
+	* Gets the current residency of an object
+	* @param ObjectId Id of the object to check
+	* @return Residency of the object
+	*/
+	COREUOBJECT_API EResidence GetResidence(FRemoteObjectId ObjectId);
+
+	/**
+	* Gets the current residency of an object
+	* @param Stub Remote object stub representing the object to check
+	* @return Residency of the object
+	*/
+	COREUOBJECT_API EResidence GetResidence(const FRemoteObjectStub* Stub);
+
+	/**
+	* Gets the current residency of an object
+	* @param Object Object whose residency is to be checked
+	* @return Residency of the object
+	*/
+	COREUOBJECT_API EResidence GetResidence(const UObject* Object);
+
 	/**
 	* Checks if an object associated with the specified unique id is remote
 	* @return True if an object associated with the specified unique id is remote
 	*/
+	UE_DEPRECATED(5.8, "Use GetResidence(ObjectId) instead.")
 	COREUOBJECT_API bool IsRemote(FRemoteObjectId ObjectId);
 
 	/**
 	* Checks if an object associated with the specified stub is remote
 	* @return True if an object associated with the specified stub is remote
 	*/
+	UE_DEPRECATED(5.8, "Use GetResidence(Stub) instead.")
 	COREUOBJECT_API bool IsRemote(const FRemoteObjectStub* Stub);
 
 	/**
 	* Checks if an object (memory that has not yet been GC'd) is remote
 	* @return True if the object is remote
 	*/
+	UE_DEPRECATED(5.8, "Use GetResidence(Object) instead.")
 	COREUOBJECT_API bool IsRemote(const UObject* Object);
 
 	/**
@@ -188,4 +216,9 @@ namespace UE::RemoteObject::Handle
 	* Sets the owner server id for a locally resident object
 	*/
 	COREUOBJECT_API void ChangeOwnerServerId(const UObject* Object, FRemoteServerId NewOwnerServerId);
+
+	/**
+	* Checks if the current callstack is inside of a transactional PostMigrate step of migrating objects to this server
+	*/
+	COREUOBJECT_API bool IsTransactionallyPostMigratingObjects();
 }
diff --git a/Engine/Source/Runtime/CoreUObject/Public/UObject/RemoteObjectTransfer.h b/Engine/Source/Runtime/CoreUObject/Public/UObject/RemoteObjectTransfer.h
--- a/Engine/Source/Runtime/CoreUObject/Public/UObject/RemoteObjectTransfer.h
+++ b/Engine/Source/Runtime/CoreUObject/Public/UObject/RemoteObjectTransfer.h
@@ -219,6 +219,12 @@ namespace UE::RemoteObject::Transfer
 	*/
 	COREUOBJECT_API void MigrateObjectFromRemoteServer(FRemoteObjectId ObjectId, FRemoteServerId CurrentOwnerServerId);
 
+	/**
+	* Deserializes objects (possibly creating them in the process) using data from a remote server and a migration context inside of a transaction
+	* @param MigratedData Data received from a remote server coupled with the migration context struct
+	*/
+	void TransactionallyMigrateObjects(TRefCountPtr<FTransactionalMigrationData> MigratedData);
+
 	/**
 	* Reports code that touches a resident object
 	*/
@@ -259,6 +265,11 @@ namespace UE::RemoteObject::Transfer
 
 	COREUOBJECT_API void InitRemoteObjectTransfer();
 
+	/**
+	* Returns object data that has not yet been deserialized on this server to their respective owners
+	*/
+	COREUOBJECT_API int32 ReturnObjectDataToOwnedServers();
+
 	/** Delegate that transfers object data to another server */
 	extern COREUOBJECT_API TDelegate<void(const FMigrateSendParams& /*Params*/)> RemoteObjectTransferDelegate;
 
@@ -276,7 +287,7 @@ namespace UE::RemoteObject::Transfer
 	extern COREUOBJECT_API TMulticastDelegate<void(const TArray<UObject*>& /*ReceivedObjects*/, const FRemoteObjectData& /*ObjectData*/, const FUObjectMigrationContext& /*MigrationContext*/)> OnObjectsReceivedDelegate;
 
 	/** Delegate executed when object data has been migrated to another server */
-	extern COREUOBJECT_API TMulticastDelegate<void(const TSet<UObject*>& /*SentObjects*/,  const FRemoteObjectData& /*ObjectData*/, const FUObjectMigrationContext& /*MigrationContext*/)> OnObjectsSentDelegate;
+	extern COREUOBJECT_API TMulticastDelegate<void(const TSet<TObjectPtr<UObject>>& /*SentObjects*/,  const FRemoteObjectData& /*ObjectData*/, const FUObjectMigrationContext& /*MigrationContext*/)> OnObjectsSentDelegate;
 
 	/** Delegate executed when object data has been loaded from disk */
 	extern COREUOBJECT_API TMulticastDelegate<void(const TArray<UObject*>& /*LoadedObjects*/, const FRemoteObjectData& /*ObjectData*/, const FUObjectMigrationContext& /*MigrationContext*/)> OnObjectsLoadedFromDiskDelegate;
@@ -344,9 +355,16 @@ public:
 	{
 		return ServerId;
 	}
+
+	UE_FORCEINLINE_HINT EResidence GetResidence() const
+	{
+		return UE::RemoteObject::Handle::GetResidence(ObjectId);
+	}
+
+	UE_DEPRECATED(5.8, "Use GetResidence(ObjectId) instead.")
 	UE_FORCEINLINE_HINT bool IsRemote() const
 	{
-		return UE::RemoteObject::Handle::IsRemote(ObjectId);
+		return GetResidence() == EResidence::Remote;
 	}
 
 	COREUOBJECT_API bool Serialize(FArchive& Ar);
diff --git a/Engine/Source/Runtime/CoreUObject/Public/UObject/WeakObjectPtr.h b/Engine/Source/Runtime/CoreUObject/Public/UObject/WeakObjectPtr.h
--- a/Engine/Source/Runtime/CoreUObject/Public/UObject/WeakObjectPtr.h
+++ b/Engine/Source/Runtime/CoreUObject/Public/UObject/WeakObjectPtr.h
@@ -332,11 +332,26 @@ public:
 	/**
 	* Returns true if the object this weak pointer represents is owned by another server
 	*/
+	UE_FORCEINLINE_HINT EResidence GetResidence() const
+	{
+		return UE::RemoteObject::Handle::GetResidence(ObjectRemoteId);
+	}
+
+	/**
+	* Returns true if the object this weak pointer represents is owned by another server
+	*/
+	UE_DEPRECATED(5.8, "Use GetResidence() instead.")
 	UE_FORCEINLINE_HINT bool IsRemote() const
 	{
-		return UE::RemoteObject::Handle::IsRemote(ObjectRemoteId);
+		return GetResidence() == EResidence::Remote;
 	}
 #else
+	UE_FORCEINLINE_HINT EResidence GetResidence() const
+	{
+		return EResidence::Local;
+	}
+
+	UE_DEPRECATED(5.8, "Use GetResidence() instead.")
 	UE_FORCEINLINE_HINT bool IsRemote() const
 	{
 		return false;
diff --git a/Engine/Source/Runtime/Engine/Private/ActorReplication.cpp b/Engine/Source/Runtime/Engine/Private/ActorReplication.cpp
--- a/Engine/Source/Runtime/Engine/Private/ActorReplication.cpp
+++ b/Engine/Source/Runtime/Engine/Private/ActorReplication.cpp
@@ -1110,7 +1110,7 @@ void AActor::UpdateOwningNetConnection()
 		{
 			for (TObjectPtr<AActor>& ActorChildPtr : Actor->Children)
 			{
-				if (!ActorChildPtr.IsRemote())
+				if (ActorChildPtr.GetResidence() != EResidence::Remote)
 				{
 					AActor* ActorChild = ActorChildPtr;
 					EveryChildren.Push(ActorChild);
diff --git a/Engine/Source/Runtime/Engine/Private/Components/ActorComponent.cpp b/Engine/Source/Runtime/Engine/Private/Components/ActorComponent.cpp
--- a/Engine/Source/Runtime/Engine/Private/Components/ActorComponent.cpp
+++ b/Engine/Source/Runtime/Engine/Private/Components/ActorComponent.cpp
@@ -716,6 +716,8 @@ void UActorComponent::PostLoad()
 #if UE_WITH_REMOTE_OBJECT_HANDLE
 void UActorComponent::PostMigrate(const struct FUObjectMigrationContext& MigrationContext)
 {
+	Super::PostMigrate(MigrationContext);
+
 	// We want to be able to store "annotations" (metadata) about this Component while we loan it out, so that
 	// when we have the loan returned to us, we can check what has changed...
 	struct FLoanedState
@@ -1725,7 +1727,7 @@ void FActorComponentTickFunction::ExecuteTick(float DeltaTime, enum ELevelTick T
 FString FActorComponentTickFunction::DiagnosticMessage()
 {
 #if	UE_WITH_REMOTE_OBJECT_HANDLE
-	return (!Target.IsRemote() ? Target->GetFullName() : Target.GetRemoteId().ToString()) + TEXT("[TickComponent]");
+	return (Target.GetResidence() != EResidence::Remote ? Target->GetFullName() : Target.GetRemoteId().ToString()) + TEXT("[TickComponent]");
 #else
 	return Target->GetFullName() + TEXT("[TickComponent]");
 #endif
diff --git a/Engine/Source/Runtime/Engine/Private/Components/SceneComponent.cpp b/Engine/Source/Runtime/Engine/Private/Components/SceneComponent.cpp
--- a/Engine/Source/Runtime/Engine/Private/Components/SceneComponent.cpp
+++ b/Engine/Source/Runtime/Engine/Private/Components/SceneComponent.cpp
@@ -3960,7 +3960,7 @@ void USceneComponent::Serialize(FArchive& Ar)
 		{
 			TObjectPtr<USceneComponent>& Child = AttachChildren[Index];
 
-			if (ensureMsgf(!Child.IsRemote(), TEXT("About to save a Remote Attached Child in %s"), *GetFullNameSafe(this)))
+			if (ensureMsgf(Child.GetResidence() != EResidence::Remote, TEXT("About to save a Remote Attached Child in %s"), *GetFullNameSafe(this)))
 			{
 				if (Child->GetOwner() != GetOwner())
 				{
@@ -4007,7 +4007,7 @@ void USceneComponent::Serialize(FArchive& Ar)
 		{
 			for (TObjectPtr<USceneComponent>& Child : AttachChildren)
 			{
-				ensureMsgf(!Child.IsRemote(), TEXT("Loaded a Remote Attached Child in %s"), *GetFullNameSafe(this));
+				ensureMsgf(Child.GetResidence() != EResidence::Remote, TEXT("Loaded a Remote Attached Child in %s"), *GetFullNameSafe(this));
 			}
 		}
 	}
diff --git a/Engine/Source/Runtime/Engine/Private/Components/ShapeComponent.cpp b/Engine/Source/Runtime/Engine/Private/Components/ShapeComponent.cpp
--- a/Engine/Source/Runtime/Engine/Private/Components/ShapeComponent.cpp
+++ b/Engine/Source/Runtime/Engine/Private/Components/ShapeComponent.cpp
@@ -114,7 +114,7 @@ void UShapeComponent::Serialize(FArchive& Ar)
 		// So null it out here so we don't end up with that runtime CDO SubObject from another server.
 		// If we migrated, it would show up as a RemoteObject and crash since its Outer is the CDO.
 		// Note: We should ban CDO mutation on runtime objects. This is never safe due to delta serialization.
-		if (ShapeBodySetup.IsRemote())
+		if (ShapeBodySetup.GetResidence() == EResidence::Remote)
 		{
 			ShapeBodySetup = nullptr;
 		}
diff --git a/Engine/Source/Runtime/Engine/Private/LevelTick.cpp b/Engine/Source/Runtime/Engine/Private/LevelTick.cpp
--- a/Engine/Source/Runtime/Engine/Private/LevelTick.cpp
+++ b/Engine/Source/Runtime/Engine/Private/LevelTick.cpp
@@ -574,7 +574,7 @@ void UWorld::ProcessLevelStreamingVolumes(FVector* OverrideViewLocation)
 	bool bStreamingVolumesAreRelevant = false;
 	for (FConstPlayerControllerIterator Iterator = GetPlayerControllerIterator(); Iterator; ++Iterator)
 	{
-		if (!Iterator->IsRemote())
+		if (Iterator->GetResidence() != EResidence::Remote)
 		{
 			APlayerController* PlayerActor = Iterator->Get();
 			if (PlayerActor && PlayerActor->bIsUsingStreamingVolumes)
@@ -619,7 +619,7 @@ void UWorld::ProcessLevelStreamingVolumes(FVector* OverrideViewLocation)
 		// volumes that contain player viewpoints.
 		for( FConstPlayerControllerIterator Iterator = GetPlayerControllerIterator(); Iterator; ++Iterator )
 		{
-			if (Iterator->IsRemote())
+			if (Iterator->GetResidence() == EResidence::Remote)
 			{
 				continue;
 			}
@@ -1806,7 +1806,7 @@ void UWorld::Tick( ELevelTick TickType, float DeltaSeconds )
 				// Update cameras last. This needs to be done before NetUpdates, and after all actors have been ticked.
 				for( FConstPlayerControllerIterator Iterator = GetPlayerControllerIterator(); Iterator; ++Iterator )
 				{
-					if (Iterator->IsRemote())
+					if (Iterator->GetResidence() == EResidence::Remote)
 					{
 						continue;
 					}
diff --git a/Engine/Source/Runtime/Engine/Private/NetConnection.cpp b/Engine/Source/Runtime/Engine/Private/NetConnection.cpp
--- a/Engine/Source/Runtime/Engine/Private/NetConnection.cpp
+++ b/Engine/Source/Runtime/Engine/Private/NetConnection.cpp
@@ -3088,8 +3088,8 @@ bool UNetConnection::ReadPacketInfo(FBitReader& Reader, bool bHasPacketInfoPaylo
 		if (PlayerController)
 		{
 			// PlayerController is valid and local (as opposed to currently loaned out)
-			UE_CLOG(PlayerController.IsRemote(), LogNet, Warning, TEXT("%hs: PlayerController is unexpectedly a RemoteObject"), __func__);
-			if (!PlayerController.IsRemote())
+			UE_CLOG(PlayerController.GetResidence() == EResidence::Remote, LogNet, Warning, TEXT("%hs: PlayerController is unexpectedly a RemoteObject"), __func__);
+			if (PlayerController.GetResidence() != EResidence::Remote)
 			{
 				double PlayerFacingPing = RTTExclFrame;
 				switch(EPingType(CVarPingPlayerFacingPingType.GetValueOnAnyThread()))
diff --git a/Engine/Source/Runtime/Engine/Private/NetDriver.cpp b/Engine/Source/Runtime/Engine/Private/NetDriver.cpp
--- a/Engine/Source/Runtime/Engine/Private/NetDriver.cpp
+++ b/Engine/Source/Runtime/Engine/Private/NetDriver.cpp
@@ -8977,12 +8977,12 @@ void UNetDriver::UpdateNetworkStats()
 							if (Connection->PlayerController && Connection->PlayerController->PlayerState)
 							{
 #if UE_WITH_REMOTE_OBJECT_HANDLE
-								if (Connection->PlayerController.IsRemote())
+								if (Connection->PlayerController.GetResidence() == EResidence::Remote)
 								{
 									ensureMsgf(false, TEXT("Expecting the Connection's PlayerController to be locally owned by now"));
 									continue;
 								}
-								else if (Connection->PlayerController->PlayerState.IsRemote())
+								else if (Connection->PlayerController->PlayerState.GetResidence() == EResidence::Remote)
 								{
 									continue;
 								}
diff --git a/Engine/Source/Runtime/Engine/Private/PackageMapClient.cpp b/Engine/Source/Runtime/Engine/Private/PackageMapClient.cpp
--- a/Engine/Source/Runtime/Engine/Private/PackageMapClient.cpp
+++ b/Engine/Source/Runtime/Engine/Private/PackageMapClient.cpp
@@ -1127,7 +1127,7 @@ FNetworkGUID UPackageMapClient::InternalLoadObject( FArchive & Ar, UObject *& Ob
 		if (Object == nullptr)
 		{
 			FRemoteObjectId RemoteObjectId = FNetGUIDCache::GetRemoteObjectIdFromNetworkGUID(NetGUID);
-			if (UE::RemoteObject::Handle::IsRemote(RemoteObjectId))
+			if (UE::RemoteObject::Handle::GetResidence(RemoteObjectId) != EResidence::Local)
 			{
 				if (AutoRTFM::IsTransactional())
 				{
@@ -3759,7 +3759,7 @@ UObject* FNetGUIDCache::GetObjectFromNetGUID( const FNetworkGUID& NetGUID, const
 
 	if (CacheObjectPtr->Object.IsValid()
 #if UE_WITH_REMOTE_OBJECT_HANDLE
-		&& !UE::RemoteObject::Handle::IsRemote(FWeakObjectPtr(CacheObjectPtr->Object).GetRemoteId())
+		&& (UE::RemoteObject::Handle::GetResidence(FWeakObjectPtr(CacheObjectPtr->Object).GetRemoteId()) != EResidence::Remote)
 #endif
 		)
 	{
diff --git a/Engine/Source/Runtime/Engine/Private/PhysicsEngine/Experimental/ChaosScopedSceneLock.cpp b/Engine/Source/Runtime/Engine/Private/PhysicsEngine/Experimental/ChaosScopedSceneLock.cpp
--- a/Engine/Source/Runtime/Engine/Private/PhysicsEngine/Experimental/ChaosScopedSceneLock.cpp
+++ b/Engine/Source/Runtime/Engine/Private/PhysicsEngine/Experimental/ChaosScopedSceneLock.cpp
@@ -186,7 +186,7 @@ void FScopedSceneLock_Chaos::LockScene()
 	case EPhysicsInterfaceScopedLockType::Write:
 		{
 #if UE_WITH_REMOTE_OBJECT_HANDLE
-			if (TransactionMode == EPhysicsInterfaceScopedTransactionMode::MultiServer)
+			if (TransactionMode == EPhysicsInterfaceScopedTransactionMode::MultiServer && !Chaos::Private::FDstmScopedWriteLockInvalidationIgnore::WriteLocksInvalidationAreIgnored())
 			{
 				UE::RemoteExecutor::TransactionRequiresMultiServerCommit(TEXT("Physics Write Lock"));
 			}
@@ -307,3 +307,27 @@ FScopedSceneLockWithContext_Chaos::FScopedSceneLockWithContext_Chaos(FPhysicsCon
 #endif
 	LockSceneForConstraint(InConstraintHandle);
 }
+
+#if UE_WITH_REMOTE_OBJECT_HANDLE
+namespace Chaos::Private
+{
+	thread_local bool bIgnoreWriteLockInvalidation = false;
+
+	FDstmScopedWriteLockInvalidationIgnore::FDstmScopedWriteLockInvalidationIgnore(bool bIgnore)
+	{
+		bOldState = bIgnoreWriteLockInvalidation;
+		bIgnoreWriteLockInvalidation = bIgnore;
+	}
+
+	FDstmScopedWriteLockInvalidationIgnore::~FDstmScopedWriteLockInvalidationIgnore()
+	{
+		bIgnoreWriteLockInvalidation = bOldState;
+	}
+
+	bool FDstmScopedWriteLockInvalidationIgnore::WriteLocksInvalidationAreIgnored()
+	{
+		return bIgnoreWriteLockInvalidation || UE::RemoteObject::Handle::IsTransactionallyPostMigratingObjects();
+	}
+} // namespace Chaos::Private
+
+#endif // UE_WITH_REMOTE_OBJECT_HANDLE
diff --git a/Engine/Source/Runtime/Engine/Private/PhysicsEngine/PhysicsConstraintComponent.cpp b/Engine/Source/Runtime/Engine/Private/PhysicsEngine/PhysicsConstraintComponent.cpp
--- a/Engine/Source/Runtime/Engine/Private/PhysicsEngine/PhysicsConstraintComponent.cpp
+++ b/Engine/Source/Runtime/Engine/Private/PhysicsEngine/PhysicsConstraintComponent.cpp
@@ -61,7 +61,7 @@ UPrimitiveComponent* UPhysicsConstraintComponent::GetComponentInternal(EConstrai
 		if(OverrideComponent1.IsValid())
 		{
 #if UE_WITH_REMOTE_OBJECT_HANDLE
-			if (OverrideComponent1.IsRemote() && !AutoRTFM::IsClosed())
+			if (OverrideComponent1.GetResidence() == EResidence::Remote && !AutoRTFM::IsClosed())
 			{
 				ensureMsgf(false, TEXT("UPhysicsConstraintComponent::GetComponentInternal: Frame1 Component is remote. Simulation result is not guaranteed."));
 				return NULL;
@@ -72,7 +72,7 @@ UPrimitiveComponent* UPhysicsConstraintComponent::GetComponentInternal(EConstrai
 
 		ComponentName = ComponentName1.ComponentName;
 #if UE_WITH_REMOTE_OBJECT_HANDLE
-		if (ConstraintActor1.IsRemote() && !AutoRTFM::IsClosed())
+		if (ConstraintActor1.GetResidence() == EResidence::Remote && !AutoRTFM::IsClosed())
 		{
 			ensureMsgf(false, TEXT("UPhysicsConstraintComponent::GetComponentInternal: Frame1 Actor is remote. Simulation result is not guaranteed."));
 			return NULL;
@@ -87,7 +87,7 @@ UPrimitiveComponent* UPhysicsConstraintComponent::GetComponentInternal(EConstrai
 		if(OverrideComponent2.IsValid())
 		{
 #if UE_WITH_REMOTE_OBJECT_HANDLE
-			if (OverrideComponent2.IsRemote() && !AutoRTFM::IsClosed())
+			if (OverrideComponent2.GetResidence() == EResidence::Remote && !AutoRTFM::IsClosed())
 			{
 				ensureMsgf(false, TEXT("UPhysicsConstraintComponent::GetComponentInternal: Frame2 Component is remote. Simulation result is not guaranteed."));
 				return NULL;
@@ -98,7 +98,7 @@ UPrimitiveComponent* UPhysicsConstraintComponent::GetComponentInternal(EConstrai
 
 		ComponentName = ComponentName2.ComponentName;
 #if UE_WITH_REMOTE_OBJECT_HANDLE
-		if (ConstraintActor2.IsRemote() && !AutoRTFM::IsClosed())
+		if (ConstraintActor2.GetResidence() == EResidence::Remote && !AutoRTFM::IsClosed())
 		{
 			ensureMsgf(false, TEXT("UPhysicsConstraintComponent::GetComponentInternal: Frame2 Actor is remote. Simulation result is not guaranteed."));
 			return NULL;
@@ -850,7 +850,7 @@ void UPhysicsConstraintComponent::MigratePhysicsState(FArchive& Ar)
 
 bool UPhysicsConstraintComponent::AreConstrainedComponentsInSameActor() const
 {
-	if (OverrideComponent1.IsRemote() || OverrideComponent2.IsRemote())
+	if (OverrideComponent1.GetResidence() == EResidence::Remote || OverrideComponent2.GetResidence() == EResidence::Remote)
 	{
 		return false;
 	}
@@ -873,7 +873,7 @@ bool UPhysicsConstraintComponent::AreConstrainedComponentsInSameActor() const
 
 	if (ConstraintActor1 != NULL)
 	{
-		if (ConstraintActor1.IsRemote())
+		if (ConstraintActor1.GetResidence() == EResidence::Remote)
 		{
 			return false;
 		}
@@ -885,7 +885,7 @@ bool UPhysicsConstraintComponent::AreConstrainedComponentsInSameActor() const
 
 	if (ConstraintActor2 != NULL)
 	{
-		if (ConstraintActor2.IsRemote())
+		if (ConstraintActor2.GetResidence() == EResidence::Remote)
 		{
 			return false;
 		}
diff --git a/Engine/Source/Runtime/Engine/Private/PhysicsEngine/PhysicsReplication.cpp b/Engine/Source/Runtime/Engine/Private/PhysicsEngine/PhysicsReplication.cpp
--- a/Engine/Source/Runtime/Engine/Private/PhysicsEngine/PhysicsReplication.cpp
+++ b/Engine/Source/Runtime/Engine/Private/PhysicsEngine/PhysicsReplication.cpp
@@ -2739,7 +2739,7 @@ float FPhysicsReplication::GetLocalPing() const
 		{
 			if (const TObjectPtr<APlayerState>& PlayerState = PlayerController->PlayerState)
 			{
-				if (!PlayerState.IsRemote())
+				if (PlayerState.GetResidence() != EResidence::Remote)
 				{
 					return PlayerState->ExactPing;
 				}
diff --git a/Engine/Source/Runtime/Engine/Private/PlayerController.cpp b/Engine/Source/Runtime/Engine/Private/PlayerController.cpp
--- a/Engine/Source/Runtime/Engine/Private/PlayerController.cpp
+++ b/Engine/Source/Runtime/Engine/Private/PlayerController.cpp
@@ -5039,7 +5039,7 @@ static void UpdateNetConnectionControllerPosition(UNetConnection* NetConnection)
 	bool bNoPawnPlayerControllerPositionSet = false;
 	if (ControllerActor)
 	{
-		if (NetConnection->ViewTarget.IsRemote())
+		if (NetConnection->ViewTarget.GetResidence() == EResidence::Remote)
 		{
 			UE_LOG(LogPlayerController, Warning, TEXT("UpdateNetConnectionControllerPosition: Incoming ViewTarget is remote."));
 		}
diff --git a/Engine/Source/Runtime/Engine/Private/TimerManager.cpp b/Engine/Source/Runtime/Engine/Private/TimerManager.cpp
--- a/Engine/Source/Runtime/Engine/Private/TimerManager.cpp
+++ b/Engine/Source/Runtime/Engine/Private/TimerManager.cpp
@@ -255,6 +255,8 @@ namespace
 	*/
 	bool IsTimerTargetResident(const FTimerUnifiedDelegate& TimerDelegate)
 	{
+		using namespace UE::RemoteObject::Handle;
+
 		FRemoteObjectId RemoteObjectId = UE::RemoteObject::Handle::GetRemoteObjectId([&TimerDelegate]()
 		{
 			const UObject* TargetObject;
@@ -278,7 +280,7 @@ namespace
 			return TargetObject;
 		});
 
-		return !UE::RemoteObject::Handle::IsRemote(RemoteObjectId);
+		return GetResidence(RemoteObjectId) == EResidence::Local;
 	}
 	#endif
 }
diff --git a/Engine/Source/Runtime/Engine/Public/Physics/Experimental/ChaosScopedSceneLock.h b/Engine/Source/Runtime/Engine/Public/Physics/Experimental/ChaosScopedSceneLock.h
--- a/Engine/Source/Runtime/Engine/Public/Physics/Experimental/ChaosScopedSceneLock.h
+++ b/Engine/Source/Runtime/Engine/Public/Physics/Experimental/ChaosScopedSceneLock.h
@@ -86,3 +86,19 @@ struct UE_INTERNAL FScopedSceneLockWithContext_Chaos : FScopedSceneLock_Chaos
 	ENGINE_API FScopedSceneLockWithContext_Chaos(const FPhysicsActorHandle& InActorHandle, EPhysicsInterfaceScopedLockType InLockType, EPhysicsInterfaceScopedThreadContext InThreadContext, EPhysicsInterfaceScopedTransactionMode InTransactionMode);
 	ENGINE_API FScopedSceneLockWithContext_Chaos(FPhysicsConstraintHandle const* InConstraintHandle, EPhysicsInterfaceScopedLockType InLockType, EPhysicsInterfaceScopedThreadContext InThreadContext, EPhysicsInterfaceScopedTransactionMode InTransactionMode);
 };
+
+#if UE_WITH_REMOTE_OBJECT_HANDLE
+namespace Chaos::Private
+{
+	struct UE_INTERNAL FDstmScopedWriteLockInvalidationIgnore
+	{
+		ENGINE_API FDstmScopedWriteLockInvalidationIgnore(bool bInvalidate);
+		ENGINE_API ~FDstmScopedWriteLockInvalidationIgnore();
+
+		ENGINE_API static bool WriteLocksInvalidationAreIgnored();
+
+	private:
+		bool bOldState;
+	};
+} // namespace Chaos::Private
+#endif //  UE_WITH_REMOTE_OBJECT_HANDLE
diff --git a/Engine/Source/Runtime/GameplayDebugger/Private/GameplayDebuggerCategoryReplicator.cpp b/Engine/Source/Runtime/GameplayDebugger/Private/GameplayDebuggerCategoryReplicator.cpp
--- a/Engine/Source/Runtime/GameplayDebugger/Private/GameplayDebuggerCategoryReplicator.cpp
+++ b/Engine/Source/Runtime/GameplayDebugger/Private/GameplayDebuggerCategoryReplicator.cpp
@@ -715,7 +715,7 @@ void AGameplayDebuggerCategoryReplicator::OnExtensionsChanged()
 
 UNetConnection* AGameplayDebuggerCategoryReplicator::GetNetConnection() const
 {
-	return (!OwnerPC.IsRemote() && IsValid(OwnerPC)) ? OwnerPC->GetNetConnection() : nullptr;
+	return (OwnerPC.GetResidence() != EResidence::Remote && IsValid(OwnerPC)) ? OwnerPC->GetNetConnection() : nullptr;
 }
 
 bool AGameplayDebuggerCategoryReplicator::IsNetRelevantFor(const AActor* RealViewer, const AActor* ViewTarget, const FVector& SrcLocation) const
diff --git a/Engine/Source/Runtime/GameplayDebugger/Private/GameplayDebuggerPlayerManager.cpp b/Engine/Source/Runtime/GameplayDebugger/Private/GameplayDebuggerPlayerManager.cpp
--- a/Engine/Source/Runtime/GameplayDebugger/Private/GameplayDebuggerPlayerManager.cpp
+++ b/Engine/Source/Runtime/GameplayDebugger/Private/GameplayDebuggerPlayerManager.cpp
@@ -159,7 +159,7 @@ void AGameplayDebuggerPlayerManager::UpdateAuthReplicators()
 	{
 		FGameplayDebuggerPlayerData& TestData = PlayerData[Idx];
 
-		if (IsValid(TestData.Replicator) && TestData.Replicator->GetReplicationOwner().IsRemote())
+		if (IsValid(TestData.Replicator) && TestData.Replicator->GetReplicationOwner().GetResidence() == EResidence::Remote)
 		{
 			continue;
 		}
diff --git a/Engine/Source/Runtime/Net/Core/Private/Net/Core/PropertyConditions/PropertyConditions.cpp b/Engine/Source/Runtime/Net/Core/Private/Net/Core/PropertyConditions/PropertyConditions.cpp
--- a/Engine/Source/Runtime/Net/Core/Private/Net/Core/PropertyConditions/PropertyConditions.cpp
+++ b/Engine/Source/Runtime/Net/Core/Private/Net/Core/PropertyConditions/PropertyConditions.cpp
@@ -124,7 +124,7 @@ void FNetPropertyConditionManager::PostGarbageCollect()
 	{
 		if (
 #if UE_WITH_REMOTE_OBJECT_HANDLE
-			FObjectPtr(It.Key().GetRemoteId()).IsRemote() ||
+			FObjectPtr(It.Key().GetRemoteId()).GetResidence() == EResidence::Remote ||
 #endif
 			!It.Key().ResolveObjectPtr())
 		{
