Fixing deprecation warnings in a few projects

#Synced-CL 49877584

diff --git a/Engine/Source/Runtime/CoreUObject/Private/UObject/RemoteObject.cpp b/Engine/Source/Runtime/CoreUObject/Private/UObject/RemoteObject.cpp
--- a/Engine/Source/Runtime/CoreUObject/Private/UObject/RemoteObject.cpp
+++ b/Engine/Source/Runtime/CoreUObject/Private/UObject/RemoteObject.cpp
@@ -604,17 +604,29 @@ EResidence GetResidence(const UObject* Object)
 
 bool IsRemote(FRemoteObjectId ObjectId)
 {
+#if UE_WITH_REMOTE_OBJECT_HANDLE
 	return GetResidence(ObjectId) == EResidence::Remote;
+#else
+	return false;
+#endif
 }
 
 bool IsRemote(const FRemoteObjectStub* Stub)
 {
+#if UE_WITH_REMOTE_OBJECT_HANDLE
 	return GetResidence(Stub) == EResidence::Remote;
+#else
+	return false;
+#endif
 }
 
 bool IsRemote(const UObject* Object)
 {
+#if UE_WITH_REMOTE_OBJECT_HANDLE
 	return GetResidence(Object) == EResidence::Remote;
+#else
+	return false;
+#endif
 }
 
 bool IsOwned(const UObject* Object)
