Fixing deprecation warnings

#Synced-CL 49877584

diff --git a/Engine/Source/Runtime/CoreUObject/Public/UObject/ObjectHandle.h b/Engine/Source/Runtime/CoreUObject/Public/UObject/ObjectHandle.h
--- a/Engine/Source/Runtime/CoreUObject/Public/UObject/ObjectHandle.h
+++ b/Engine/Source/Runtime/CoreUObject/Public/UObject/ObjectHandle.h
@@ -370,6 +370,7 @@ UE_DEPRECATED(5.8, "Use GetObjectHandleResidence() instead.")
 inline bool IsObjectHandleRemote(FObjectHandle Handle)
 {
 #if UE_WITH_REMOTE_OBJECT_HANDLE
+	PRAGMA_DISABLE_DEPRECATION_WARNINGS
 	if (Handle.ExtractHandleType() == UE::CoreUObject::Private::ERemoteObjectHandleType::ObjectPointer)
 	{
 		if (!IsObjectHandleNull(Handle))
@@ -385,6 +386,7 @@ inline bool IsObjectHandleRemote(FObjectHandle Handle)
 	{
 		return UE::RemoteObject::Handle::IsRemote(Handle.ExtractRemoteId());
 	}
+	PRAGMA_ENABLE_DEPRECATION_WARNINGS
 #endif
 	return false;
 }
