some build fixes for UE_WITH_REMOTE_OBJECT_HANDLE

#Synced-CL 49893950

diff --git a/Engine/Source/Editor/KismetCompiler/Private/KismetCompilerVMBackend.cpp b/Engine/Source/Editor/KismetCompiler/Private/KismetCompilerVMBackend.cpp
--- a/Engine/Source/Editor/KismetCompiler/Private/KismetCompilerVMBackend.cpp
+++ b/Engine/Source/Editor/KismetCompiler/Private/KismetCompilerVMBackend.cpp
@@ -98,7 +98,7 @@ public:
 #if UE_WITH_OBJECT_HANDLE_LATE_RESOLVE
 		ScriptPointerType D = (ScriptPointerType)Res.GetHandle().PointerOrRef;
 #elif UE_WITH_REMOTE_OBJECT_HANDLE
-		ScriptPointerType D = (ScriptPointerType)Res.GetHandle().PointerOrHandle;
+		ScriptPointerType D = (ScriptPointerType)Res.GetHandle().Payload;
 #else
 		ScriptPointerType D = (ScriptPointerType)Res.GetHandle();
 #endif
diff --git a/Engine/Source/Runtime/AutoRTFM/Public/AutoRTFM.h b/Engine/Source/Runtime/AutoRTFM/Public/AutoRTFM.h
--- a/Engine/Source/Runtime/AutoRTFM/Public/AutoRTFM.h
+++ b/Engine/Source/Runtime/AutoRTFM/Public/AutoRTFM.h
@@ -815,7 +815,7 @@ struct FHeapRedirectCallbacks
 };
 
 #if UE_AUTORTFM
-void RegisterHeapRedirectCallbacks(FHeapRedirectCallbacks Callbacks);
+UE_AUTORTFM_API void RegisterHeapRedirectCallbacks(FHeapRedirectCallbacks Callbacks);
 #else
 UE_AUTORTFM_CRITICAL_INLINE void RegisterHeapRedirectCallbacks(FHeapRedirectCallbacks /*Callbacks*/) {}
 #endif
diff --git a/Engine/Source/Runtime/CoreUObject/Public/UObject/Object.h b/Engine/Source/Runtime/CoreUObject/Public/UObject/Object.h
--- a/Engine/Source/Runtime/CoreUObject/Public/UObject/Object.h
+++ b/Engine/Source/Runtime/CoreUObject/Public/UObject/Object.h
@@ -557,7 +557,7 @@ public:
 	/**
 	 * Called when an object is about to be migrated and after it's been migrated to another process.
 	 */
-	virtual void PostMigrate(const struct FUObjectMigrationContext& MigrationContext);
+	COREUOBJECT_API virtual void PostMigrate(const struct FUObjectMigrationContext& MigrationContext);
 
 	virtual bool IsMigrationRoot() const { return false; }
 #endif
diff --git a/Engine/Source/Runtime/CoreUObject/Public/UObject/ObjectHandle.h b/Engine/Source/Runtime/CoreUObject/Public/UObject/ObjectHandle.h
--- a/Engine/Source/Runtime/CoreUObject/Public/UObject/ObjectHandle.h
+++ b/Engine/Source/Runtime/CoreUObject/Public/UObject/ObjectHandle.h
@@ -223,7 +223,18 @@ namespace UE::CoreUObject::Private
 
 		explicit constexpr FRemoteObjectHandlePrivate(const UObject* Object)
 		{
-			AssignFromObjectPointer(Object);
+			UE_IF_CONSTEVAL
+			{
+				// ensure that the handle type is zero so that we can just
+				// copy the pointer in and not need to worry about the
+				// bit representation which isn't knowable at compile time
+				static_assert((uint8)ERemoteObjectHandleType::ObjectPointer == 0);
+				Pointer = Object;
+			}
+			else
+			{
+				AssignFromObjectPointer(Object);
+			}
 		}
 
 		explicit FRemoteObjectHandlePrivate(const UE::RemoteObject::Handle::FRemoteObjectStub* RemoteInfo)
@@ -231,26 +242,35 @@ namespace UE::CoreUObject::Private
 			AssignFromStubPointer(RemoteInfo);
 		}
 
-		FRemoteObjectHandlePrivate& operator=(const FRemoteObjectHandlePrivate& Rhs)
+		constexpr FRemoteObjectHandlePrivate& operator=(const FRemoteObjectHandlePrivate& Rhs)
 		{
-			// special rules for assignment:
-			// if this handle is on the RemoteHeap then we are only
-			// allowed to store it as a RemoteId
-			
-			if (Rhs.ExtractHandleType() == ERemoteObjectHandleType::RemoteId)
-			{
-				// if Rhs is already encoded as a RemoteId, just copy the payload
-				Payload = Rhs.Payload;
-			}
-			else if (!RemoteObject::IsPointerOnRemoteHeap(this))
+			UE_IF_CONSTEVAL
 			{
-				// if this handle is not on the remote heap, just copy the payload
+				// the remote heap doesn't exist at compile time so we can always
+				// just copy the payload
 				Payload = Rhs.Payload;
 			}
 			else
 			{
-				// otherwise get the RemoteId from the Rhs and assign from that
-				AssignFromRemoteId(Rhs.GetRemoteId());
+				// special rules for assignment:
+				// if this handle is on the RemoteHeap then we are only
+				// allowed to store it as a RemoteId
+			
+				if (Rhs.ExtractHandleType() == ERemoteObjectHandleType::RemoteId)
+				{
+					// if Rhs is already encoded as a RemoteId, just copy the payload
+					Payload = Rhs.Payload;
+				}
+				else if (!RemoteObject::IsPointerOnRemoteHeap(this))
+				{
+					// if this handle is not on the remote heap, just copy the payload
+					Payload = Rhs.Payload;
+				}
+				else
+				{
+					// otherwise get the RemoteId from the Rhs and assign from that
+					AssignFromRemoteId(Rhs.GetRemoteId());
+				}
 			}
 
 			return *this;
diff --git a/Engine/Source/Runtime/Experimental/Chaos/Public/Framework/Threading.h b/Engine/Source/Runtime/Experimental/Chaos/Public/Framework/Threading.h
--- a/Engine/Source/Runtime/Experimental/Chaos/Public/Framework/Threading.h
+++ b/Engine/Source/Runtime/Experimental/Chaos/Public/Framework/Threading.h
@@ -1154,7 +1154,7 @@ FORCEINLINE void EnsureIsInGameThreadContext()
 		void (*WriteUnlock)(FPhysSceneLockRemoteObject*) = nullptr;
 	};
 
-	extern FPhysSceneLockCallbacks GPhysSceneLockRemoteObjectCallbacks;
+	CHAOS_API extern FPhysSceneLockCallbacks GPhysSceneLockRemoteObjectCallbacks;
 
 	struct FPhysSceneLockRemoteObject
 	{
