Don't create an `FString` just to call `StartsWith` on it.

#Synced-CL 49919335

diff --git a/Engine/Source/Runtime/CoreUObject/Private/Serialization/TransactionallySafeAsyncLoading.cpp b/Engine/Source/Runtime/CoreUObject/Private/Serialization/TransactionallySafeAsyncLoading.cpp
--- a/Engine/Source/Runtime/CoreUObject/Private/Serialization/TransactionallySafeAsyncLoading.cpp
+++ b/Engine/Source/Runtime/CoreUObject/Private/Serialization/TransactionallySafeAsyncLoading.cpp
@@ -540,13 +540,22 @@ private:
 				if (LoadedPackage.IsValid())
 				{
 					// Skip trashed packages.
-					static FString TrashedPackage = FName(NAME_TrashedPackage).ToString();
-					if (LoadedPackage->GetFName().ToString().StartsWith(TrashedPackage))
+					const EName* const NameOrNull = LoadedPackage->GetFName().ToEName();
+
+					if (NameOrNull && (*NameOrNull == NAME_TrashedPackage))
 					{
 						continue;
 					}
 
-					return LoadedPackage;
+					// Out of an abundance of caution we'll leave the old trashing logic in
+					// but ensure that it isn't actually hit for the time being.
+					static FString TrashedPackage = FName(NAME_TrashedPackage).ToString();
+					const bool bIsTrashedPackage = LoadedPackage->GetFName().ToString().StartsWith(TrashedPackage);
+
+					if (ensure(!bIsTrashedPackage))
+					{
+						return LoadedPackage;
+					}
 				}
 			}
 
@@ -598,8 +607,8 @@ private:
 
 				if (FMapPayload* const Payload = PackageCache.Find(PackagePathStr))
 				{
-					// If we are not retrying it means we've done a package load outside of a transaction, so we want to just forget all previous loaded packages,
-					// and start from a blank array.
+					// If we are not retrying it means we've done a package load outside of a transaction,
+					// so we want to just forget all previous loaded packages, and start from a blank array.
 					if (!AutoRTFM::IsRetrying())
 					{
 						Payload->LoadedPackages.Empty();
