[ZenLoader]
- Add options to disable async loading and multithreaded loading in commandlet and cook commandlet
- Move option from AsyncPackageLoader.cpp into AsyncLoading2.cpp for order consistency

#Synced-CL 49920339

diff --git a/Engine/Source/Runtime/CoreUObject/Private/Serialization/AsyncLoading2.cpp b/Engine/Source/Runtime/CoreUObject/Private/Serialization/AsyncLoading2.cpp
--- a/Engine/Source/Runtime/CoreUObject/Private/Serialization/AsyncLoading2.cpp
+++ b/Engine/Source/Runtime/CoreUObject/Private/Serialization/AsyncLoading2.cpp
@@ -11044,6 +11044,47 @@ FAsyncLoadingThread2::FAsyncLoadingThread2(FIoDispatcher& InIoDispatcher)
 		GAllowMultithreadedLoading = bConfigValue;
 	}
 
+	// Support blanket removal of async loading and multithreaded loading for commandlet if desired.
+	if (IsRunningCommandlet())
+	{
+		if (GConfig->GetBool(SectionName, TEXT("s.NoCommandletAsyncLoading"), bConfigValue, GEngineIni) && bConfigValue)
+		{
+			FAsyncLoadingThreadSettings::Get().bAsyncLoadingThreadEnabled = false;
+		}
+
+		if (GConfig->GetBool(SectionName, TEXT("s.NoCommandletMultithreadedLoading"), bConfigValue, GEngineIni) && bConfigValue)
+		{
+			GAllowMultithreadedLoading = false;
+		}
+	}
+
+	// A bit more constrained on the cook only case if we want async loading in some commandlet that aren't cook related
+	if (IsRunningCookCommandlet())
+	{
+		if (GConfig->GetBool(SectionName, TEXT("s.NoCookCommandletAsyncLoading"), bConfigValue, GEngineIni) && bConfigValue)
+		{
+			FAsyncLoadingThreadSettings::Get().bAsyncLoadingThreadEnabled = false;
+		}
+
+		if (GConfig->GetBool(SectionName, TEXT("s.NoCookCommandletMultithreadedLoading"), bConfigValue, GEngineIni) && bConfigValue)
+		{
+			GAllowMultithreadedLoading = false;
+		}
+	}
+
+#if WITH_EDITOR
+	if (FAsyncLoadingThreadSettings::Get().bAsyncLoadingThreadEnabled)
+	{
+		const TCHAR* FlushStreamingOnGCStr = TEXT("gc.FlushStreamingOnGC");
+		IConsoleVariable* FlushStreamingOnGCVar = IConsoleManager::Get().FindConsoleVariable(FlushStreamingOnGCStr);
+		check(FlushStreamingOnGCVar);
+		const int CurrentFlushStreamingOnGCValue = FlushStreamingOnGCVar->GetInt();
+
+		UE_CLOG(CurrentFlushStreamingOnGCValue == 0, LogStreaming, Display, TEXT("Using AsyncLoadingThread in Editor; forcing %s=1"), FlushStreamingOnGCStr);
+		FlushStreamingOnGCVar->Set(1);
+	}
+#endif
+
 	UE_LOG(LogStreaming, Log, TEXT("AsyncLoading2 - Created: Event Driven Loader: %s, Async Loading Thread: %s, Async Post Load: %s"),
 		GEventDrivenLoaderEnabled ? TEXT("true") : TEXT("false"),
 		FAsyncLoadingThreadSettings::Get().bAsyncLoadingThreadEnabled ? TEXT("true") : TEXT("false"),
diff --git a/Engine/Source/Runtime/CoreUObject/Private/Serialization/AsyncPackageLoader.cpp b/Engine/Source/Runtime/CoreUObject/Private/Serialization/AsyncPackageLoader.cpp
--- a/Engine/Source/Runtime/CoreUObject/Private/Serialization/AsyncPackageLoader.cpp
+++ b/Engine/Source/Runtime/CoreUObject/Private/Serialization/AsyncPackageLoader.cpp
@@ -146,19 +146,6 @@ FAsyncLoadingThreadSettings::FAsyncLoadingThreadSettings()
 	bool bCommandLineEnable = FParse::Param(FCommandLine::Get(), TEXT("AsyncLoadingThread"));
 	bAsyncLoadingThreadEnabled = bCommandLineEnable || (bConfigValue && FApp::ShouldUseThreadingForPerformance() && !bCommandLineDisable);
 
-#if WITH_EDITOR
-	if (bAsyncLoadingThreadEnabled)
-	{
-		const TCHAR* FlushStreamingOnGCStr = TEXT("gc.FlushStreamingOnGC");
-		IConsoleVariable* FlushStreamingOnGCVar = IConsoleManager::Get().FindConsoleVariable(FlushStreamingOnGCStr);
-		check(FlushStreamingOnGCVar);
-		const int CurrentFlushStreamingOnGCValue = FlushStreamingOnGCVar->GetInt();
-
-		UE_CLOG(CurrentFlushStreamingOnGCValue == 0, LogStreaming, Display, TEXT("Using AsyncLoadingThread in Editor; forcing %s=1"), FlushStreamingOnGCStr);
-		FlushStreamingOnGCVar->Set(1);
-	}
-#endif
-
 	// if we are on an uncooked platform the default value for async loading is disabled while it is enabled for cooked platform
 	bConfigValue = FPlatformProperties::RequiresCookedData() ? true : false;
 
