PropertyHelper: Fix threading issue when comparing objects with the PPF_DeepComparison flag by adding thread_local to the RecursionCheck array in AreInstancedObjectsIdentical.

#rnx

#Synced-CL 49808502

diff --git a/Engine/Source/Runtime/CoreUObject/Private/UObject/PropertyHelper.cpp b/Engine/Source/Runtime/CoreUObject/Private/UObject/PropertyHelper.cpp
--- a/Engine/Source/Runtime/CoreUObject/Private/UObject/PropertyHelper.cpp
+++ b/Engine/Source/Runtime/CoreUObject/Private/UObject/PropertyHelper.cpp
@@ -46,7 +46,7 @@ bool AreInstancedObjectsIdentical( UObject* ObjectA, UObject* ObjectB, uint32 Po
 			return ObjectA == Other.ObjectA && ObjectB == Other.ObjectB && PortFlags == Other.PortFlags;
 		}
 	};
-	static TArray<FRecursionCheck> RecursionCheck;
+	static thread_local TArray<FRecursionCheck> RecursionCheck;
 	FRecursionCheck Test;
 	Test.ObjectA = ObjectA;
 	Test.ObjectB = ObjectB;
