Define a constant for the number of reserved bits in RemoteObjectId and use it where needed.

#Synced-CL 49887316

diff --git a/Engine/Source/Runtime/CoreUObject/Public/UObject/RemoteObjectTypes.h b/Engine/Source/Runtime/CoreUObject/Public/UObject/RemoteObjectTypes.h
--- a/Engine/Source/Runtime/CoreUObject/Public/UObject/RemoteObjectTypes.h
+++ b/Engine/Source/Runtime/CoreUObject/Public/UObject/RemoteObjectTypes.h
@@ -34,6 +34,7 @@ namespace UE::RemoteObject::Handle
 	};
 }
 
+constexpr inline int32 REMOTE_OBJECT_RESERVED_BIT_SIZE = 2;
 constexpr inline int32 REMOTE_OBJECT_SERVER_ID_BIT_SIZE = 10;
 constexpr inline int32 REMOTE_OBJECT_SERIAL_NUMBER_BIT_SIZE = 52;
 
@@ -226,13 +227,15 @@ private:
 	{
 		struct
 		{
-			uint64 Reserved : 2;  /* Reserved for static flag in NetGUID and FRemoteObjectHandle */
+			uint64 Reserved : REMOTE_OBJECT_RESERVED_BIT_SIZE;  /* Reserved for static flag in NetGUID and FRemoteObjectHandle */
 			uint64 SerialNumber : REMOTE_OBJECT_SERIAL_NUMBER_BIT_SIZE;
 			uint64 ServerId : REMOTE_OBJECT_SERVER_ID_BIT_SIZE;
 		};
 		uint64 Id = 0;
 	};
 
+	static_assert((REMOTE_OBJECT_RESERVED_BIT_SIZE + REMOTE_OBJECT_SERIAL_NUMBER_BIT_SIZE + REMOTE_OBJECT_SERVER_ID_BIT_SIZE) <= (sizeof(Id) * 8), "FRemoteObjectId component bit sizes must fit within the Id member");
+
 	inline int64 GetComparableValue() const
 	{
 		return int64(SerialNumber) | (int64(FRemoteServerId::GlobalizeId(ServerId)) << REMOTE_OBJECT_SERIAL_NUMBER_BIT_SIZE);
diff --git a/Engine/Source/Runtime/Engine/Private/NetDriver.cpp b/Engine/Source/Runtime/Engine/Private/NetDriver.cpp
--- a/Engine/Source/Runtime/Engine/Private/NetDriver.cpp
+++ b/Engine/Source/Runtime/Engine/Private/NetDriver.cpp
@@ -7646,7 +7646,7 @@ void UNetDriver::CreateReplicationSystem(bool bInitAsClient)
 
 				// Shift away the reserved bits in FRemoteObjectId so we're just left with the server id and serial number
 				// which will be embedded in the NetRefHandle serial.
-				const uint64 Serial = (RemoteObjectIdValue >> 2);
+				const uint64 Serial = (RemoteObjectIdValue >> REMOTE_OBJECT_RESERVED_BIT_SIZE);
 
 				return Serial;
 			});
@@ -7654,7 +7654,7 @@ void UNetDriver::CreateReplicationSystem(bool bInitAsClient)
 			Params.NetRefHandleSerialToObject.BindLambda([this](const uint64 Serial)
 			{
 				// Shift left to made room for the reserved bits on FRemoteObjectId.
-				const uint64 RemoteObjectIdValue = (Serial << 2);
+				const uint64 RemoteObjectIdValue = (Serial << REMOTE_OBJECT_RESERVED_BIT_SIZE);
 				const FRemoteObjectId RemoteObjectId = FRemoteObjectId::CreateFromInt(RemoteObjectIdValue).GetLocalized();
 
 				return TObjectPtr<UObject>(FObjectPtr(RemoteObjectId));
diff --git a/Engine/Source/Runtime/Net/Core/Public/Net/Core/NetHandle/NetHandle.h b/Engine/Source/Runtime/Net/Core/Public/Net/Core/NetHandle/NetHandle.h
--- a/Engine/Source/Runtime/Net/Core/Public/Net/Core/NetHandle/NetHandle.h
+++ b/Engine/Source/Runtime/Net/Core/Public/Net/Core/NetHandle/NetHandle.h
@@ -148,7 +148,7 @@ inline FNetHandle& FNetHandle::operator=(FNetHandle&& Other)
 inline uint64 FNetHandle::GetId() const
 {
 #if UE_WITH_REMOTE_OBJECT_HANDLE
-	return (Value.GetIdNumber() >> 2);
+	return (Value.GetIdNumber() >> REMOTE_OBJECT_RESERVED_BIT_SIZE);
 #else
 	return BitCast<FInternalValue>(Value).Id;
 #endif
