[UObject]

Deprecate REN_ForceNoResetLoaders macro as it was not properly deprecated in 5.5 (did not use the UE_DEPRECATED_MACRO macro)

[RN] The UObject::Rename flag REN_ForceNoResetLoaders has been deprecated as UObject::Rename will no longer call ResetLoaders or flush loading during renaming as this was the cause of unnecessary loading stalls, and could introduce hard to diagnose bugs due to the sync load in the middle of a possible temporarily, intentionally incorrect UObject state such as during Blueprint reinstancing. Rename will still attempt to maintain a package's Linker export mapping as renames occur unless the flag REN_AllowPackageLinkerMismatch is provided. As a result without specifying REN_AllowPackageLinkerMismatch, Rename will detach the linker from a loaded object whenever the object is moved to a new package. Similarly, renaming a loaded object without specifying REN_AllowPackageLinkerMismatch will detach the export object from the linker that loaded it allowing future package requests for the same package to load a new object instance using the original object name if no object redirector was created during renaming (i.e. REN_DontCreateRedirectors). If an object is being renamed and future package reloads should not create a new instance of original object name, InvalidateExport() should be called on the object before renaming. In cases where one wants to temporarily rename a loaded object and then rename it back to its original name, (REN_AllowPackageLinkerMismatch | REN_DontCreateRedirectors) should be provided as your rename flags to ensure the export is not detached.

#Synced-CL 49945268

diff --git a/Engine/Source/Runtime/CoreUObject/Public/UObject/ObjectMacros.h b/Engine/Source/Runtime/CoreUObject/Public/UObject/ObjectMacros.h
--- a/Engine/Source/Runtime/CoreUObject/Public/UObject/ObjectMacros.h
+++ b/Engine/Source/Runtime/CoreUObject/Public/UObject/ObjectMacros.h
@@ -2413,8 +2413,9 @@ typedef uint32 ERenameFlags;
 
 /** Default rename behavior */
 #define REN_None					    (0x0000)
-/** UE_DEPRECATED(5.5, "Rename will no longer call ResetLoaders making this flag no longer needed.") -  Rename won't call ResetLoaders or flush async loading. You should pass this if you are renaming a deep subobject and do not need to reset loading for the outer package */
-#define REN_ForceNoResetLoaders		    (0x0001) 
+/** Rename won't call ResetLoaders or flush async loading. You should pass this if you are renaming a deep subobject and do not need to reset loading for the outer package */
+#define REN_ForceNoResetLoaders		    UE_DEPRECATED_MACRO(5.8, "Rename will no longer call ResetLoaders making this flag no longer needed. Prefer REN_AllowPackageLinkerMismatch if you wish to intentionally allow the linker to contain references to objects whose names no longer match what was loaded from disk.") \
+										(0x0001) 
 /** Just test to make sure that the rename is guaranteed to succeed if an non test rename immediately follows */
 #define REN_Test					    (0x0002) 
 /** Indicates that the object (and new outer) should not be dirtied */
