AsyncLoader:
- ZenLoader: add instancing context support, by maintaining an additional package id mapping.
- This allow entity streaming packages to be instanced through Blueprint, Level Instances or, in the future, instancing cooked Prefabs.
- LinkerLoad: allow cooked builds to have private references in other package if both objects share a common outer.
- This allows entities to reference the level (which is private) as an external outer.
This should make LinkerLoad and ZenLoader on parity with instancing context support.

#tests Editor, cooked and AutoRTFM test locally.

#Synced-CL 49950557

diff --git a/Engine/Source/Runtime/CoreUObject/Private/Serialization/AsyncLoading2.cpp b/Engine/Source/Runtime/CoreUObject/Private/Serialization/AsyncLoading2.cpp
--- a/Engine/Source/Runtime/CoreUObject/Private/Serialization/AsyncLoading2.cpp
+++ b/Engine/Source/Runtime/CoreUObject/Private/Serialization/AsyncLoading2.cpp
@@ -833,7 +833,7 @@ struct FAsyncPackageDesc2
 	FPackagePath PackagePathToLoad;
 	// Package referencer
 	FPackageReferencer PackageReferencer;
-	// Packages with a custom name can't be imported
+	// If the package can be imported
 	bool bCanBeImported;
 	// Indicates which loader to use when loading the package
 	EPackageLoader Loader;
@@ -861,11 +861,7 @@ struct FAsyncPackageDesc2
 			UPackageName,
 			MoveTemp(Request.PackagePath),
 			Request.PackageReferencer,
-#if WITH_EDITOR
 			true,
-#else
-			Request.CustomName.IsNone(),
-#endif
 			Loader
 		};
 	}
@@ -6264,6 +6260,7 @@ bool FAsyncLoadingThread2::CreateAsyncPackagesFromQueue(FAsyncLoadingThreadState
 						else
 						{
 							check(PackageLoader == EPackageLoader::Zen);
+							Package->Desc.InstancingContext = PackageDesc.InstancingContext;
 							InitializeAsyncPackageFromPackageStore(ThreadState, &IoBatch, Package, PackageEntry);
 						}
 						Package->StartLoading(ThreadState, IoBatch);
@@ -7085,7 +7082,6 @@ void FAsyncPackage2::ImportPackagesRecursiveInner(FAsyncLoadingThreadState2& Thr
 		}
 
 		bool bIsInstanced = false;
-#if WITH_EDITORONLY_DATA
 		if (bHasImportedPackageNames && LinkerLoadState.IsSet())
 		{
 			// Use the instancing context remap from disk package name to upackage name
@@ -7105,7 +7101,21 @@ void FAsyncPackage2::ImportPackagesRecursiveInner(FAsyncLoadingThreadState2& Thr
 				ImportedPackageNames[LocalImportedPackageIndex] = ImportedPackageUPackageName;
 			}
 		}
-#endif
+		else
+		{
+			// Use the instancing context remap from disk package name to upackage name
+			const FLinkerInstancingContext& InstancingContext = Desc.InstancingContext;
+
+			FPackageId RemappedImportedPackageId = InstancingContext.RemapPackageId(ImportedPackageId);
+
+			if (RemappedImportedPackageId != ImportedPackageId)
+			{
+				bIsInstanced = true;
+				ImportedPackageId = RemappedImportedPackageId;
+				// Rewrite the import table
+				ImportedPackageIds[LocalImportedPackageIndex] = ImportedPackageId;
+			}
+		}
 
 		// Note, localization should apply after remapping from instancing has occurred
 		{
@@ -7299,12 +7309,15 @@ void FAsyncPackage2::ImportPackagesRecursiveInner(FAsyncLoadingThreadState2& Thr
 
 				if (ImportedPackageLoader == EPackageLoader::LinkerLoad)
 				{
-					// Only propagate the instancing context if the imported package is also instanced and it isn't a zen package
-					ImportedPackage->InitializeLinkerLoadState(PackageLoader == EPackageLoader::Zen || !bIsInstanced ? nullptr : &LinkerLoadState->Linker->GetInstancingContext());
+					// Only propagate the instancing context if the imported package is also instanced
+					ImportedPackage->InitializeLinkerLoadState(bIsInstanced ? &LinkerLoadState->Linker->GetInstancingContext() : nullptr);
 				}
 				else
 				{
-					// TODO: Here we should probably also propagate the instancing context if the imported package is also instanced (similar to the call to InitializeLinkerLoadState done above)
+					if (bIsInstanced)
+					{
+						ImportedPackage->Desc.InstancingContext = Desc.InstancingContext;
+					}
 					AsyncLoadingThread.InitializeAsyncPackageFromPackageStore(ThreadState, &IoBatch, ImportedPackage, ImportedPackageEntry);
 				}
 				ImportedPackage->StartLoading(ThreadState, IoBatch);
@@ -8285,7 +8298,7 @@ EEventLoadNodeExecutionResult FAsyncPackage2::Event_ProcessPackageSummary(FAsync
 		// Imported packages won't have a UPackage name set unless they were redirected, in which case they will have the source package name
 		if (Package->Desc.UPackageName.IsNone())
 		{
-			Package->Desc.UPackageName = Package->HeaderData.PackageName;
+			Package->Desc.UPackageName = Package->Desc.InstancingContext.RemapPackage(Package->HeaderData.PackageName);
 		}
 		check(Package->Desc.UPackageId == FPackageId::FromName(Package->Desc.UPackageName));
 		Package->CreateUPackage();
diff --git a/Engine/Source/Runtime/CoreUObject/Private/UObject/LinkerInstancingContext.cpp b/Engine/Source/Runtime/CoreUObject/Private/UObject/LinkerInstancingContext.cpp
--- a/Engine/Source/Runtime/CoreUObject/Private/UObject/LinkerInstancingContext.cpp
+++ b/Engine/Source/Runtime/CoreUObject/Private/UObject/LinkerInstancingContext.cpp
@@ -70,7 +70,7 @@ public:
 	{
 		FLinkerInstancingContextReadLock ScopeLock(Lock);
 		LLM_SCOPE_BYTAG(Loading_LinkerInstancingContext);
-		if (const FName* InstancedPtr = InstancedPackageMap.InstancedPackageMapping.Find(Original))
+		if (const FName* InstancedPtr = InstancedPackageMap.InstancedPackageMapping.Find(FPackageId::FromName(Original)))
 		{
 			Instanced = *InstancedPtr;
 			return true;
@@ -214,6 +214,16 @@ public:
 
 		return RemappedPackageName;
 	}
+
+	FPackageId RemapPackageId(const FPackageId& PackageId)
+	{
+		FPackageId RemappedPackageId;
+		{
+			FLinkerInstancingContextReadLock ScopeLock(Lock);
+			RemappedPackageId = InstancedPackageMap.RemapPackageId(PackageId);
+		}
+		return RemappedPackageId;
+	}
 };
 
 FLinkerInstancingContext::FLinkerInstancingContext()
@@ -249,7 +259,7 @@ FLinkerInstancingContext::FLinkerInstancingContext(const FLinkerInstancingContex
 		break;
 		case EContextSharingMode::Link:
 		{
-			if(InOther.ImmutableData.IsValid())
+			if (InOther.ImmutableData.IsValid())
 			{
 				check(InOther.SharedData.IsValid());
 				
@@ -338,6 +348,18 @@ FName FLinkerInstancingContext::RemapPackage(const FName& PackageName) const
 	return RemappedName;
 }
 
+FPackageId FLinkerInstancingContext::RemapPackageId(const FPackageId& PackageId) const
+{
+	FPackageId RemappedId = SharedData->RemapPackageId(PackageId);
+		
+	if (ImmutableData.IsValid() && RemappedId == PackageId)
+	{
+		RemappedId = ImmutableData->RemapPackageId(PackageId);
+	}
+	
+	return RemappedId;
+}
+
 /**
  * Remap the top level asset part of the path name to its instanced counterpart, otherwise return the name unmodified.
  * i.e. remaps /Path/To/Package.AssetName:Inner to /NewPath/To/NewPackage.NewAssetName:Inner
@@ -453,7 +475,7 @@ void FLinkerInstancedPackageMap::AddPackageMapping(FName Original, FName Instanc
 	FName Key = InstanceMappingDirection == EInstanceMappingDirection::OriginalToInstanced ? Original : Instanced;
 	FName Value = InstanceMappingDirection == EInstanceMappingDirection::OriginalToInstanced ? Instanced : Original;
 
-	const FName& StoredValue = InstancedPackageMapping.FindOrAdd(Key, Value);
+	const FName& StoredValue = InstancedPackageMapping.FindOrAdd(FPackageId::FromName(Key), Value);
 	if (StoredValue == Value)
 	{
 		bIsInstanced |= !Value.IsNone();
@@ -531,13 +553,13 @@ bool FLinkerInstancedPackageMap::ValidateAgainst(FLinkerInstancedPackageMap& Oth
 	// This method ensures the object passed as a parameter doesn't contain any mapping that diverge from us
 	bool bResult = true;
 
-	for (TPair<FName, FName> It : InstancedPackageMapping)
+	for (TPair<FPackageId, FName> It : InstancedPackageMapping)
 	{
 		FName* OtherInstanced = OtherPackageMap.InstancedPackageMapping.Find(It.Key);
 
 		bool bThisResult = (!OtherInstanced || (*OtherInstanced == It.Value));
 		
-		ensureAlwaysMsgf(bThisResult, TEXT("Discrepancy in instance package mapping for %s, maps to %s and %s"), *It.Key.ToString(), *It.Value.ToString(), *OtherInstanced->ToString());
+		ensureAlwaysMsgf(bThisResult, TEXT("Discrepancy in instance package mapping for 0x%llx, maps to %s and %s"), It.Key.Value(), *It.Value.ToString(), *OtherInstanced->ToString());
 
 		bResult &= bThisResult;
 	}
diff --git a/Engine/Source/Runtime/CoreUObject/Private/UObject/LinkerLoad.cpp b/Engine/Source/Runtime/CoreUObject/Private/UObject/LinkerLoad.cpp
--- a/Engine/Source/Runtime/CoreUObject/Private/UObject/LinkerLoad.cpp
+++ b/Engine/Source/Runtime/CoreUObject/Private/UObject/LinkerLoad.cpp
@@ -4334,14 +4334,10 @@ bool FLinkerLoad::VerifyImportInner(const int32 ImportIndex, FString& WarningSuf
 							}
 						}
 
-						// Since import can have export outer and vice versa now, consider import and export sharing outers to be allowed, in editor only
+						// Since import can have export outer and vice versa now, consider import and export sharing outers to be allowed
 						auto IsPrivateImportAllowed = [this](int32 InImportIndex)
 						{
-						#if WITH_EDITOR
 							return ImportIsInAnyExport(InImportIndex) || AnyExportIsInImport(InImportIndex) || AnyExportShareOuterWithImport(InImportIndex);
-						#else
-							return false;
-						#endif
 						};
 
 						const bool bIsImportPublic = !!(SourceExport.ObjectFlags & RF_Public);
diff --git a/Engine/Source/Runtime/CoreUObject/Public/UObject/LinkerInstancingContext.h b/Engine/Source/Runtime/CoreUObject/Public/UObject/LinkerInstancingContext.h
--- a/Engine/Source/Runtime/CoreUObject/Public/UObject/LinkerInstancingContext.h
+++ b/Engine/Source/Runtime/CoreUObject/Public/UObject/LinkerInstancingContext.h
@@ -3,6 +3,7 @@
 #pragma once
 
 #include "UObject/SoftObjectPath.h"
+#include "IO/PackageId.h"
 
 /**
  * Helper class to map between an original package and an instance of it (including world partition cells).
@@ -38,7 +39,7 @@ public:
 	/** Remap the package name from the import table to its instanced counterpart, otherwise return the name unmodified. */
 	FName RemapPackage(const FName& PackageName, bool* bWasFoundPtr = nullptr) const
 	{
-		if (const FName* RemappedName = InstancedPackageMapping.Find(PackageName))
+		if (const FName* RemappedName = InstancedPackageMapping.Find(FPackageId::FromName(PackageName)))
 		{
 			if (bWasFoundPtr)
 			{
@@ -52,7 +53,24 @@ public:
 		}
 		return PackageName;
 	}
-	
+
+	FPackageId RemapPackageId(const FPackageId& PackageId, bool* bWasFoundPtr = nullptr) const
+	{
+		if (const FName* RemappedName = InstancedPackageMapping.Find(PackageId))
+		{
+			if (bWasFoundPtr)
+			{
+				*bWasFoundPtr = true;
+			}
+			return FPackageId::FromName(*RemappedName);
+		}
+		if (bWasFoundPtr)
+		{
+			*bWasFoundPtr = false;
+		}
+		return PackageId;
+	}
+
 	/** Add a mapping from a package name to a new package name. There should be no separators (. or :) in these strings. */
 	COREUOBJECT_API void AddPackageMapping(FName Original, FName Instanced);
 
@@ -72,7 +90,7 @@ private:
 	 * Key=Original and Value=Instanced when InstanceMappingDirection==OriginalToInstanced
 	 * Key=Instanced and Value=Original when InstanceMappingDirection==InstancedToOriginal
 	 */
-	TMap<FName, FName> InstancedPackageMapping;
+	TMap<FPackageId, FName> InstancedPackageMapping;
 
 	/**
 	 * In which direction has this mapping been built?
@@ -139,6 +157,7 @@ public:
 
 	/** Remap the package name from the import table to its instanced counterpart, otherwise return the name unmodified. */
 	COREUOBJECT_API FName RemapPackage(const FName& PackageName) const;
+	COREUOBJECT_API FPackageId RemapPackageId(const FPackageId& PackageId) const;
 
 	/**
 	 * Remap the top level asset part of the path name to its instanced counterpart, otherwise return the name unmodified. 
