Core: Cleanup of 5.1 deprecations in Core and CoreUObject.
A strict code removal operation.
Skip removal of enum values to minimize risk.
Skip removal of extern bool GIsFirstInstance since it is still being used.
Skip removal of UE::Tasks::Private::IsThreadRetractingTask since it is still being used.
Skip removal of UE::SoftObjectPath::Private functions ConvertSoftObjectPaths and ConvertObjectPathNames that are still being used.

#rnx
#tests BuildCookRun Lyra, Editor boot

#Synced-CL 49950633

diff --git a/Engine/Plugins/Editor/ContentBrowser/ContentBrowserAssetDataSource/Source/ContentBrowserAssetDataSource/Private/ContentBrowserAssetDataSource.cpp b/Engine/Plugins/Editor/ContentBrowser/ContentBrowserAssetDataSource/Source/ContentBrowserAssetDataSource/Private/ContentBrowserAssetDataSource.cpp
--- a/Engine/Plugins/Editor/ContentBrowser/ContentBrowserAssetDataSource/Source/ContentBrowserAssetDataSource/Private/ContentBrowserAssetDataSource.cpp
+++ b/Engine/Plugins/Editor/ContentBrowser/ContentBrowserAssetDataSource/Source/ContentBrowserAssetDataSource/Private/ContentBrowserAssetDataSource.cpp
@@ -1776,9 +1776,6 @@ void UContentBrowserAssetDataSource::CompileFilter(const FName InPath, const FCo
 						FARFilter CustomSourceAssetsFilter;
 						CustomSourceAssetsFilter.PackageNames = Params.AssetDataFilter->InclusiveFilter.PackageNames.Array();
 						CustomSourceAssetsFilter.PackagePaths = Params.AssetDataFilter->InclusiveFilter.PackagePaths.Array();
-PRAGMA_DISABLE_DEPRECATION_WARNINGS
-						CustomSourceAssetsFilter.ObjectPaths = Params.AssetDataFilter->InclusiveFilter.ObjectPaths.Array();
-PRAGMA_ENABLE_DEPRECATION_WARNINGS
 						CustomSourceAssetsFilter.ClassPaths = Params.AssetDataFilter->InclusiveFilter.ClassPaths.Array();
 						CustomSourceAssetsFilter.TagsAndValues = Params.AssetDataFilter->InclusiveFilter.TagsAndValues;
 						CustomSourceAssetsFilter.bIncludeOnlyOnDiskAssets = Params.AssetDataFilter->InclusiveFilter.bIncludeOnlyOnDiskAssets;
diff --git a/Engine/Plugins/Runtime/ModelViewViewModel/Source/ModelViewViewModel/Public/MVVMMessageLog.h b/Engine/Plugins/Runtime/ModelViewViewModel/Source/ModelViewViewModel/Public/MVVMMessageLog.h
--- a/Engine/Plugins/Runtime/ModelViewViewModel/Source/ModelViewViewModel/Public/MVVMMessageLog.h
+++ b/Engine/Plugins/Runtime/ModelViewViewModel/Source/ModelViewViewModel/Public/MVVMMessageLog.h
@@ -27,7 +27,6 @@ public:
 	using ::FMessageLog::AddMessage;
 	using ::FMessageLog::AddMessages;
 	UE_API TSharedRef<FTokenizedMessage> Message(EMessageSeverity::Type InSeverity, const FText& InMessage);
-	using ::FMessageLog::CriticalError;
 	UE_API TSharedRef<FTokenizedMessage> Error(const FText& InMessage);
 	UE_API TSharedRef<FTokenizedMessage> PerformanceWarning(const FText& InMessage);
 	UE_API TSharedRef<FTokenizedMessage> Warning(const FText& InMessage);
diff --git a/Engine/Source/Editor/ClassViewer/Private/UnloadedBlueprintData.cpp b/Engine/Source/Editor/ClassViewer/Private/UnloadedBlueprintData.cpp
--- a/Engine/Source/Editor/ClassViewer/Private/UnloadedBlueprintData.cpp
+++ b/Engine/Source/Editor/ClassViewer/Private/UnloadedBlueprintData.cpp
@@ -133,16 +133,6 @@ TSharedPtr<FString> FUnloadedBlueprintData::GetClassName() const
 	return TSharedPtr<FString>();
 }
 
-FName FUnloadedBlueprintData::GetClassPath() const
-{
-	if (ClassViewerNode.IsValid())
-	{
-		return FName(*ClassViewerNode.Pin()->ClassPath.ToString());
-	}
-
-	return FName();
-}
-
 FTopLevelAssetPath FUnloadedBlueprintData::GetClassPathName() const
 {
 	if (ClassViewerNode.IsValid())
diff --git a/Engine/Source/Editor/ClassViewer/Private/UnloadedBlueprintData.h b/Engine/Source/Editor/ClassViewer/Private/UnloadedBlueprintData.h
--- a/Engine/Source/Editor/ClassViewer/Private/UnloadedBlueprintData.h
+++ b/Engine/Source/Editor/ClassViewer/Private/UnloadedBlueprintData.h
@@ -38,9 +38,6 @@ public:
 
 	virtual TSharedPtr<FString> GetClassName() const override;
 
-	UE_DEPRECATED(5.1, "Class names are now represented by path names. Please use GetClassPathName.")
-	virtual FName GetClassPath() const override;
-
 	virtual FTopLevelAssetPath GetClassPathName() const override;
 
 	virtual const UClass* GetClassWithin() const override;
diff --git a/Engine/Source/Editor/ClassViewer/Public/ClassViewerFilter.h b/Engine/Source/Editor/ClassViewer/Public/ClassViewerFilter.h
--- a/Engine/Source/Editor/ClassViewer/Public/ClassViewerFilter.h
+++ b/Engine/Source/Editor/ClassViewer/Public/ClassViewerFilter.h
@@ -360,12 +360,6 @@ public:
 	 */
 	virtual TSharedPtr<FString> GetClassName() const = 0;
 
-	/**
-	 * Get the class path of this blueprint.
-	 */
-	UE_DEPRECATED(5.1, "Class names are now represented by path names. Please use GetClassPathName.")
-	virtual FName GetClassPath() const = 0;
-
 	/**
 	 * Get the class path of this blueprint.
 	 */
diff --git a/Engine/Source/Editor/KismetCompiler/Private/KismetCompiler.cpp b/Engine/Source/Editor/KismetCompiler/Private/KismetCompiler.cpp
--- a/Engine/Source/Editor/KismetCompiler/Private/KismetCompiler.cpp
+++ b/Engine/Source/Editor/KismetCompiler/Private/KismetCompiler.cpp
@@ -5420,10 +5420,6 @@ void FKismetCompilerContext::PostCDOCompiled(const UObject::FPostCDOCompiledCont
 	NewClass->GetDefaultObject(false)->PostCDOCompiled(Context);
 	if (!Context.bIsSkeletonOnly)
 	{
-		PRAGMA_DISABLE_DEPRECATION_WARNINGS;
-		NewClass->GetDefaultObject(false)->PostCDOCompiled();
-		PRAGMA_ENABLE_DEPRECATION_WARNINGS;
-		
 		for (const TFunction<void(const UObject::FPostCDOCompiledContext&, UObject*)>& Step : PostCDOCompileSteps)
 		{
 			Step(Context, NewClass->GetDefaultObject(false));
diff --git a/Engine/Source/Editor/UMGEditor/Private/WidgetBlueprintEditorUtils.cpp b/Engine/Source/Editor/UMGEditor/Private/WidgetBlueprintEditorUtils.cpp
--- a/Engine/Source/Editor/UMGEditor/Private/WidgetBlueprintEditorUtils.cpp
+++ b/Engine/Source/Editor/UMGEditor/Private/WidgetBlueprintEditorUtils.cpp
@@ -2510,13 +2510,6 @@ public:
 		return ClassName;
 	}
 
-	virtual FName GetClassPath() const
-	{
-PRAGMA_DISABLE_DEPRECATION_WARNINGS
-		return ClassPath.ToFName();
-PRAGMA_ENABLE_DEPRECATION_WARNINGS
-	}
-
 	virtual FTopLevelAssetPath GetClassPathName() const
 	{
 		return ClassPath;
diff --git a/Engine/Source/Editor/UnrealEd/Private/ComponentTypeRegistry.cpp b/Engine/Source/Editor/UnrealEd/Private/ComponentTypeRegistry.cpp
--- a/Engine/Source/Editor/UnrealEd/Private/ComponentTypeRegistry.cpp
+++ b/Engine/Source/Editor/UnrealEd/Private/ComponentTypeRegistry.cpp
@@ -228,14 +228,6 @@ namespace UE::Editor::ComponentTypeRegistry::Private
 			return ClassName;
 		}
 
-		UE_DEPRECATED(5.1, "Class names are now represented by path names. Please use GetClassPathName.")
-		virtual FName GetClassPath() const override
-		{
-PRAGMA_DISABLE_DEPRECATION_WARNINGS
-			return ClassPath.ToFName();
-PRAGMA_ENABLE_DEPRECATION_WARNINGS
-		}
-		
 		virtual FTopLevelAssetPath GetClassPathName() const override
 		{
 			return ClassPath;
diff --git a/Engine/Source/Editor/UnrealEd/Private/EditorClassUtils.cpp b/Engine/Source/Editor/UnrealEd/Private/EditorClassUtils.cpp
--- a/Engine/Source/Editor/UnrealEd/Private/EditorClassUtils.cpp
+++ b/Engine/Source/Editor/UnrealEd/Private/EditorClassUtils.cpp
@@ -387,25 +387,12 @@ bool FEditorClassUtils::IsBlueprintAsset(const FAssetData& InAssetData, bool* bO
 	return bIsBP || bIsBPGC;
 }
 
-FName FEditorClassUtils::GetClassPathFromAssetTag(const FAssetData& InAssetData)
-{
-	const FString GeneratedClassPath = InAssetData.GetTagValueRef<FString>(FBlueprintTags::GeneratedClassPath);
-	return FName(FPackageName::ExportTextPathToObjectPath(FStringView(GeneratedClassPath)));
-}
-
 FTopLevelAssetPath FEditorClassUtils::GetClassPathNameFromAssetTag(const FAssetData& InAssetData)
 {
 	const FString GeneratedClassPath = InAssetData.GetTagValueRef<FString>(FBlueprintTags::GeneratedClassPath);
 	return FTopLevelAssetPath(FPackageName::ExportTextPathToObjectPath(FStringView(GeneratedClassPath)));
 }
 
-FName FEditorClassUtils::GetClassPathFromAsset(const FAssetData& InAssetData, bool bGenerateClassPathIfMissing /*= false*/)
-{
-PRAGMA_DISABLE_DEPRECATION_WARNINGS
-	return GetClassPathNameFromAsset(InAssetData, bGenerateClassPathIfMissing).ToFName();
-PRAGMA_ENABLE_DEPRECATION_WARNINGS
-}
-
 FTopLevelAssetPath FEditorClassUtils::GetClassPathNameFromAsset(const FAssetData& InAssetData, bool bGenerateClassPathIfMissing /*= false*/)
 {
 	bool bIsBPGC = false;
@@ -467,4 +454,4 @@ void FEditorClassUtils::GetImplementedInterfaceClassPathsFromAsset(const struct
 			}
 		}
 	}
-}
\ No newline at end of file
+}
diff --git a/Engine/Source/Editor/UnrealEd/Public/EditorClassUtils.h b/Engine/Source/Editor/UnrealEd/Public/EditorClassUtils.h
--- a/Engine/Source/Editor/UnrealEd/Public/EditorClassUtils.h
+++ b/Engine/Source/Editor/UnrealEd/Public/EditorClassUtils.h
@@ -160,15 +160,6 @@ namespace FEditorClassUtils
 	 */
 	UNREALED_API bool IsBlueprintAsset(const FAssetData& InAssetData, bool* bOutIsBPGC = nullptr);
 
-	/**
-	 * Gets the class path from the asset tag (i.e. GeneratedClassPath tag on blueprints)
-	 * 
-	 * @param	InAssetData		Reference to an asset data entry.
-	 * @return					Class path or None if the asset cannot or doesn't have a class associated with it
-	 */
-	UE_DEPRECATED(5.1, "Class names are now represented by path names. Please use GetClassPathNameFromAssetTag.")
-	UNREALED_API FName GetClassPathFromAssetTag(const FAssetData& InAssetData);
-
 	/**
 	 * Gets the class path from the asset tag (i.e. GeneratedClassPath tag on blueprints)
 	 * 
@@ -177,17 +168,6 @@ namespace FEditorClassUtils
 	 */
 	UNREALED_API FTopLevelAssetPath GetClassPathNameFromAssetTag(const FAssetData& InAssetData);
 
-	/**
-	 * Gets the object path of the class associated with the specified asset 
-	 * (i.e. the BlueprintGeneratedClass of a Blueprint asset or the BlueprintGeneratedClass asset itself)
-	 * 
-	 * @param	InAssetData					Reference to an asset data entry
-	 * @param	bGenerateClassPathIfMissing	Whether to generate a class path if the class is missing (and the asset can have a class associated with it)
-	 * @return								Class path or None if the asset cannot or doesn't have a class associated with it
-	 */
-	UE_DEPRECATED(5.1, "Class names are now represented by path names. Please use GetClassPathNameFromAsset.")
-	UNREALED_API FName GetClassPathFromAsset(const FAssetData& InAssetData, bool bGenerateClassPathIfMissing = false);
-
 	/**
 	 * Gets the object path of the class associated with the specified asset 
 	 * (i.e. the BlueprintGeneratedClass of a Blueprint asset or the BlueprintGeneratedClass asset itself)
diff --git a/Engine/Source/Runtime/Core/Private/Android/AndroidPlatformMisc.cpp b/Engine/Source/Runtime/Core/Private/Android/AndroidPlatformMisc.cpp
--- a/Engine/Source/Runtime/Core/Private/Android/AndroidPlatformMisc.cpp
+++ b/Engine/Source/Runtime/Core/Private/Android/AndroidPlatformMisc.cpp
@@ -3116,11 +3116,6 @@ const char* FAndroidMisc::GetThreadName(uint32 ThreadId)
 	return ThreadName ? *ThreadName : nullptr;
 }
 
-void FAndroidMisc::SetDeviceOrientation(EDeviceScreenOrientation NewDeviceOrentation)
-{
-	SetAllowedDeviceOrientation(NewDeviceOrentation);
-}
-
 void FAndroidMisc::SetCellularPreference(int32 Value)
 {
 #if USE_ANDROID_JNI
diff --git a/Engine/Source/Runtime/Core/Private/GenericPlatform/GenericPlatformMisc.cpp b/Engine/Source/Runtime/Core/Private/GenericPlatform/GenericPlatformMisc.cpp
--- a/Engine/Source/Runtime/Core/Private/GenericPlatform/GenericPlatformMisc.cpp
+++ b/Engine/Source/Runtime/Core/Private/GenericPlatform/GenericPlatformMisc.cpp
@@ -2143,11 +2143,6 @@ EDeviceScreenOrientation FGenericPlatformMisc::GetDeviceOrientation()
 	return EDeviceScreenOrientation::Unknown;
 }
 
-void FGenericPlatformMisc::SetDeviceOrientation(EDeviceScreenOrientation NewDeviceOrientation)
-{
-	FPlatformMisc::SetAllowedDeviceOrientation(NewDeviceOrientation);
-}
-
 EDeviceScreenOrientation FGenericPlatformMisc::GetAllowedDeviceOrientation()
 {
 	return AllowedDeviceOrientation;
diff --git a/Engine/Source/Runtime/Core/Private/IOS/IOSPlatformMisc.cpp b/Engine/Source/Runtime/Core/Private/IOS/IOSPlatformMisc.cpp
--- a/Engine/Source/Runtime/Core/Private/IOS/IOSPlatformMisc.cpp
+++ b/Engine/Source/Runtime/Core/Private/IOS/IOSPlatformMisc.cpp
@@ -574,11 +574,6 @@ EDeviceScreenOrientation FIOSPlatformMisc::GetDeviceOrientation()
 #endif
 }
 
-void FIOSPlatformMisc::SetDeviceOrientation(EDeviceScreenOrientation NewDeviceOrientation)
-{
-	SetAllowedDeviceOrientation(NewDeviceOrientation);
-}
-
 void FIOSPlatformMisc::SetAllowedDeviceOrientation(EDeviceScreenOrientation NewAllowedDeviceOrientation)
 {
 	AllowedDeviceOrientation = NewAllowedDeviceOrientation;
diff --git a/Engine/Source/Runtime/Core/Private/Logging/MessageLog.cpp b/Engine/Source/Runtime/Core/Private/Logging/MessageLog.cpp
--- a/Engine/Source/Runtime/Core/Private/Logging/MessageLog.cpp
+++ b/Engine/Source/Runtime/Core/Private/Logging/MessageLog.cpp
@@ -175,21 +175,6 @@ TSharedRef<FTokenizedMessage> FMessageLog::Message( EMessageSeverity::Type InSev
 	return Message;
 }
 
-TSharedRef<FTokenizedMessage> FMessageLog::CriticalError( const FText& InMessage )
-{
-	// Deprecated: This call has been deprecated. As EMessageSeverity::CriticalError has been deprecated as well, this call
-	//		will no longer function as before. To avoid this going unnoticed an explicit call to checkf will be made
-	checkf(false, 
-		TEXT("[Deprecated] FMessageLog::CriticalError has been deprecated. Please update callsites. The original error message: \n%s"), 
-		*InMessage.ToString());
-
-	LLM_SCOPE_BYTAG(EngineMisc_MessageLog);
-	
-	TSharedRef<FTokenizedMessage> Message = FTokenizedMessage::Create(EMessageSeverity::Error, InMessage);
-	Messages.Add(Message);
-	return Message;
-}
-
 TSharedRef<FTokenizedMessage> FMessageLog::Error( const FText& InMessage )
 {
 	LLM_SCOPE_BYTAG(EngineMisc_MessageLog);
diff --git a/Engine/Source/Runtime/Core/Public/Android/AndroidPlatformMisc.h b/Engine/Source/Runtime/Core/Public/Android/AndroidPlatformMisc.h
--- a/Engine/Source/Runtime/Core/Public/Android/AndroidPlatformMisc.h
+++ b/Engine/Source/Runtime/Core/Public/Android/AndroidPlatformMisc.h
@@ -40,8 +40,6 @@ struct FAndroidMisc : public FGenericPlatformMisc
 	static CORE_API void PlatformTearDown();
 	static CORE_API void PlatformHandleSplashScreen(bool ShowSplashScreen);
     static EDeviceScreenOrientation GetDeviceOrientation() { return DeviceOrientation; }
-	UE_DEPRECATED(5.1, "SetDeviceOrientation is deprecated. Use SetAllowedDeviceOrientation instead.")
-	static CORE_API void SetDeviceOrientation(EDeviceScreenOrientation NewDeviceOrentation);
 	static CORE_API void SetAllowedDeviceOrientation(EDeviceScreenOrientation NewAllowedDeviceOrientation);
 
 	// Change this to an Enum with Always allow, allow and deny
diff --git a/Engine/Source/Runtime/Core/Public/GenericPlatform/GenericPlatformMath.h b/Engine/Source/Runtime/Core/Public/GenericPlatform/GenericPlatformMath.h
--- a/Engine/Source/Runtime/Core/Public/GenericPlatform/GenericPlatformMath.h
+++ b/Engine/Source/Runtime/Core/Public/GenericPlatform/GenericPlatformMath.h
@@ -599,30 +599,6 @@ struct FGenericPlatformMath
 		return BitCast<uint64>(A) >= (uint64)0x8000000000000000; // Detects sign bit.
 	}
 
-	UE_DEPRECATED(5.1, "IsNegativeFloat has been deprecated in favor of IsNegativeOrNegativeZero or simply A < 0.")
-	[[nodiscard]] static UE_FORCEINLINE_HINT bool IsNegativeFloat(float A)
-	{
-		return IsNegativeOrNegativeZero(A);
-	}
-
-	UE_DEPRECATED(5.1, "IsNegativeDouble has been deprecated in favor of IsNegativeOrNegativeZero or simply A < 0.")
-	[[nodiscard]] static UE_FORCEINLINE_HINT bool IsNegativeDouble(double A)
-	{
-		return IsNegativeOrNegativeZero(A);
-	}
-
-	UE_DEPRECATED(5.1, "IsNegative has been deprecated in favor of IsNegativeOrNegativeZero or simply A < 0.")
-	[[nodiscard]] static UE_FORCEINLINE_HINT bool IsNegative(float A)
-	{
-		return IsNegativeOrNegativeZero(A);
-	}
-
-	UE_DEPRECATED(5.1, "IsNegative has been deprecated in favor of IsNegativeOrNegativeZero or simply A < 0.")
-	[[nodiscard]] static UE_FORCEINLINE_HINT bool IsNegative(double A)
-	{
-		return IsNegativeOrNegativeZero(A);
-	}
-
 	/** Returns a random integer between 0 and RAND_MAX, inclusive */
 	[[nodiscard]] static UE_FORCEINLINE_HINT int32 Rand() { return rand(); }
 
@@ -1377,4 +1353,4 @@ template<>
 [[nodiscard]] UE_FORCEINLINE_HINT double FGenericPlatformMath::Abs(const double A)
 {
 	return fabs(A);
-}
\ No newline at end of file
+}
diff --git a/Engine/Source/Runtime/Core/Public/GenericPlatform/GenericPlatformMisc.h b/Engine/Source/Runtime/Core/Public/GenericPlatform/GenericPlatformMisc.h
--- a/Engine/Source/Runtime/Core/Public/GenericPlatform/GenericPlatformMisc.h
+++ b/Engine/Source/Runtime/Core/Public/GenericPlatform/GenericPlatformMisc.h
@@ -1765,12 +1765,6 @@ public:
 	 */
 	static CORE_API EDeviceScreenOrientation GetDeviceOrientation();
 
-	/**
-	 * Change the orientation of the device.
-	 */
-	UE_DEPRECATED(5.1, "SetDeviceOrientation is deprecated. Use SetAllowedDeviceOrientation instead.")
-	static CORE_API void SetDeviceOrientation(EDeviceScreenOrientation NewDeviceOrientation);
-
 	/**
 	 * Returns the allowed orientation of the device. This is NOT the same as GetDeviceOrientation, which only returns Portrait, LandscapeLeft, 
 	 * PortraitUpsideDown or LandscapeRight. The allowed orientation limits what orientation your device can have. So if you set the allowed orientation 
diff --git a/Engine/Source/Runtime/Core/Public/IOS/IOSPlatformMisc.h b/Engine/Source/Runtime/Core/Public/IOS/IOSPlatformMisc.h
--- a/Engine/Source/Runtime/Core/Public/IOS/IOSPlatformMisc.h
+++ b/Engine/Source/Runtime/Core/Public/IOS/IOSPlatformMisc.h
@@ -68,8 +68,6 @@ struct CORE_API FIOSPlatformMisc : public FApplePlatformMisc
 	static float GetDeviceTemperatureLevel();
 	static EDeviceThermalState GetDeviceThermalState();
 	static EDeviceScreenOrientation GetDeviceOrientation();
-	UE_DEPRECATED(5.1, "SetDeviceOrientation is deprecated. Use SetAllowedDeviceOrientation instead.")
-	static void SetDeviceOrientation(EDeviceScreenOrientation NewDeviceOrientation);
 	static void SetAllowedDeviceOrientation(EDeviceScreenOrientation NewAllowedDeviceOrientation);
 	static int32 GetDeviceVolume();
 	static void SetBrightness(float Brightness);
diff --git a/Engine/Source/Runtime/Core/Public/Logging/MessageLog.h b/Engine/Source/Runtime/Core/Public/Logging/MessageLog.h
--- a/Engine/Source/Runtime/Core/Public/Logging/MessageLog.h
+++ b/Engine/Source/Runtime/Core/Public/Logging/MessageLog.h
@@ -56,8 +56,6 @@ public:
 	 * @return the message for chaining calls.
 	 */
 	CORE_API TSharedRef<FTokenizedMessage> Message( EMessageSeverity::Type InSeverity, const FText& InMessage = FText() );
-	UE_DEPRECATED(5.1, "CriticalError was removed because it can't trigger an assert at the callsite. Use 'checkf' instead.")
-	CORE_API TSharedRef<FTokenizedMessage> CriticalError(const FText& InMessage = FText());
 	CORE_API TSharedRef<FTokenizedMessage> Error( const FText& InMessage = FText() );
 	CORE_API TSharedRef<FTokenizedMessage> PerformanceWarning( const FText& InMessage = FText() );
 	CORE_API TSharedRef<FTokenizedMessage> Warning( const FText& InMessage = FText() );
diff --git a/Engine/Source/Runtime/Core/Public/Misc/DataDrivenPlatformInfoRegistry.h b/Engine/Source/Runtime/Core/Public/Misc/DataDrivenPlatformInfoRegistry.h
--- a/Engine/Source/Runtime/Core/Public/Misc/DataDrivenPlatformInfoRegistry.h
+++ b/Engine/Source/Runtime/Core/Public/Misc/DataDrivenPlatformInfoRegistry.h
@@ -300,17 +300,6 @@ struct FDataDrivenPlatformInfoRegistry
 	static CORE_API const TArray<FName> GetSortedPlatformNames(EPlatformInfoType PlatformType);
 	static CORE_API const TArray<const FDataDrivenPlatformInfo*>& GetSortedPlatformInfos(EPlatformInfoType PlatformType);
 
-	UE_DEPRECATED(5.1, "Use GetSoprtedPlatformNames that takes a PlatformType parameter")
-	static const TArray<FName> GetSortedPlatformNames()
-	{
-		return GetSortedPlatformNames(EPlatformInfoType::AllPlatformInfos);
-	}
-	UE_DEPRECATED(5.1, "Use GetSortedPlatformInfos that takes a PlatformType parameter")
-	static const TArray<const FDataDrivenPlatformInfo*>& GetSortedPlatformInfos()
-	{
-		return GetSortedPlatformInfos(EPlatformInfoType::AllPlatformInfos);
-	}
-
 	/**
 	 * Gets a list of all known confidential platforms (note these are just the platforms you have access to, so, for example PS4 won't be
 	 * returned if you are not a PS4 licensee)
diff --git a/Engine/Source/Runtime/Core/Public/Misc/DateTime.h b/Engine/Source/Runtime/Core/Public/Misc/DateTime.h
--- a/Engine/Source/Runtime/Core/Public/Misc/DateTime.h
+++ b/Engine/Source/Runtime/Core/Public/Misc/DateTime.h
@@ -129,12 +129,6 @@ public:
 		return *this;
 	}
 
-	UE_DEPRECATED(5.1, "Adding dates doesn't make sense. Please use FDateTime + FTimespan instead")
-	FDateTime operator+(const FDateTime& Other)
-	{
-		return FDateTime(Ticks + Other.Ticks);
-	}
-
 	/**
 	 * Returns time span between this date and the given date.
 	 *
diff --git a/Engine/Source/Runtime/Core/Public/Misc/KeyChainUtilities.h b/Engine/Source/Runtime/Core/Public/Misc/KeyChainUtilities.h
--- a/Engine/Source/Runtime/Core/Public/Misc/KeyChainUtilities.h
+++ b/Engine/Source/Runtime/Core/Public/Misc/KeyChainUtilities.h
@@ -83,7 +83,6 @@ public:
 		return *this;
 	}
 
-	PRAGMA_DISABLE_DEPRECATION_WARNINGS
 	FRSAKeyHandle GetSigningKey() const { return SigningKey; }
 	void SetSigningKey(FRSAKeyHandle key) { SigningKey = key; }
 
@@ -94,15 +93,10 @@ public:
 	TMap<FGuid, FNamedAESKey>& GetEncryptionKeys() { return EncryptionKeys; }
 
 	void SetEncryptionKeys(const TMap<FGuid, FNamedAESKey>& keys) { EncryptionKeys = keys; }
-	PRAGMA_ENABLE_DEPRECATION_WARNINGS
 
-	UE_DEPRECATED(5.1, "Use Get/SetSigningKey instead")
+private:
 	FRSAKeyHandle SigningKey = InvalidRSAKeyHandle;
-
-	UE_DEPRECATED(5.1, "Use GetEncryptionKeys instead")
 	TMap<FGuid, FNamedAESKey> EncryptionKeys;
-	
-	UE_DEPRECATED(5.1, "Use Get/SetPrincipalEncryptionKey instead")
 	const FNamedAESKey* MasterEncryptionKey = nullptr;
 };
 
diff --git a/Engine/Source/Runtime/CoreUObject/Private/AssetRegistry/AssetBundleData.cpp b/Engine/Source/Runtime/CoreUObject/Private/AssetRegistry/AssetBundleData.cpp
--- a/Engine/Source/Runtime/CoreUObject/Private/AssetRegistry/AssetBundleData.cpp
+++ b/Engine/Source/Runtime/CoreUObject/Private/AssetRegistry/AssetBundleData.cpp
@@ -340,75 +340,6 @@ PRAGMA_ENABLE_DEPRECATION_WARNINGS
 #endif
 }
 
-void FAssetBundleData::AddBundleAsset(FName BundleName, const FSoftObjectPath& AssetPath)
-{
-	if (!AssetPath.IsValid())
-	{
-		return;
-	}
-
-	FAssetBundleEntry* FoundEntry = FindEntry(BundleName);
-
-	if (!FoundEntry)
-	{
-		FoundEntry = &Bundles.Emplace_GetRef(BundleName);
-	}
-
-#if WITH_EDITORONLY_DATA
-PRAGMA_DISABLE_DEPRECATION_WARNINGS
-	FoundEntry->BundleAssets.AddUnique(AssetPath);
-PRAGMA_ENABLE_DEPRECATION_WARNINGS
-#endif
-	FoundEntry->AssetPaths.AddUnique(AssetPath.GetAssetPath());
-}
-
-void FAssetBundleData::AddBundleAssets(FName BundleName, const TArray<FSoftObjectPath>& AssetPaths)
-{
-	FAssetBundleEntry* FoundEntry = FindEntry(BundleName);
-
-	for (const FSoftObjectPath& Path : AssetPaths)
-	{
-		if (Path.IsValid())
-		{
-			// Only create if required
-			if (!FoundEntry)
-			{
-				FoundEntry = &Bundles.Emplace_GetRef(BundleName);
-			}
-
-#if WITH_EDITORONLY_DATA
-PRAGMA_DISABLE_DEPRECATION_WARNINGS
-			FoundEntry->BundleAssets.AddUnique(Path);
-PRAGMA_ENABLE_DEPRECATION_WARNINGS
-#endif
-			FoundEntry->AssetPaths.AddUnique(Path.GetAssetPath());
-		}
-	}
-}
-
-void FAssetBundleData::SetBundleAssets(FName BundleName, TArray<FSoftObjectPath>&& AssetPaths)
-{
-	FAssetBundleEntry* FoundEntry = FindEntry(BundleName);
-
-	if (!FoundEntry)
-	{
-		FoundEntry = &Bundles.Emplace_GetRef(BundleName);
-	}
-
-	FoundEntry->AssetPaths.Reset(AssetPaths.Num());
-	for( const FSoftObjectPath& Path : AssetPaths )
-	{	
-		FoundEntry->AssetPaths.Add(Path.GetAssetPath());
-	}
-
-#if WITH_EDITORONLY_DATA
-PRAGMA_DISABLE_DEPRECATION_WARNINGS
-	FoundEntry->BundleAssets.Reset();
-	Algo::Transform(FoundEntry->AssetPaths, FoundEntry->BundleAssets, [](FTopLevelAssetPath Path) { return FSoftObjectPath(Path, {}); });
-PRAGMA_ENABLE_DEPRECATION_WARNINGS
-#endif
-}
-
 void FAssetBundleData::AddBundleAssetTruncated(FName BundleName, const FSoftObjectPath& AssetPath)
 {
 	if (!AssetPath.IsValid())
diff --git a/Engine/Source/Runtime/CoreUObject/Private/AssetRegistry/AssetData.cpp b/Engine/Source/Runtime/CoreUObject/Private/AssetRegistry/AssetData.cpp
--- a/Engine/Source/Runtime/CoreUObject/Private/AssetRegistry/AssetData.cpp
+++ b/Engine/Source/Runtime/CoreUObject/Private/AssetRegistry/AssetData.cpp
@@ -340,16 +340,6 @@ private:
 
 } // namespace UE::AssetRegistry
 
-FAssetData::FAssetData(FName InPackageName, FName InPackagePath, FName InAssetName, FName InAssetClass, FAssetDataTagMap InTags, TArrayView<const int32> InChunkIDs, uint32 InPackageFlags)
-	: FAssetData(InPackageName, InPackagePath, InAssetName, FAssetData::TryConvertShortClassNameToPathName(InAssetClass), InTags, InChunkIDs, InPackageFlags)
-{
-}
-
-FAssetData::FAssetData(const FString& InLongPackageName, const FString& InObjectPath, FName InAssetClass, FAssetDataTagMap InTags, TArrayView<const int32> InChunkIDs, uint32 InPackageFlags)
-	: FAssetData(InLongPackageName, InObjectPath, FAssetData::TryConvertShortClassNameToPathName(InAssetClass), InTags, InChunkIDs, InPackageFlags)
-{
-}
-
 FAssetData::FAssetData(FName InPackageName, FName InPackagePath, FName InAssetName, FTopLevelAssetPath InAssetClassPathName, FAssetDataTagMap InTags, TArrayView<const int32> InChunkIDs, uint32 InPackageFlags)
 	: PackageName(InPackageName)
 	, PackagePath(InPackagePath)
diff --git a/Engine/Source/Runtime/CoreUObject/Private/Serialization/EditorBulkData.cpp b/Engine/Source/Runtime/CoreUObject/Private/Serialization/EditorBulkData.cpp
--- a/Engine/Source/Runtime/CoreUObject/Private/Serialization/EditorBulkData.cpp
+++ b/Engine/Source/Runtime/CoreUObject/Private/Serialization/EditorBulkData.cpp
@@ -2183,15 +2183,6 @@ void FEditorBulkData::SetCompressionOptions(ECompressedBufferCompressor Compress
 	}
 }
 
-FCustomVersionContainer FEditorBulkData::GetCustomVersions(FArchive& InlineArchive)
-{
-	FPackageFileVersion OutUEVersion;
-	int32 OutLicenseeUEVersion;
-	FCustomVersionContainer OutCustomVersions;
-	GetBulkDataVersions(InlineArchive, OutUEVersion, OutLicenseeUEVersion, OutCustomVersions);
-	return OutCustomVersions;
-}
-
 void FEditorBulkData::GetBulkDataVersions(FArchive& InlineArchive, FPackageFileVersion& OutUEVersion,
 	int32& OutLicenseeUEVersion, FCustomVersionContainer& OutCustomVersions) const
 {
diff --git a/Engine/Source/Runtime/CoreUObject/Private/UObject/PackageTrailer.cpp b/Engine/Source/Runtime/CoreUObject/Private/UObject/PackageTrailer.cpp
--- a/Engine/Source/Runtime/CoreUObject/Private/UObject/PackageTrailer.cpp
+++ b/Engine/Source/Runtime/CoreUObject/Private/UObject/PackageTrailer.cpp
@@ -232,11 +232,6 @@ TUniquePtr<UE::FPackageTrailerBuilder> FPackageTrailerBuilder::CreateReferenceTo
 	return Builder;
 }
 
-FPackageTrailerBuilder::FPackageTrailerBuilder(const FName& PackageName)
-	: DebugContext(PackageName.ToString())
-{
-}
-
 FPackageTrailerBuilder::FPackageTrailerBuilder(FString&& InDebugContext)
 	: DebugContext(MoveTemp(InDebugContext))
 {
diff --git a/Engine/Source/Runtime/CoreUObject/Private/UObject/ScriptCore.cpp b/Engine/Source/Runtime/CoreUObject/Private/UObject/ScriptCore.cpp
--- a/Engine/Source/Runtime/CoreUObject/Private/UObject/ScriptCore.cpp
+++ b/Engine/Source/Runtime/CoreUObject/Private/UObject/ScriptCore.cpp
@@ -605,13 +605,6 @@ void FFrame::GetScriptCallstack(FStringBuilderBase& ScriptStack, bool bReturnEmp
 #endif
 }
 
-FString FFrame::GetStackDescription() const
-{
-	TStringBuilder<256> StringBuilder;
-	GetStackDescription(StringBuilder);
-	return FString(StringBuilder);
-}
-
 void FFrame::GetStackDescription(FStringBuilderBase& StringBuilder) const
 {
 	Node->GetOuter()->GetPathName(nullptr, StringBuilder);
diff --git a/Engine/Source/Runtime/CoreUObject/Public/AssetRegistry/ARFilter.h b/Engine/Source/Runtime/CoreUObject/Public/AssetRegistry/ARFilter.h
--- a/Engine/Source/Runtime/CoreUObject/Public/AssetRegistry/ARFilter.h
+++ b/Engine/Source/Runtime/CoreUObject/Public/AssetRegistry/ARFilter.h
@@ -191,13 +191,11 @@ struct TStructOpsTypeTraits<FARFilter> : public TStructOpsTypeTraitsBase2<FARFil
  */
 struct FARCompiledFilter
 {
-PRAGMA_DISABLE_DEPRECATION_WARNINGS // Compilers can complain about deprecated members in compiler generated code
 	FARCompiledFilter() = default;
 	FARCompiledFilter(FARCompiledFilter&&) = default;
 	FARCompiledFilter(const FARCompiledFilter&) = default;
 	FARCompiledFilter& operator=(FARCompiledFilter&&) = default;
 	FARCompiledFilter& operator=(const FARCompiledFilter&) = default;
-PRAGMA_ENABLE_DEPRECATION_WARNINGS
 
 	/** The filter component for package names */
 	TSet<FName> PackageNames;
@@ -205,17 +203,9 @@ PRAGMA_ENABLE_DEPRECATION_WARNINGS
 	/** The filter component for package paths */
 	TSet<FName> PackagePaths;
 
-	/** The filter component containing specific object paths */
-	UE_DEPRECATED(5.1, "Object path FNames have been deprecated, use FSoftObjectPath instead.")
-	TSet<FName> ObjectPaths;
-
 	/** The filter component containing specific object paths */
 	TSet<FSoftObjectPath> SoftObjectPaths;
 
-	/** Deprecated. The filter component for class names. Instances of the specified classes, but not subclasses (by default), will be included. Derived classes will be included only if bRecursiveClasses is true. */
-	UE_DEPRECATED(5.1, "Class names are now represented by path names. Please use ClassPaths.")
-	TSet<FName> ClassNames;
-
 	/** The filter component for class names. Instances of the specified classes, but not subclasses (by default), will be included. Derived classes will be included only if bRecursiveClasses is true. */
 	TSet<FTopLevelAssetPath> ClassPaths;
 
diff --git a/Engine/Source/Runtime/CoreUObject/Public/AssetRegistry/AssetBundleData.h b/Engine/Source/Runtime/CoreUObject/Public/AssetRegistry/AssetBundleData.h
--- a/Engine/Source/Runtime/CoreUObject/Public/AssetRegistry/AssetBundleData.h
+++ b/Engine/Source/Runtime/CoreUObject/Public/AssetRegistry/AssetBundleData.h
@@ -108,25 +108,6 @@ public:
 	/** Returns pointer to an entry with given Scope/Name */
 	COREUOBJECT_API FAssetBundleEntry* FindEntry(FName SearchName);
 
-	/** Adds or updates an entry with the given BundleName -> Path. Scope is empty and will be filled in later */
-	UE_DEPRECATED(5.1, "The contents of an asset bundle now use FTopLevelAssetPath instead of FSoftObjectPath.")
-	COREUOBJECT_API void AddBundleAsset(FName BundleName, const FSoftObjectPath& AssetPath);
-
-	template< typename T >
-	UE_DEPRECATED(5.1, "The contents of an asset bundle now use FTopLevelAssetPath instead of FSoftObjectPath.")
-	void AddBundleAsset(FName BundleName, const TSoftObjectPtr<T>& SoftObjectPtr)
-	{
-		AddBundleAsset(BundleName, SoftObjectPtr.ToSoftObjectPath().GetAssetPath());
-	}
-
-	/** Adds multiple assets at once */
-	UE_DEPRECATED(5.1, "The contents of an asset bundle now use FTopLevelAssetPath instead of FSoftObjectPath.")
-	COREUOBJECT_API void AddBundleAssets(FName BundleName, const TArray<FSoftObjectPath>& AssetPaths);
-
-	/** A fast set of asset bundle assets, will destroy copied in path list */
-	UE_DEPRECATED(5.1, "The contents of an asset bundle now use FTopLevelAssetPath instead of FSoftObjectPath.")
-	COREUOBJECT_API void SetBundleAssets(FName BundleName, TArray<FSoftObjectPath>&& AssetPaths);
-
 	/** Adds or updates an entry with the given BundleName -> Path.*/
 	COREUOBJECT_API void AddBundleAsset(FName BundleName, const FTopLevelAssetPath& AssetPath);
 
diff --git a/Engine/Source/Runtime/CoreUObject/Public/AssetRegistry/AssetData.h b/Engine/Source/Runtime/CoreUObject/Public/AssetRegistry/AssetData.h
--- a/Engine/Source/Runtime/CoreUObject/Public/AssetRegistry/AssetData.h
+++ b/Engine/Source/Runtime/CoreUObject/Public/AssetRegistry/AssetData.h
@@ -257,11 +257,6 @@ PRAGMA_DISABLE_DEPRECATION_WARNINGS // Compilers can complain about deprecated m
 	FAssetData& operator=(const FAssetData&) = default;
 PRAGMA_ENABLE_DEPRECATION_WARNINGS
 
-	UE_DEPRECATED(5.1, "Class names are now represented by path names. Please use a version of FAssetData constructor that uses FTopLevelAssetPath.")
-	COREUOBJECT_API FAssetData(FName InPackageName, FName InPackagePath, FName InAssetName, FName InAssetClassName, FAssetDataTagMap InTags = FAssetDataTagMap(), TArrayView<const int32> InChunkIDs = TArrayView<const int32>(), uint32 InPackageFlags = 0);
-	UE_DEPRECATED(5.1, "Class names are now represented by path names. Please use a version of FAssetData constructor that uses FTopLevelAssetPath.")
-	COREUOBJECT_API FAssetData(const FString& InLongPackageName, const FString& InObjectPath, FName InAssetClassName, FAssetDataTagMap InTags = FAssetDataTagMap(), TArrayView<const int32> InChunkIDs = TArrayView<const int32>(), uint32 InPackageFlags = 0);
-
 	/** Constructor building the ObjectPath in the form of InPackageName.InAssetName. does not work for object outer-ed to a different package. */
 	COREUOBJECT_API FAssetData(FName InPackageName, FName InPackagePath, FName InAssetName, FTopLevelAssetPath InAssetClassPathName, FAssetDataTagMap InTags = FAssetDataTagMap(), TArrayView<const int32> InChunkIDs = TArrayView<const int32>(), uint32 InPackageFlags = 0);
 	/** Constructor with a long package name and a full object path which might not be part of the package this asset is in. */
@@ -1088,7 +1083,6 @@ private:
  * Helper struct for FAssetIdentifier (e.g., for the FOnViewAssetIdentifiersInReferenceViewer delegate and Reference Viewer functions).
  */
 #if WITH_EDITORONLY_DATA
-PRAGMA_DISABLE_DEPRECATION_WARNINGS
 struct FReferenceViewerParams
 {
 	FReferenceViewerParams()
@@ -1103,7 +1097,6 @@ struct FReferenceViewerParams
 		// Checkbox options
 		, bShowShowReferencesOptions(true)
 		, bShowShowSearchableNames(true)
-		, bShowShowNativePackages(true)
 		, bShowShowCodePackages(true)
 		, bShowShowFilteredPackagesOnly(true)
 		, bShowCompactMode(true)
@@ -1137,9 +1130,6 @@ struct FReferenceViewerParams
 	bool bShowShowReferencesOptions;
 	/** Whether to visually show to the user the option of "Show Searchable Names" */
 	bool bShowShowSearchableNames;
-	/** Whether to visually show to the user the option of "Show Native Packages" */
-	UE_DEPRECATED(5.1, "bShowShowNativePackages is deprecated, please use bShowShowCodePackages instead.")
-	bool bShowShowNativePackages;
 	/** Whether to visually show to the user the option of "Show C++ Packages" */
 	bool bShowShowCodePackages;
 	/** Whether to visually show to the user the option of "Show Filtered Packages Only" */
@@ -1147,5 +1137,4 @@ struct FReferenceViewerParams
 	/** Whether to visually show to the user the option of "Compact Mode" */
 	bool bShowCompactMode;
 };
-PRAGMA_ENABLE_DEPRECATION_WARNINGS
 #endif // WITH_EDITORONLY_DATA
diff --git a/Engine/Source/Runtime/CoreUObject/Public/AssetRegistry/AssetDataTagMap.h b/Engine/Source/Runtime/CoreUObject/Public/AssetRegistry/AssetDataTagMap.h
--- a/Engine/Source/Runtime/CoreUObject/Public/AssetRegistry/AssetDataTagMap.h
+++ b/Engine/Source/Runtime/CoreUObject/Public/AssetRegistry/AssetDataTagMap.h
@@ -25,20 +25,16 @@ namespace FixedTagPrivate { class FStoreBuilder; }
  */
 struct FAssetRegistryExportPath
 {
-PRAGMA_DISABLE_DEPRECATION_WARNINGS // Compilers can complain about deprecated members in compiler generated code
 	FAssetRegistryExportPath() = default;
 	FAssetRegistryExportPath(FAssetRegistryExportPath&&) = default;
 	FAssetRegistryExportPath(const FAssetRegistryExportPath&) = default;
 	FAssetRegistryExportPath& operator=(FAssetRegistryExportPath&&) = default;
 	FAssetRegistryExportPath& operator=(const FAssetRegistryExportPath&) = default;
-COREUOBJECT_API PRAGMA_ENABLE_DEPRECATION_WARNINGS
 
 	explicit FAssetRegistryExportPath(FWideStringView String);
 	COREUOBJECT_API explicit FAssetRegistryExportPath(FAnsiStringView String);
 
 	FTopLevelAssetPath ClassPath;
-	UE_DEPRECATED(5.1, "Class names are now represented by path names. Please use ClassPath member")
-	FName Class;
 	FName Package;
 	FName Object;
 
diff --git a/Engine/Source/Runtime/CoreUObject/Public/Serialization/EditorBulkData.h b/Engine/Source/Runtime/CoreUObject/Public/Serialization/EditorBulkData.h
--- a/Engine/Source/Runtime/CoreUObject/Public/Serialization/EditorBulkData.h
+++ b/Engine/Source/Runtime/CoreUObject/Public/Serialization/EditorBulkData.h
@@ -307,9 +307,6 @@ public:
 	 */
 	COREUOBJECT_API void SetCompressionOptions(ECompressedBufferCompressor Compressor, ECompressedBufferCompressionLevel CompressionLevel);
 
-	UE_DEPRECATED(5.1, "Call GetBulkDataVersions instead.")
-	COREUOBJECT_API FCustomVersionContainer GetCustomVersions(FArchive& InlineArchive);
-
 	/**
 	 * Get the versions used in the file containing the payload.
 	 *
diff --git a/Engine/Source/Runtime/CoreUObject/Public/UObject/Class.h b/Engine/Source/Runtime/CoreUObject/Public/UObject/Class.h
--- a/Engine/Source/Runtime/CoreUObject/Public/UObject/Class.h
+++ b/Engine/Source/Runtime/CoreUObject/Public/UObject/Class.h
@@ -3008,13 +3008,6 @@ public:
 	/** Checks if enum has entry with given name. Includes autogenerated _MAX entry. */
 	COREUOBJECT_API bool IsValidEnumName(FName InName) const;
 
-	/** Removes the Names in this enum from the primary AllEnumNames list */
-	UE_DEPRECATED(5.1, "RemoveNamesFromMasterList is deprecated, please use RemoveNamesFromPrimaryList instead.")
-	void RemoveNamesFromMasterList()
-	{
-		RemoveNamesFromPrimaryList();
-	}
-
 	/** Removes the Names in this enum from the primary AllEnumNames list */
 	COREUOBJECT_API void RemoveNamesFromPrimaryList();
 
@@ -3076,15 +3069,6 @@ public:
 	 */
 	COREUOBJECT_API static int64 LookupEnumName(FName PackageName, FName TestName, EFindFirstObjectOptions Options = EFindFirstObjectOptions::None, UEnum** OutFoundEnum = nullptr);
 
-	/** searches the list of all enum value names for the specified name
-	 * @return the value the specified name represents if found, otherwise INDEX_NONE
-	 */
-	UE_DEPRECATED(5.1, "LookupEnumName that takes only enum name is deprecated. Please use the version of this function that also takes enum package name.")
-	static int64 LookupEnumName(FName TestName, UEnum** FoundEnum = nullptr)
-	{
-		return LookupEnumName(FName(), TestName, EFindFirstObjectOptions::None, FoundEnum);
-	}
-
 	/** 
 	 * Searches the list of all enum value names for the specified name
 	 * @param PackageName Package where the enum was defined (/Script/ModuleName or uasset package name in case of script defined enums). If an empty (None) name is specified the function will look through all native enums to find a match
@@ -3095,15 +3079,6 @@ public:
 	 */
 	COREUOBJECT_API static int64 LookupEnumNameSlow(FName PackageName, const TCHAR* InTestShortName, EFindFirstObjectOptions Options = EFindFirstObjectOptions::None, UEnum** OutFoundEnum = nullptr);
 
-	/** searches the list of all enum value names for the specified name
-	 * @return the value the specified name represents if found, otherwise INDEX_NONE
-	 */
-	UE_DEPRECATED(5.1, "LookupEnumNameSlow that takes only enum name is deprecated. Please use the version of this function that also takes enum package name.")
-	static int64 LookupEnumNameSlow(const TCHAR* InTestShortName, UEnum** FoundEnum = nullptr)
-	{
-		return LookupEnumNameSlow(FName(), InTestShortName, EFindFirstObjectOptions::None, FoundEnum);
-	}
-
 	/** parses the passed in string for a name, then searches for that name in any Enum (in any package)
 	 * @param Str	pointer to string to parse; if we successfully find an enum, this pointer is advanced past the name found
 	 * @return index of the value the parsed enum name matches, or INDEX_NONE if no matches
@@ -3617,13 +3592,6 @@ protected:
 	/** global list of all value names used by all enums in memory, used for property text import */
 	static TMap<FName, TMap<FName, UEnum*> > AllEnumNames;
 
-	/** adds the Names in this enum to the primary AllEnumNames list */
-	UE_DEPRECATED(5.1, "AddNamesToMasterList is deprecated, please use AddNamesToPrimaryList instead.")
-	void AddNamesToMasterList()
-	{
-		AddNamesToPrimaryList();
-	}
-
 	/** adds the Names in this enum to the primary AllEnumNames list */
 	COREUOBJECT_API void AddNamesToPrimaryList();
 
diff --git a/Engine/Source/Runtime/CoreUObject/Public/UObject/Object.h b/Engine/Source/Runtime/CoreUObject/Public/UObject/Object.h
--- a/Engine/Source/Runtime/CoreUObject/Public/UObject/Object.h
+++ b/Engine/Source/Runtime/CoreUObject/Public/UObject/Object.h
@@ -250,15 +250,6 @@ public:
 	virtual void PostCDOCompiled(const FPostCDOCompiledContext& Context)
 	{
 	}
-
-	/**
-	 * Called after the Blueprint compiler has finished generating the Class Default Object (CDO) for a class. This can only happen in the editor.
-	 * @note This version is deprecated as it wasn't called for skeleton-only compilation. Use the version taking FPostCDOCompiledContext instead.
-	 */
-	UE_DEPRECATED(5.1, "Use version that takes FPostCDOCompiledContext instead.")
-	virtual void PostCDOCompiled()
-	{
-	}
 #endif
 
 	/**
diff --git a/Engine/Source/Runtime/CoreUObject/Public/UObject/Package.h b/Engine/Source/Runtime/CoreUObject/Public/UObject/Package.h
--- a/Engine/Source/Runtime/CoreUObject/Public/UObject/Package.h
+++ b/Engine/Source/Runtime/CoreUObject/Public/UObject/Package.h
@@ -134,9 +134,6 @@ struct FSavePackageResultStruct
 	/** Total size of all files written out, including bulk data */
 	int64 TotalFileSize;
 
-	UE_DEPRECATED(5.1, "CookedHash is now available through PackageWriter->CommitPackage instead. For waiting on completion in the non-cook case, use UPackage::WaitForAsyncFileWrites.")
-	TFuture<FMD5Hash> CookedHash;
-
 	/** Serialized package flags */
 	uint32 SerializedPackageFlags;
 
diff --git a/Engine/Source/Runtime/CoreUObject/Public/UObject/PackageTrailer.h b/Engine/Source/Runtime/CoreUObject/Public/UObject/PackageTrailer.h
--- a/Engine/Source/Runtime/CoreUObject/Public/UObject/PackageTrailer.h
+++ b/Engine/Source/Runtime/CoreUObject/Public/UObject/PackageTrailer.h
@@ -212,12 +212,6 @@ public:
 	 */
 	[[nodiscard]] static COREUOBJECT_API FPackageTrailerBuilder CreateFromTrailer(const class FPackageTrailer& Trailer, FArchive& Ar, FString DebugContext);
 
-	UE_DEPRECATED(5.1, "Use the overload that takes a FString instead of an FName for the last parameter")
-	static FPackageTrailerBuilder CreateFromTrailer(const class FPackageTrailer& Trailer, FArchive& Ar, const FName& PackageName)
-	{
-		return FPackageTrailerBuilder::CreateFromTrailer(Trailer, Ar, PackageName.ToString());
-	}
-
 	/**
 	 * Creates a builder from a pre-existing FPackageTrailer that will will reference the local payloads of the
 	 * source trailer. 
@@ -228,15 +222,7 @@ public:
 	 */
 	[[nodiscard]] static COREUOBJECT_API TUniquePtr<UE::FPackageTrailerBuilder> CreateReferenceToTrailer(const class FPackageTrailer& Trailer, FString DebugContext);
 	
-	UE_DEPRECATED(5.1, "Use the overload that takes a FString instead of an FName for the last parameter")
-	static TUniquePtr<UE::FPackageTrailerBuilder> CreateReferenceToTrailer(const class FPackageTrailer& Trailer, const FName& PackageName)
-	{
-		return FPackageTrailerBuilder::CreateReferenceToTrailer(Trailer, PackageName.ToString());
-	}
-
 	FPackageTrailerBuilder() = default;
-	UE_DEPRECATED(5.1, "Use the overload that takes a FString instead of an FName")
-	COREUOBJECT_API FPackageTrailerBuilder(const FName& InPackageName);
 	COREUOBJECT_API FPackageTrailerBuilder(FString&& DebugContext);
 	~FPackageTrailerBuilder() = default;
 
diff --git a/Engine/Source/Runtime/CoreUObject/Public/UObject/Script.h b/Engine/Source/Runtime/CoreUObject/Public/UObject/Script.h
--- a/Engine/Source/Runtime/CoreUObject/Public/UObject/Script.h
+++ b/Engine/Source/Runtime/CoreUObject/Public/UObject/Script.h
@@ -499,13 +499,6 @@ struct COREUOBJECT_API FBlueprintContextTracker : TThreadSingleton<FBlueprintCon
 		return ScriptEntryTag;
 	}
 
-	/** Returns current script stack frame */
-	UE_DEPRECATED(5.1, "GetScriptStack() inefficiently copies the array to return and is now deprecated. Use GetCurrentScriptStack() which returns a TArrayView instead")
-	UE_FORCEINLINE_HINT TArray<const FFrame*> GetScriptStack() const
-	{
-		return TArray<const FFrame*>(GetCurrentScriptStack());
-	}
-
 	/** Returns current script stack frame */
 	UE_FORCEINLINE_HINT TArrayView<const FFrame* const> GetCurrentScriptStack() const
 	{
diff --git a/Engine/Source/Runtime/CoreUObject/Public/UObject/Stack.h b/Engine/Source/Runtime/CoreUObject/Public/UObject/Stack.h
--- a/Engine/Source/Runtime/CoreUObject/Public/UObject/Stack.h
+++ b/Engine/Source/Runtime/CoreUObject/Public/UObject/Stack.h
@@ -278,12 +278,6 @@ public:
 	**/
 	COREUOBJECT_API static void GetScriptCallstack(FStringBuilderBase& StringBuilder, bool bReturnEmpty = false, bool bTopOfStackOnly = false);
 		
-	/** 
-	 * This will return a string of the form "ScopeName.FunctionName" associated with this stack frame:
-	 */
-	UE_DEPRECATED(5.1, "Please use GetStackDescription(FStringBuilderBase&).")
-	COREUOBJECT_API FString GetStackDescription() const;
-
 	/**
 	* This will append a string of the form "ScopeName.FunctionName" associated with this stack frame
 	*
diff --git a/Engine/Source/Runtime/CoreUObject/Public/UObject/TopLevelAssetPath.h b/Engine/Source/Runtime/CoreUObject/Public/UObject/TopLevelAssetPath.h
--- a/Engine/Source/Runtime/CoreUObject/Public/UObject/TopLevelAssetPath.h
+++ b/Engine/Source/Runtime/CoreUObject/Public/UObject/TopLevelAssetPath.h
@@ -105,12 +105,6 @@ public:
 	/** Copy the full asset path (e.g. '/Path/To/Package.AssetName') into the provided string. */
 	COREUOBJECT_API void ToString(FString& OutString) const;
 
-	// Return the full asset path (e.g. '/Path/To/Package.AssetName') as an FName.
-	UE_DEPRECATED(5.1, "FNames containing full asset paths have been replaced by FTopLevelAssetPath/FSoftLevelObjectPath."
-		"This function is only for temporary use interfacing with APIs that still expect an FName."
-		"Those APIS should be updated to use FTopLevelAssetPath or FSoftLevelObjectPath.")
-	FName ToFName() const { return *ToString(); }
-
 	/** Check if this could possibly refer to a real object */
 	bool IsValid() const
 	{
