LinkerLoad: dll export FLinkerLoad::ResolveResource(FPackageIndex ResourceIndex) so that it can be called form the new serialization stack plugins.

#rnx

#Synced-CL 49987382

diff --git a/Engine/Plugins/Experimental/PlainPropsUObject/Source/PlainPropsUObject.Build.cs b/Engine/Plugins/Experimental/PlainPropsUObject/Source/PlainPropsUObject.Build.cs
--- a/Engine/Plugins/Experimental/PlainPropsUObject/Source/PlainPropsUObject.Build.cs
+++ b/Engine/Plugins/Experimental/PlainPropsUObject/Source/PlainPropsUObject.Build.cs
@@ -7,6 +7,7 @@ public class PlainPropsUObject : ModuleRules
 	public PlainPropsUObject(ReadOnlyTargetRules Target) : base(Target)
 	{
 		bDisableStaticAnalysis = true;
+		bValidateInternalApi = false; // for UE_INTERNAL FLinkerLoad::ResolveResource access
 
 		CppStandard = CppStandardVersion.Cpp20;
 
diff --git a/Engine/Source/Runtime/CoreUObject/Public/UObject/LinkerLoad.h b/Engine/Source/Runtime/CoreUObject/Public/UObject/LinkerLoad.h
--- a/Engine/Source/Runtime/CoreUObject/Public/UObject/LinkerLoad.h
+++ b/Engine/Source/Runtime/CoreUObject/Public/UObject/LinkerLoad.h
@@ -706,7 +706,7 @@ public:
 	 *
 	 * @return	The UObject at ResourceIndex, or nullptr if not found
 	 */
-	UObject* ResolveResource(FPackageIndex ResourceIndex);
+	UE_INTERNAL COREUOBJECT_API UObject* ResolveResource(FPackageIndex ResourceIndex);
 	
 	UE_DEPRECATED(5.6, "This function was used only for the now-deleted Conform functionality and will be removed. Contact Epic if you need this function.")
 	int32 FindExportIndex( FName ClassName, FName ClassPackage, FName ObjectName, FPackageIndex ExportOuterIndex );
