Core: Cleanup of 5.2 deprecations in Core and CoreUObject.
A strict code removal operation.
Remove the missed 5.0 deprecation FSavePackageContext::bForceLegacyOffsets.
Remove the missed 5.0 deprecation FStringTableRegistry::FindTableIdAndKey.
Skip removal of FNetworkVersion functions that are still being used in ReplayTypes.cpp.

#rnx
#tests BuildCookRun Lyra

#Synced-CL 49987382

diff --git a/Engine/Source/Runtime/Core/Private/Internationalization/StringTableRegistry.cpp b/Engine/Source/Runtime/Core/Private/Internationalization/StringTableRegistry.cpp
--- a/Engine/Source/Runtime/Core/Private/Internationalization/StringTableRegistry.cpp
+++ b/Engine/Source/Runtime/Core/Private/Internationalization/StringTableRegistry.cpp
@@ -138,11 +138,6 @@ void FStringTableRegistry::EnumerateStringTables(const TFunctionRef<bool(const F
 	}
 }
 
-bool FStringTableRegistry::FindTableIdAndKey(const FText& InText, FName& OutTableId, FString& OutKey) const
-{
-	return FTextInspector::GetTableIdAndKey(InText, OutTableId, OutKey);
-}
-
 void FStringTableRegistry::LogMissingStringTableEntry(const FName InTableId, const FTextKey& InKey)
 {
 	UE::TScopeLock LoggedMissingEntriesLock(LoggedMissingEntriesCS);
diff --git a/Engine/Source/Runtime/Core/Private/Windows/WindowsPlatformMisc.cpp b/Engine/Source/Runtime/Core/Private/Windows/WindowsPlatformMisc.cpp
--- a/Engine/Source/Runtime/Core/Private/Windows/WindowsPlatformMisc.cpp
+++ b/Engine/Source/Runtime/Core/Private/Windows/WindowsPlatformMisc.cpp
@@ -3026,16 +3026,6 @@ public:
 		return CPUIDStaticCache.CPUInfo2;
 	}
 
-	/**
-	 * Gets cache line size.
-	 *
-	 * @returns Cache line size.
-	 */
-	static int32 GetCacheLineSize()
-	{
-		return CPUIDStaticCache.CacheLineSize;
-	}
-
 private:
 
 #if !PLATFORM_CPU_ARM_FAMILY
@@ -4116,11 +4106,6 @@ bool FWindowsPlatformMisc::HasTimedPauseCPUFeature()
 	return FCPUIDQueriedData::HasTimedPauseInstruction();
 }
 
-int32 FWindowsPlatformMisc::GetCacheLineSize()
-{
-	return FCPUIDQueriedData::GetCacheLineSize();
-}
-
 bool FWindowsPlatformMisc::QueryRegKey( const Windows::HKEY InKey, const TCHAR* InSubKey, const TCHAR* InValueName, FString& OutData )
 {
 	bool bSuccess = false;
diff --git a/Engine/Source/Runtime/Core/Public/Containers/DepletableMpmcQueue.h b/Engine/Source/Runtime/Core/Public/Containers/DepletableMpmcQueue.h
deleted file mode 100644
diff --git a/Engine/Source/Runtime/Core/Public/GenericPlatform/GenericPlatformMisc.h b/Engine/Source/Runtime/Core/Public/GenericPlatform/GenericPlatformMisc.h
--- a/Engine/Source/Runtime/Core/Public/GenericPlatform/GenericPlatformMisc.h
+++ b/Engine/Source/Runtime/Core/Public/GenericPlatform/GenericPlatformMisc.h
@@ -1546,12 +1546,6 @@ public:
 		Prefetch(reinterpret_cast<const void*>(reinterpret_cast<UPTRINT>(Ptr) + Offset));
 	}
 
-	UE_DEPRECATED(5.2, "Must supply size when prefetching a block of data")
-	FORCEINLINE static void PrefetchBlock(const void* Ptr)
-	{
-		Prefetch(Ptr);
-	}
-
 	FORCEINLINE static void PrefetchBlock(const void* Ptr, int32 NumBytes)
 	{
 		// Use compile-time PLATFORM_CACHE_LINE_SIZE to avoid having to load a value before issuing prefetch instructions
diff --git a/Engine/Source/Runtime/Core/Public/Internationalization/StringTableRegistry.h b/Engine/Source/Runtime/Core/Public/Internationalization/StringTableRegistry.h
--- a/Engine/Source/Runtime/Core/Public/Internationalization/StringTableRegistry.h
+++ b/Engine/Source/Runtime/Core/Public/Internationalization/StringTableRegistry.h
@@ -45,10 +45,6 @@ public:
 	/** Enumerate all registered string tables */
 	CORE_API void EnumerateStringTables(const TFunctionRef<bool(const FName&, const FStringTableConstRef&)>& InEnumerator) const;
 
-	/** Given an FText, try and find the table ID and key for it */
-	UE_DEPRECATED(5.0, "FindTableIdAndKey is deprecated. Use FTextInspector::GetTableIdAndKey instead.")
-	CORE_API bool FindTableIdAndKey(const FText& InText, FName& OutTableId, FString& OutKey) const;
-
 	/** Log a missing string table entry (will only log each missing entry once to avoid spam) */
 	CORE_API void LogMissingStringTableEntry(const FName InTableId, const FTextKey& InKey);
 
diff --git a/Engine/Source/Runtime/Core/Public/Misc/LargeWorldRenderPosition.h b/Engine/Source/Runtime/Core/Public/Misc/LargeWorldRenderPosition.h
--- a/Engine/Source/Runtime/Core/Public/Misc/LargeWorldRenderPosition.h
+++ b/Engine/Source/Runtime/Core/Public/Misc/LargeWorldRenderPosition.h
@@ -59,10 +59,6 @@ public:
 	TResult GetTile() const { return static_cast<TResult>(Tile); }
 	template <typename TResult = float>
 	TResult GetOffset() const { return static_cast<TResult>(Offset); }
-	UE_DEPRECATED(5.2, "GetTileAsDouble is deprecated, please use GetTile<double>() instead.")
-	double GetTileAsDouble() const { return GetTile<double>(); }
-	UE_DEPRECATED(5.2, "GetOffsetAsDouble is deprecated, please use GetOffset<double>() instead.")
-	double GetOffsetAsDouble() const { return GetOffset<double>(); }
 	double GetTileOffset() const { return static_cast<double>(Tile) * GetTileSize(); }
 	double GetAbsolute() const { return GetTileOffset() + Offset; }
 
diff --git a/Engine/Source/Runtime/Core/Public/Templates/AreTypesEqual.h b/Engine/Source/Runtime/Core/Public/Templates/AreTypesEqual.h
deleted file mode 100644
diff --git a/Engine/Source/Runtime/Core/Public/Windows/WindowsPlatformMisc.h b/Engine/Source/Runtime/Core/Public/Windows/WindowsPlatformMisc.h
--- a/Engine/Source/Runtime/Core/Public/Windows/WindowsPlatformMisc.h
+++ b/Engine/Source/Runtime/Core/Public/Windows/WindowsPlatformMisc.h
@@ -315,8 +315,6 @@ struct FWindowsPlatformMisc
 	 */
 	static CORE_API bool GetVSComnTools(int32 Version, FString& OutData);
 
-	UE_DEPRECATED(5.2, "Please use PLATFORM_CACHE_LINE_SIZE instead, runtime query of cache line size not supported")
-	static CORE_API int32 GetCacheLineSize();
 	/**
 	* @return Windows path separator.
 	*/
diff --git a/Engine/Source/Runtime/CoreUObject/Private/UObject/SavePackage/SaveContext.h b/Engine/Source/Runtime/CoreUObject/Private/UObject/SavePackage/SaveContext.h
--- a/Engine/Source/Runtime/CoreUObject/Private/UObject/SavePackage/SaveContext.h
+++ b/Engine/Source/Runtime/CoreUObject/Private/UObject/SavePackage/SaveContext.h
@@ -1282,13 +1282,6 @@ public:
 		return PackageWriter;
 	}
 
-	PRAGMA_DISABLE_DEPRECATION_WARNINGS
-	ISavePackageValidator* GetPackageValidator() const
-	{
-		return SaveArgs.SavePackageContext ? SaveArgs.SavePackageContext->GetValidator() : nullptr;
-	}
-	PRAGMA_ENABLE_DEPRECATION_WARNINGS
-
 	bool HasExternalImportValidations() const
 	{
 		return SaveArgs.SavePackageContext ? SaveArgs.SavePackageContext->GetExternalImportValidations().Num() > 0 : false;
diff --git a/Engine/Source/Runtime/CoreUObject/Private/UObject/SavePackage2.cpp b/Engine/Source/Runtime/CoreUObject/Private/UObject/SavePackage2.cpp
--- a/Engine/Source/Runtime/CoreUObject/Private/UObject/SavePackage2.cpp
+++ b/Engine/Source/Runtime/CoreUObject/Private/UObject/SavePackage2.cpp
@@ -1088,17 +1088,6 @@ ESavePackageResult ValidateImports(FSaveContext& SaveContext)
 		}
 	}
 
-	PRAGMA_DISABLE_DEPRECATION_WARNINGS
-	if (ISavePackageValidator* Validator = SaveContext.GetPackageValidator())
-	{
-		ESavePackageResult ValidatorResult = Validator->ValidateImports(Package, Imports);
-		if (ValidatorResult != ESavePackageResult::Success)
-		{
-			return ValidatorResult;
-		}
-	}
-	PRAGMA_ENABLE_DEPRECATION_WARNINGS
-
 	return ReturnSuccessOrCancel();
 }
 
diff --git a/Engine/Source/Runtime/CoreUObject/Public/UObject/LinkerInstancingContext.h b/Engine/Source/Runtime/CoreUObject/Public/UObject/LinkerInstancingContext.h
--- a/Engine/Source/Runtime/CoreUObject/Public/UObject/LinkerInstancingContext.h
+++ b/Engine/Source/Runtime/CoreUObject/Public/UObject/LinkerInstancingContext.h
@@ -161,12 +161,6 @@ public:
 	COREUOBJECT_API void SetSoftObjectPathRemappingEnabled(bool bInSoftObjectPathRemappingEnabled);
 	COREUOBJECT_API bool GetSoftObjectPathRemappingEnabled() const;
 
-	UE_DEPRECATED(5.2, "No longer used, pass ELoadFlags::LOAD_RegenerateBulkDataGuids to LoadPackage instead")
-	void SetRegenerateUniqueBulkDataGuids(bool bFlag) { }
-
-	UE_DEPRECATED(5.2, "No longer used, check ELoadFlags::LOAD_RegenerateBulkDataGuids in the LoadFlags instead")
-	bool ShouldRegenerateUniqueBulkDataGuids() const { return false; }
-
 	/** Return the instanced package name for a given instanced outer package and an object package name */
 	static FString GetInstancedPackageName(const FString& InOuterPackageName, const FString& InPackageName)
 	{
@@ -192,4 +186,4 @@ private:
 
 	TSharedPtr<FSharedLinkerInstancingContextData> SharedData;
 	TSharedPtr<FSharedLinkerInstancingContextData> ImmutableData;
-};
\ No newline at end of file
+};
diff --git a/Engine/Source/Runtime/CoreUObject/Public/UObject/ObjectRef.h b/Engine/Source/Runtime/CoreUObject/Public/UObject/ObjectRef.h
--- a/Engine/Source/Runtime/CoreUObject/Public/UObject/ObjectRef.h
+++ b/Engine/Source/Runtime/CoreUObject/Public/UObject/ObjectRef.h
@@ -119,7 +119,4 @@ private:
 
 };
 
-UE_DEPRECATED(5.2, "Use FObjectRef::IsNull() instead")
-inline bool IsObjectRefNull(const FObjectRef& ObjectRef) { return ObjectRef.IsNull(); }
-
 #endif
diff --git a/Engine/Source/Runtime/CoreUObject/Public/UObject/SavePackage.h b/Engine/Source/Runtime/CoreUObject/Public/UObject/SavePackage.h
--- a/Engine/Source/Runtime/CoreUObject/Public/UObject/SavePackage.h
+++ b/Engine/Source/Runtime/CoreUObject/Public/UObject/SavePackage.h
@@ -129,17 +129,6 @@ struct FSavePackageArgs
 	}
 };
 
-/** Interface for SavePackage to test for caller-specific errors. */
-class UE_DEPRECATED(5.2, "Use a FSavePackageContext::ExternalValidationFunc if you need to run external validation") ISavePackageValidator;
-class ISavePackageValidator
-{
-public:
-	virtual ~ISavePackageValidator() = default;
-
-	virtual ESavePackageResult ValidateImports(const UPackage* Package, const TSet<TObjectPtr<UObject>>& Imports) = 0;
-};
-
-
 /** Param struct for external import validation functions */
 struct FImportsValidationContext
 {
@@ -235,25 +224,8 @@ public:
 		}
 	}
 
-	UE_DEPRECATED(5.0, "bInForceLegacyOffsets is no longer supported; remove the variable from your constructor call")
-	FSavePackageContext(const ITargetPlatform* InTargetPlatform, IPackageWriter* InPackageWriter, bool InbForceLegacyOffsets)
-		: FSavePackageContext(InTargetPlatform, InPackageWriter)
-	{
-	}
-
 	COREUOBJECT_API ~FSavePackageContext();
 
-	PRAGMA_DISABLE_DEPRECATION_WARNINGS
-	ISavePackageValidator* GetValidator()
-	{
-		return Validator.Get();	
-	}
-	void SetValidator(TUniquePtr<ISavePackageValidator>&& InValidator)
-	{
-		Validator = MoveTemp(InValidator);
-	}
-	PRAGMA_ENABLE_DEPRECATION_WARNINGS
-
 	const TArray<TFunction<FSavePackageSettings::ExternalImportValidationFunc>>& GetExternalImportValidations() const
 	{
 		return SavePackageSettings.GetExternalImportValidations();
@@ -269,14 +241,7 @@ public:
 	IPackageWriter::FCapabilities PackageWriterCapabilities;
 
 private:
-	PRAGMA_DISABLE_DEPRECATION_WARNINGS
-	TUniquePtr<ISavePackageValidator> Validator;
-	PRAGMA_ENABLE_DEPRECATION_WARNINGS
 	FSavePackageSettings SavePackageSettings;
-public:
-
-	UE_DEPRECATED(5.0, "bForceLegacyOffsets is no longer supported; remove uses of the variable")
-	const bool bForceLegacyOffsets = false;
 };
 
 namespace UE::SavePackageUtilities
