Add support for pausing and resuming garbage collection (e.g for pausing during initial startup)
Adds functions PauseGarbageCollection(), ResumeGarbageCollection() and IsGarbageCollectionPaused
GC will be automatically resumed if free memory drops lower than gc.PauseGCFreeMemThresholdMB or pause duration exceeds gc.PauseGCTimeLimitSeconds

diff --git a/Engine/Source/Runtime/CoreUObject/Private/UObject/GarbageCollection.cpp b/Engine/Source/Runtime/CoreUObject/Private/UObject/GarbageCollection.cpp
--- a/Engine/Source/Runtime/CoreUObject/Private/UObject/GarbageCollection.cpp
+++ b/Engine/Source/Runtime/CoreUObject/Private/UObject/GarbageCollection.cpp
@@ -11,6 +11,7 @@
 #include "Containers/ConsumeAllMpmcQueue.h"
 #include "Containers/Deque.h"
 #include "Containers/StaticBitArray.h"
+#include "Containers/Ticker.h"
 #include "HAL/MallocLeakDetection.h"
 #include "HAL/ThreadSafeBool.h"
 #include "Misc/TimeGuard.h"
@@ -37,6 +38,7 @@
 #include "UObject/Package.h"
 #include "Async/ParallelFor.h"
 #include "ProfilingDebugging/CsvProfiler.h"
+#include "HAL/PlatformMemoryHelpers.h"
 #include "HAL/Runnable.h"
 #include "HAL/RunnableThread.h"
 #include "HAL/ThreadHeartBeat.h"
@@ -242,6 +244,23 @@ static FAutoConsoleVariableRef CIncrementalBeginDestroyGranularity(
 	ECVF_Default
 );
 
+static float GPauseGCTimeLimitSeconds = 30.0;
+static FAutoConsoleVariableRef CPauseGCTimeLimitSeconds(
+	TEXT("gc.PauseGCTimeLimitSeconds"),
+	GPauseGCTimeLimitSeconds,
+	TEXT("If GC is paused, it will resume automatically after this duration"),
+	ECVF_Default
+);
+
+static int32 GPauseGCFreeMemThresholdMB = 200;
+static FAutoConsoleVariableRef CPauseGCFreeMemThresholdMB(
+	TEXT("gc.PauseGCFreeMemThresholdMB"),
+	GPauseGCFreeMemThresholdMB,
+	TEXT("If GC is paused, it will resume automatically if free physical memory is less than this threshold"),
+	ECVF_Default
+);
+
+
 #if WITH_VERSE_VM || defined(__INTELLISENSE__)
 bool GEnableFrankenGC = true;
 static FAutoConsoleVariableRef CEnableFrankenGC(
@@ -6306,6 +6325,145 @@ bool UnhashUnreachableObjects(bool bUseTimeLimit, double TimeLimit)
 	return bTimeLimitReached;
 }
 
+class FGarbageCollectionPause
+{
+public:
+	static FGarbageCollectionPause& Get()
+	{
+		static FGarbageCollectionPause Inst;
+		return Inst;
+	}
+	bool IsPaused() const
+	{
+		return StartTime > 0.0;
+	}
+
+	bool Resume(bool bCleanupDelegate = true, bool bAllowCatchupGC = true)
+	{
+		check(IsInGameThread());
+		if (!IsPaused())
+		{
+			// Nothing to do
+			return false;
+		}
+		UE_LOG(LogGarbage, Log, TEXT("Garbage collection resumed"));
+		StartTime = 0.0;
+		if (bCleanupDelegate)
+		{
+			FTSTicker::GetCoreTicker().RemoveTicker(TickHandle);
+		}
+
+		// If we missed any GCs while we were paused then request one now
+		if (bAllowCatchupGC && MissedGcCount>0)
+		{
+			TryCollectGarbage(GARBAGE_COLLECTION_KEEPFLAGS, bFullPurgeRequestedWhilePaused);
+		}
+		return true;
+	}
+
+	bool ShouldResume() const
+	{
+		if (GPauseGCTimeLimitSeconds > 0.0 && FPlatformTime::Seconds() - StartTime > GPauseGCTimeLimitSeconds)
+		{
+			UE_LOG(LogGarbage, Log, TEXT("Resuming GC due to time limit: %.2fs"), GPauseGCTimeLimitSeconds);
+			return true;
+		}
+		if (GPauseGCFreeMemThresholdMB > 0 && PlatformMemoryHelpers::GetFrameMemoryStats().AvailablePhysical < ((uint64)GPauseGCFreeMemThresholdMB << 20ull))
+		{
+			UE_LOG(LogGarbage, Log, TEXT("Resuming GC due to mem threshold: %dMB"), GPauseGCFreeMemThresholdMB);
+			return true;
+		}
+		return false;
+	}
+
+	void Pause()
+	{
+		check(IsInGameThread());
+		if (IsPaused())
+		{
+			// Nothing to do
+			return;
+		}
+		MissedGcCount=0;
+		bFullPurgeRequestedWhilePaused = false;
+		StartTime = FPlatformTime::Seconds();
+
+		// Check this every frame to ensure the cvars are runtime modifyable
+		TickHandle = FTSTicker::GetCoreTicker().AddTicker(FTickerDelegate::CreateLambda([this](float) 
+			{
+				if (ShouldResume())
+				{
+					Resume(false);
+					return false;
+				}
+				return true;
+			}), 0.0f);
+
+		UE_LOG(LogGarbage, Log, TEXT("Garbage collection paused"));
+		if (GPauseGCFreeMemThresholdMB > 0)
+		{
+			UE_LOG(LogGarbage, Log, TEXT("  Free mem threshold: %dMB"), GPauseGCFreeMemThresholdMB);
+		}
+		if (GPauseGCTimeLimitSeconds > 0.0)
+		{
+			UE_LOG(LogGarbage, Log, TEXT("  Time limit: %.2fs"), GPauseGCTimeLimitSeconds);
+		}
+	}
+
+	// Called from GC. Returns true if we're paused but attempts to resume first if possible
+	// Note: if KeepFlags is anything other than GARBAGE_COLLECTION_KEEPFLAGS, will resume immediately
+	bool OnGc(EObjectFlags KeepFlags, bool bFullPurge)
+	{
+		if (!IsPaused())
+		{
+			return false;
+		}
+		if (ShouldResume())
+		{
+			Resume(true, false);
+			return false;
+		}
+		if (KeepFlags != GARBAGE_COLLECTION_KEEPFLAGS)
+		{
+			UE_LOG(LogGarbage, Warning, TEXT("GC Resumed due to non-standard KeepFlags: 0x%08X"), (uint32)KeepFlags);
+			Resume(true, false);
+			return false;
+		}
+		MissedGcCount++;
+		bFullPurgeRequestedWhilePaused |= bFullPurge;
+		return true;
+	}
+
+	int32 GetMissedGcCount() const
+	{
+		return MissedGcCount;
+	}
+
+private:
+	FGarbageCollectionPause() {}
+
+	FTSTicker::FDelegateHandle TickHandle;
+	double StartTime = 0.0;
+	int32 MissedGcCount = 0;
+	bool bFullPurgeRequestedWhilePaused = false;
+};
+
+void PauseGarbageCollection()
+{
+	FGarbageCollectionPause::Get().Pause();
+}
+
+bool ResumeGarbageCollection()
+{
+	return FGarbageCollectionPause::Get().Resume();
+}
+
+bool IsGarbageCollectionPaused()
+{
+	return FGarbageCollectionPause::Get().IsPaused();
+}
+
+
 void CollectGarbage(EObjectFlags KeepFlags, bool bPerformFullPurge)
 {
 	if (GIsInitialLoad)
@@ -6315,6 +6473,12 @@ void CollectGarbage(EObjectFlags KeepFlags, bool bPerformFullPurge)
 		return;
 	}
 
+	if (FGarbageCollectionPause::Get().OnGc(KeepFlags, bPerformFullPurge))
+	{
+		UE_LOG(LogGarbage, Log, TEXT("GC is paused. Skipped CollectGarbage (Full purge: %s)"), bPerformFullPurge ? TEXT("Enabled") : TEXT("Disabled"));
+		return;
+	}
+
 	// No other thread may be performing UObject operations while we're running
 	AcquireGCLock();
 
@@ -6340,6 +6504,13 @@ bool TryCollectGarbage(EObjectFlags KeepFlags, bool bPerformFullPurge)
 		return false;
 	}
 
+	if (FGarbageCollectionPause::Get().OnGc(KeepFlags, bPerformFullPurge))
+	{
+		// GC was paused. Only log on the first skip to avoid logspam, since this is often called every frame
+		UE_CLOG(FGarbageCollectionPause::Get().GetMissedGcCount()==1, LogGarbage, Log, TEXT("GC is paused. Skipped TryCollectGarbage (Full purge: %s)"), bPerformFullPurge ? TEXT("Enabled") : TEXT("Disabled"));
+		return false;
+	}
+
 	// No other thread may be performing UObject operations while we're running so try to acquire GC lock
 	if (UE::GC::GIsIncrementalReachabilityPending)
 	{
diff --git a/Engine/Source/Runtime/CoreUObject/Public/UObject/UObjectGlobals.h b/Engine/Source/Runtime/CoreUObject/Public/UObject/UObjectGlobals.h
--- a/Engine/Source/Runtime/CoreUObject/Public/UObject/UObjectGlobals.h
+++ b/Engine/Source/Runtime/CoreUObject/Public/UObject/UObjectGlobals.h
@@ -943,6 +943,28 @@ COREUOBJECT_API void CollectGarbage(EObjectFlags KeepFlags, bool bPerformFullPur
 */
 COREUOBJECT_API bool TryCollectGarbage(EObjectFlags KeepFlags, bool bPerformFullPurge = true);
 
+/**
+ * Disables garbage collection until the time limit or free memory threshold are exceeded, or ResumeGarbageCollection is called
+ * Cvars gc.PauseGCTimeLimitSeconds and gc.PauseGCFreeMemThresholdMB can be used to specify time and memory thresholds for 
+ * auto-resume.
+ *
+ */
+COREUOBJECT_API void PauseGarbageCollection();
+
+/**
+ * Re-enables garbage collection following a call to PauseGarbageCollection
+ * 
+ * @return true if garbage collection was paused or false otherwise 
+ */
+COREUOBJECT_API bool ResumeGarbageCollection();
+
+/**
+ * Checks if GC is paused
+ *
+ * @return true if garbage collection is paused or false otherwise
+ */
+COREUOBJECT_API bool IsGarbageCollectionPaused();
+
 /**
 * Calls ConditionalBeginDestroy on unreachable objects
 *
