Core: Cleanup of 5.3 deprecations in Core and CoreUObject.
Mostly a strict code removal operation.
Fixup internal project usage of UObjecct::IsDataValid override.
Skip removal of FNetworkVersion functions since they are still used in ReplayTypes.cpp.

#rnx
#tests BuildCookRun Lyra, Editor boot

#Synced-CL 49987382

diff --git a/Engine/Source/Runtime/Core/Private/Misc/MessageDialog.cpp b/Engine/Source/Runtime/Core/Private/Misc/MessageDialog.cpp
--- a/Engine/Source/Runtime/Core/Private/Misc/MessageDialog.cpp
+++ b/Engine/Source/Runtime/Core/Private/Misc/MessageDialog.cpp
@@ -31,11 +31,6 @@ namespace
 	}
 }
 
-void FMessageDialog::Debugf( const FText& Message, const FText* OptTitle )
-{
-	Debugf(Message, OptTitle ? *OptTitle : GetDefaultMessageTitle());
-}
-
 void FMessageDialog::Debugf( const FText& Message )
 {
 	Debugf(Message, GetDefaultMessageTitle());
@@ -76,11 +71,6 @@ void FMessageDialog::ShowLastError()
 	}
 }
 
-EAppReturnType::Type FMessageDialog::Open( EAppMsgType::Type MessageType, const FText& Message, const FText* OptTitle )
-{
-	return Open(EAppMsgCategory::Warning, MessageType, Message, OptTitle ? *OptTitle : GetDefaultMessageTitle());
-}
-
 EAppReturnType::Type FMessageDialog::Open( EAppMsgType::Type MessageType, const FText& Message )
 {
 	return Open(EAppMsgCategory::Warning, MessageType, Message, GetDefaultMessageTitle());
@@ -145,11 +135,6 @@ EAppReturnType::Type FMessageDialog::Open( EAppMsgCategory MessageCategory, EApp
 	return Open(MessageCategory, MessageType, DefaultValue, Message, Title);
 }
 
-EAppReturnType::Type FMessageDialog::Open(EAppMsgType::Type MessageType, EAppReturnType::Type DefaultValue, const FText& Message, const FText* OptTitle)
-{
-	return Open(EAppMsgCategory::Warning, MessageType, DefaultValue, Message, OptTitle ? *OptTitle : GetDefaultMessageTitle());
-}
-
 EAppReturnType::Type FMessageDialog::Open(EAppMsgType::Type MessageType, EAppReturnType::Type DefaultValue, const FText& Message)
 {
 	return Open(EAppMsgCategory::Warning, MessageType, DefaultValue, Message, GetDefaultMessageTitle());
diff --git a/Engine/Source/Runtime/Core/Public/Misc/MessageDialog.h b/Engine/Source/Runtime/Core/Public/Misc/MessageDialog.h
--- a/Engine/Source/Runtime/Core/Public/Misc/MessageDialog.h
+++ b/Engine/Source/Runtime/Core/Public/Misc/MessageDialog.h
@@ -21,8 +21,6 @@ struct FMessageDialog
 	*/
 	static CORE_API void Debugf( const FText& Message );
 	static CORE_API void Debugf( const FText& Message, const FText& Title );
-	UE_DEPRECATED(5.3, "Use the overload of Debugf that takes the Title by-value.")
-	static CORE_API void Debugf( const FText& Message, const FText* OptTitle );
 
 	/** Pops up a message dialog box containing the last system error code in string form. */
 	static CORE_API void ShowLastError();
@@ -38,8 +36,6 @@ struct FMessageDialog
 	static CORE_API EAppReturnType::Type Open( EAppMsgType::Type MessageType, const FText& Message, const FText& Title);
 	static CORE_API EAppReturnType::Type Open( EAppMsgCategory MessageCategory, EAppMsgType::Type MessageType, const FText& Message);
 	static CORE_API EAppReturnType::Type Open( EAppMsgCategory MessageCategory, EAppMsgType::Type MessageType, const FText& Message, const FText& Title);
-	UE_DEPRECATED(5.3, "Use the overload of Open that takes the Title by-value.")
-	static CORE_API EAppReturnType::Type Open( EAppMsgType::Type MessageType, const FText& Message, const FText* OptTitle);
 
 	/**
 	 * Open a modal message box dialog
@@ -53,6 +49,4 @@ struct FMessageDialog
 	static CORE_API EAppReturnType::Type Open(EAppMsgType::Type MessageType, EAppReturnType::Type DefaultValue, const FText& Message, const FText& Title);
 	static CORE_API EAppReturnType::Type Open(EAppMsgCategory MessageCategory, EAppMsgType::Type MessageType, EAppReturnType::Type DefaultValue, const FText& Message);
 	static CORE_API EAppReturnType::Type Open(EAppMsgCategory MessageCategory, EAppMsgType::Type MessageType, EAppReturnType::Type DefaultValue, const FText& Message, const FText& Title);
-	UE_DEPRECATED(5.3, "Use the overload of Open that takes the Title by-value.")
-	static CORE_API EAppReturnType::Type Open(EAppMsgType::Type MessageType, EAppReturnType::Type DefaultValue, const FText& Message, const FText* OptTitle);
 };
diff --git a/Engine/Source/Runtime/Core/Public/Stats/StatsSystemTypes.h b/Engine/Source/Runtime/Core/Public/Stats/StatsSystemTypes.h
--- a/Engine/Source/Runtime/Core/Public/Stats/StatsSystemTypes.h
+++ b/Engine/Source/Runtime/Core/Public/Stats/StatsSystemTypes.h
@@ -304,24 +304,6 @@ struct EMemoryRegion
 	static_assert(FPlatformMemory::MCR_MAX < (1 << NumBits), "Need to expand memory region field.");
 };
 
-/** Memory operation for STAT_Memory_AllocPtr. */
-enum class UE_DEPRECATED(5.3, "Use Trace/MemoryInsights and/or LLM for memory profiling.") EMemoryOperation : uint8
-{
-	/** Invalid. */
-	Invalid,
-	/** Alloc. */
-	Alloc,
-	/** Free. */
-	Free,
-	/** Realloc. */
-	Realloc,
-
-	Num,
-	Mask = 0x7,
-	/** Every allocation is aligned to 8 or 16 bytes, so we have 3/4 bits free to use. */
-	NumBits = 3,
-};
-
 /**
  * A few misc final bit packing computations
  */
diff --git a/Engine/Source/Runtime/Core/Public/Templates/RemoveCV.h b/Engine/Source/Runtime/Core/Public/Templates/RemoveCV.h
deleted file mode 100644
diff --git a/Engine/Source/Runtime/Core/Public/Templates/Sorting.h b/Engine/Source/Runtime/Core/Public/Templates/Sorting.h
--- a/Engine/Source/Runtime/Core/Public/Templates/Sorting.h
+++ b/Engine/Source/Runtime/Core/Public/Templates/Sorting.h
@@ -83,65 +83,6 @@ struct TIsContiguousContainer< TArrayRange<T> >
 	enum { Value = true };
 };
 
-/**
- * Sort elements using user defined predicate class. The sort is unstable, meaning that the ordering of equal items is not necessarily preserved.
- *
- * @param	First	pointer to the first element to sort
- * @param	Num		the number of items to sort
- * @param Predicate predicate class
- */
-template<class T, class PREDICATE_CLASS> 
-UE_DEPRECATED(5.3, "Sort is deprecated, please use Algo::Sort. Algo::Sort supports ranges with index types other than int32, and doesn't automatically dereference pointers.")
-void Sort( T* First, const int32 Num, const PREDICATE_CLASS& Predicate )
-{
-	TArrayRange<T> ArrayRange( First, Num );
-	Algo::Sort( ArrayRange, TDereferenceWrapper<T, PREDICATE_CLASS>( Predicate ) );
-}
-
-/**
- * Specialized version of the above Sort function for pointers to elements.
- *
- * @param	First	pointer to the first element to sort
- * @param	Num		the number of items to sort
- * @param Predicate predicate class
- */
-template<class T, class PREDICATE_CLASS> 
-UE_DEPRECATED(5.3, "Sort is deprecated, please use Algo::Sort. Algo::Sort supports ranges with index types other than int32, and doesn't automatically dereference pointers.")
-void Sort( T** First, const int32 Num, const PREDICATE_CLASS& Predicate )
-{
-	TArrayRange<T*> ArrayRange( First, Num );
-	Algo::Sort( ArrayRange, TDereferenceWrapper<T*, PREDICATE_CLASS>( Predicate ) );
-}
-
-/**
- * Sort elements. The sort is unstable, meaning that the ordering of equal items is not necessarily preserved.
- * Assumes < operator is defined for the template type.
- *
- * @param	First	pointer to the first element to sort
- * @param	Num		the number of items to sort
- */
-template<class T> 
-UE_DEPRECATED(5.3, "Sort is deprecated, please use Algo::Sort. Algo::Sort supports ranges with index types other than int32, and doesn't automatically dereference pointers.")
-void Sort( T* First, const int32 Num )
-{
-	TArrayRange<T> ArrayRange( First, Num );
-	Algo::Sort( ArrayRange, TDereferenceWrapper<T, TLess<T> >( TLess<T>() ) );
-}
-
-/**
- * Specialized version of the above Sort function for pointers to elements.
- *
- * @param	First	pointer to the first element to sort
- * @param	Num		the number of items to sort
- */
-template<class T> 
-UE_DEPRECATED(5.3, "Sort is deprecated, please use Algo::Sort. Algo::Sort supports ranges with index types other than int32, and doesn't automatically dereference pointers.")
-void Sort( T** First, const int32 Num )
-{
-	TArrayRange<T*> ArrayRange( First, Num );
-	Algo::Sort( ArrayRange, TDereferenceWrapper<T*, TLess<T> >( TLess<T>() ) );
-}
-
 /**
  * Stable merge to perform sort below. Stable sort is slower than non-stable
  * algorithm.
@@ -380,68 +321,6 @@ void StableSortInternal(T* First, const int32 Num, const PREDICATE_CLASS& Predic
 	TMergeSort<TRotationInPlaceMerge<TJugglingRotation<FEuclidDivisionGCD> > >::Sort(First, Num, Predicate);
 }
 
-/**
- * Stable sort elements using user defined predicate class. The sort is stable,
- * meaning that the ordering of equal items is preserved, but it's slower than
- * non-stable algorithm.
- *
- * @param	First	pointer to the first element to sort
- * @param	Num		the number of items to sort
- * @param Predicate predicate class
- */
-template<class T, class PREDICATE_CLASS>
-UE_DEPRECATED(5.3, "StableSort is deprecated, please use Algo::StableSort. Algo::StableSort supports ranges with index types other than int32, and doesn't automatically dereference pointers.")
-void StableSort(T* First, const int32 Num, const PREDICATE_CLASS& Predicate)
-{
-	StableSortInternal(First, Num, TDereferenceWrapper<T, PREDICATE_CLASS>(Predicate));
-}
-
-/**
- * Specialized version of the above StableSort function for pointers to elements.
- * Stable sort is slower than non-stable algorithm.
- *
- * @param	First	pointer to the first element to sort
- * @param	Num		the number of items to sort
- * @param Predicate predicate class
- */
-template<class T, class PREDICATE_CLASS>
-UE_DEPRECATED(5.3, "StableSort is deprecated, please use Algo::StableSort. Algo::StableSort supports ranges with index types other than int32, and doesn't automatically dereference pointers.")
-void StableSort(T** First, const int32 Num, const PREDICATE_CLASS& Predicate)
-{
-	StableSortInternal(First, Num, TDereferenceWrapper<T*, PREDICATE_CLASS>(Predicate));
-}
-
-/**
- * Stable sort elements. The sort is stable, meaning that the ordering of equal
- * items is preserved, but it's slower than non-stable algorithm.
- *
- * Assumes < operator is defined for the template type.
- *
- * @param	First	pointer to the first element to sort
- * @param	Num		the number of items to sort
- */
-template<class T>
-UE_DEPRECATED(5.3, "StableSort is deprecated, please use Algo::StableSort. Algo::StableSort supports ranges with index types other than int32, and doesn't automatically dereference pointers.")
-void StableSort(T* First, const int32 Num)
-{
-	StableSortInternal(First, Num, TDereferenceWrapper<T, TLess<T> >(TLess<T>()));
-}
-
-/**
- * Specialized version of the above StableSort function for pointers to elements.
- * Stable sort is slower than non-stable algorithm.
- *
- * @param	First	pointer to the first element to sort
- * @param	Num		the number of items to sort
- */
-template<class T>
-UE_DEPRECATED(5.3, "StableSort is deprecated, please use Algo::StableSort. Algo::StableSort supports ranges with index types other than int32, and doesn't automatically dereference pointers.")
-void StableSort(T** First, const int32 Num)
-{
-	StableSortInternal(First, Num, TDereferenceWrapper<T*, TLess<T> >(TLess<T>()));
-}
-
-
 /**
  * Very fast 32bit radix sort.
  * SortKeyClass defines operator() that takes ValueType and returns a uint32. Sorting based on key.
@@ -723,4 +602,4 @@ template< typename ValueType, typename CountType >
 void RadixSort64(ValueType* RESTRICT Array, CountType Num)
 {
 	RadixSort64(Array, Num, TRadixSortKeyCastUint64< ValueType >());
-}
\ No newline at end of file
+}
diff --git a/Engine/Source/Runtime/Core/Public/Templates/TypeHash.h b/Engine/Source/Runtime/Core/Public/Templates/TypeHash.h
--- a/Engine/Source/Runtime/Core/Public/Templates/TypeHash.h
+++ b/Engine/Source/Runtime/Core/Public/Templates/TypeHash.h
@@ -197,18 +197,6 @@ inline uint32 GetTypeHash(T (&Array)[N])
 	return PointerHash(Array);
 }
 
-template <
-	typename T,
-	std::enable_if_t<std::is_same_v<const T, const TCHAR>>* = nullptr
->
-UE_DEPRECATED(5.3, "Hashing TCHAR arrays is deprecated - use PointerHash() to force a conversion to a pointer or FCrc::Strihash_DEPRECATED to do a string hash, or use TStringPointerSetKeyFuncs_DEPRECATED or TStringPointerMapKeyFuncs_DEPRECATED as keyfuncs for a TSet or TMap respectively")
-inline uint32 GetTypeHash(T* Value)
-{
-	// Hashing a TCHAR* array differently from a void* is dangerous and is deprecated.
-	// When removing these overloads post-deprecation, comment in the related static_assert in the std::is_pointer_v block of the GetTypeHash overload above.
-	return FCrc::Strihash_DEPRECATED(Value);
-}
-
 template <typename T>
 [[nodiscard]] inline uint32 GetArrayHash(const T* Ptr, uint64 Size, uint32 PreviousHash = 0)
 {
diff --git a/Engine/Source/Runtime/CoreUObject/Private/Serialization/ArchiveStackTrace.cpp b/Engine/Source/Runtime/CoreUObject/Private/Serialization/ArchiveStackTrace.cpp
--- a/Engine/Source/Runtime/CoreUObject/Private/Serialization/ArchiveStackTrace.cpp
+++ b/Engine/Source/Runtime/CoreUObject/Private/Serialization/ArchiveStackTrace.cpp
@@ -183,1488 +183,6 @@ bool ShouldBypassDiff()
 
 } // namespace UE::ArchiveStackTrace
 
-PRAGMA_DISABLE_DEPRECATION_WARNINGS;
-
-bool FArchiveDiffMap::ContainsOffset(int64 Offset) const
-{
-	for (const FArchiveDiffInfo& Diff : *this)
-	{
-		if (Diff.Offset <= Offset && Offset < (Diff.Offset + Diff.Size))
-		{
-			return true;
-		}
-	}
-
-	return false;
-}
-
-FArchiveCallstacks::FCallstackData::FCallstackData(TUniquePtr<ANSICHAR[]>&& InCallstack, UObject* InSerializedObject, FProperty* InSerializedProperty)
-	: Callstack(MoveTemp(InCallstack))
-	, SerializedProp(InSerializedProperty)
-{
-	if (InSerializedObject)
-	{
-		SerializedObjectName = InSerializedObject->GetFullName();
-	}
-	if (InSerializedProperty)
-	{
-		SerializedPropertyName = InSerializedProperty->GetFullName();
-	}
-}
-
-FString FArchiveCallstacks::FCallstackData::ToString(const TCHAR* CallstackCutoffText) const
-{
-	FString HumanReadableString;
-
-#if !NO_LOGGING
-	const TCHAR* const LineTerminator = FDiffFormatHelper::Get().LineTerminator;
-	const TCHAR* const Indent = FDiffFormatHelper::Get().Indent;
-	
-	FString StackTraceText = Callstack.Get();
-	if (CallstackCutoffText != nullptr)
-	{
-		// If the cutoff string is provided, remove all functions starting with the one specifiec in the cutoff string
-		int32 CutoffIndex = StackTraceText.Find(CallstackCutoffText, ESearchCase::CaseSensitive);
-		if (CutoffIndex > 0)
-		{
-			CutoffIndex = StackTraceText.Find(TEXT("\n"), ESearchCase::CaseSensitive, ESearchDir::FromEnd, CutoffIndex - 1);
-			if (CutoffIndex > 0)
-			{
-				StackTraceText = StackTraceText.Left(CutoffIndex + 1);
-			}
-		}
-	}
-
-	TArray<FString> StackLines;
-	StackTraceText.ParseIntoArrayLines(StackLines);
-	for (FString& StackLine : StackLines)
-	{
-		if (StackLine.StartsWith(TEXT("0x")))
-		{
-			int32 CutoffIndex = StackLine.Find(TEXT(" "), ESearchCase::CaseSensitive);
-			if (CutoffIndex >= -1 && CutoffIndex < StackLine.Len() - 2)
-			{
-				StackLine.MidInline(CutoffIndex + 1, MAX_int32, EAllowShrinking::No);
-			}
-		}
-		HumanReadableString += Indent;
-		HumanReadableString += StackLine;
-		HumanReadableString += LineTerminator;
-	}
-
-	if (!SerializedObjectName.IsEmpty())
-	{
-		HumanReadableString += LineTerminator;
-		HumanReadableString += Indent;
-		HumanReadableString += TEXT("Serialized Object: ");
-		HumanReadableString += SerializedObjectName;
-		HumanReadableString += LineTerminator;
-	}
-	if (!SerializedPropertyName.IsEmpty())
-	{
-		if (SerializedObjectName.IsEmpty())
-		{
-			HumanReadableString += LineTerminator;
-		}
-		HumanReadableString += Indent;
-		HumanReadableString += TEXT("Serialized Property: ");
-		HumanReadableString += SerializedPropertyName;
-		HumanReadableString += LineTerminator;
-	}
-#endif // !NO_LOGGING
-	return HumanReadableString;
-}
-
-FArchiveCallstacks::FCallstackData FArchiveCallstacks::FCallstackData::Clone() const
-{
-	TUniquePtr<ANSICHAR[]> CallstackCopy;
-	if (const int32 Len = FCStringAnsi::Strlen(Callstack.Get()); Len > 0)
-	{
-		CallstackCopy = MakeUnique<ANSICHAR[]>(Len + 1);
-		FMemory::Memcpy(CallstackCopy.Get(), Callstack.Get(), Len + 1);
-	}
-
-	FCallstackData Clone(MoveTemp(CallstackCopy), nullptr, SerializedProp);
-	Clone.SerializedObjectName = SerializedObjectName;
-
-	return Clone;
-}
-
-FArchiveCallstacks::FArchiveCallstacks(UObject* InAsset)
-	: Asset(InAsset)
-	, bCallstacksDirty(true)
-	, StackTraceSize(65535)
-	, LastSerializeCallstack(nullptr)
-	, TotalSize(0)
-{
-	StackTrace = MakeUnique<ANSICHAR[]>(StackTraceSize);
-	StackTrace[0] = 0;
-}
-
-FName FArchiveCallstacks::GetAssetClass() const
-{
-	return Asset != nullptr ? Asset->GetClass()->GetFName() : NAME_None;
-}
-
-ANSICHAR* FArchiveCallstacks::AddUniqueCallstack(bool bIsCollectingCallstacks, UObject* SerializedObject, FProperty* SerializedProperty, uint32& OutCallstackCRC)
-{
-	ANSICHAR* Callstack = nullptr;
-	if (bIsCollectingCallstacks)
-	{
-		OutCallstackCRC = FCrc::StrCrc32(StackTrace.Get());
-
-		if (FCallstackData* ExistingCallstack = UniqueCallstacks.Find(OutCallstackCRC))
-		{
-			Callstack = ExistingCallstack->Callstack.Get();
-		}
-		else
-		{
-			const int32 Len = FCStringAnsi::Strlen(StackTrace.Get()) + 1;
-			TUniquePtr<ANSICHAR[]> NewCallstack = MakeUnique<ANSICHAR[]>(Len);
-			FCStringAnsi::Strcpy(NewCallstack.Get(), Len, StackTrace.Get());
-			FCallstackData& NewEntry = UniqueCallstacks.Add(OutCallstackCRC, FCallstackData(MoveTemp(NewCallstack), SerializedObject, SerializedProperty));
-			Callstack = NewEntry.Callstack.Get();
-		}
-	}
-	else
-	{
-		OutCallstackCRC = 0;
-	}
-	return Callstack;
-}
-
-void FArchiveCallstacks::Add(
-	int64 Offset,
-	int64 Length,
-	UObject* SerializedObject,
-	FProperty* SerializedProperty,
-	TArrayView<const FName> DebugDataStack,
-	bool bIsCollectingCallstacks,
-	bool bCollectCurrentCallstack,
-	int32 StackIgnoreCount)
-{
-	if (GIgnoreDiffManager.ShouldBypassDiff())
-	{
-		return;
-	}
-
-	const int64 CurrentOffset = Offset;
-	TotalSize = FMath::Max(TotalSize, CurrentOffset + Length); 
-
-	const bool bShouldCollectCallstack = bIsCollectingCallstacks && bCollectCurrentCallstack && !GIgnoreDiffManager.ShouldIgnoreDiff();
-	if (bShouldCollectCallstack)
-	{
-		StackTrace[0] = '\0';
-		FPlatformStackWalk::StackWalkAndDump(StackTrace.Get(), StackTraceSize, StackIgnoreCount);
-#if WITH_EDITOR && !NO_LOGGING
-		//if we have a debug name stack, plaster it onto the end of the current stack buffer so that it's a part of the unique stack entry.
-		if (DebugDataStack.Num() > 0)
-		{
-			FCStringAnsi::Strcat(StackTrace.Get(), StackTraceSize, DebugDataStackMarker);
-
-			const FString SubIndent = FString(FDiffFormatHelper::Get().Indent) + FString("    ");
-
-			bool bIsIndenting = true;
-			for (const auto& DebugData : DebugDataStack)
-			{
-				if (bIsIndenting)
-				{
-					FCStringAnsi::Strcat(StackTrace.Get(), StackTraceSize, TCHAR_TO_ANSI(*SubIndent));
-				}
-
-				ANSICHAR DebugName[NAME_SIZE];
-				DebugData.GetPlainANSIString(DebugName);
-				FCStringAnsi::Strcat(StackTrace.Get(), StackTraceSize, DebugName);
-
-				//these are special-cased, as we assume they'll be followed by object/property names and want the names on the same line for readability's sake.
-				const bool bIsPropertyLabel = (DebugData == TEXT("SerializeScriptProperties") || DebugData == TEXT("PropertySerialize") || DebugData == TEXT("SerializeTaggedProperty"));
-				const ANSICHAR* const LineEnd = bIsPropertyLabel ? ": " : "\r\n";
-				FCStringAnsi::Strcat(StackTrace.Get(), StackTraceSize, LineEnd);
-				bIsIndenting = !bIsPropertyLabel;
-			}
-		}
-#endif
-		// Make sure we compare the new stack trace with the last one in the next if statement
-		bCallstacksDirty = true;
-	}
-
-	if (LastSerializeCallstack == nullptr || (bCallstacksDirty && FCStringAnsi::Strcmp(LastSerializeCallstack, StackTrace.Get()) != 0))
-	{
-		uint32 CallstackCRC = 0;
-		if (CallstackAtOffsetMap.Num() == 0 || CurrentOffset > CallstackAtOffsetMap.Last().Offset)
-		{
-			// New data serialized at the end of archive buffer
-			LastSerializeCallstack = AddUniqueCallstack(bIsCollectingCallstacks, SerializedObject, SerializedProperty, CallstackCRC);
-			CallstackAtOffsetMap.Add(FCallstackAtOffset {CurrentOffset, CallstackCRC, GIgnoreDiffManager.ShouldIgnoreDiff()});
-		}
-		else
-		{
-			// This happens usually after Seek() so we need to find the exiting offset or insert a new one
-			const int32 CallstackToUpdateIndex = GetCallstackIndexAtOffset(CurrentOffset);
-			check(CallstackToUpdateIndex != -1);
-			FCallstackAtOffset& CallstackToUpdate = CallstackAtOffsetMap[CallstackToUpdateIndex];
-			LastSerializeCallstack = AddUniqueCallstack(bIsCollectingCallstacks, SerializedObject, SerializedProperty, CallstackCRC);
-			if (CallstackToUpdate.Offset == CurrentOffset)
-			{
-				CallstackToUpdate.Callstack = CallstackCRC;
-			}
-			else
-			{
-				// Insert a new callstack
-				check(CallstackToUpdate.Offset < CurrentOffset);
-				CallstackAtOffsetMap.Insert(FCallstackAtOffset {CurrentOffset, CallstackCRC, GIgnoreDiffManager.ShouldIgnoreDiff()}, CallstackToUpdateIndex + 1);
-			}
-		}
-		check(CallstackCRC != 0 || !bShouldCollectCallstack);
-	}
-	else if (LastSerializeCallstack)
-	{
-		// Skip callstack comparison on next serialize call unless we grab a stack trace
-		bCallstacksDirty = false;
-	}
-}
-
-int32 FArchiveCallstacks::GetCallstackIndexAtOffset(int64 Offset, int32 MinOffsetIndex) const
-{
-	if (Offset < 0 || Offset > TotalSize || MinOffsetIndex < 0 || MinOffsetIndex >= CallstackAtOffsetMap.Num())
-	{
-		return -1;
-	}
-
-	// Find the index of the offset the InOffset maps to
-	int32 OffsetForCallstackIndex = -1;
-	int32 MaxOffsetIndex = CallstackAtOffsetMap.Num() - 1;
-
-	// Binary search
-	for (; MinOffsetIndex <= MaxOffsetIndex; )
-	{
-		int32 SearchIndex = (MinOffsetIndex + MaxOffsetIndex) / 2;
-		if (CallstackAtOffsetMap[SearchIndex].Offset < Offset)
-		{
-			MinOffsetIndex = SearchIndex + 1;
-		}
-		else if (CallstackAtOffsetMap[SearchIndex].Offset > Offset)
-		{
-			MaxOffsetIndex = SearchIndex - 1;
-		}
-		else
-		{
-			OffsetForCallstackIndex = SearchIndex;
-			break;
-		}
-	}
-	
-	if (OffsetForCallstackIndex == -1)
-	{
-		// We didn't find the exact offset value so let's try to find the first one that is lower than the requested one
-		MinOffsetIndex = FMath::Min(MinOffsetIndex, CallstackAtOffsetMap.Num() - 1);
-		for (int32 FirstLowerOffsetIndex = MinOffsetIndex; FirstLowerOffsetIndex >= 0; --FirstLowerOffsetIndex)
-		{
-			if (CallstackAtOffsetMap[FirstLowerOffsetIndex].Offset < Offset)
-			{
-				OffsetForCallstackIndex = FirstLowerOffsetIndex;
-				break;
-			}
-		}
-		check(OffsetForCallstackIndex != -1);
-		check(CallstackAtOffsetMap[OffsetForCallstackIndex].Offset < Offset);
-		check(OffsetForCallstackIndex == (CallstackAtOffsetMap.Num() - 1) || CallstackAtOffsetMap[OffsetForCallstackIndex + 1].Offset > Offset);
-	}
-
-	return OffsetForCallstackIndex;
-}
-
-void FArchiveCallstacks::Append(const FArchiveCallstacks& Other, int64 Offset)
-{
-	for (const FCallstackAtOffset& OtherOffset : Other.CallstackAtOffsetMap)
-	{
-		FCallstackAtOffset& New = CallstackAtOffsetMap.Add_GetRef(OtherOffset);
-		New.Offset += Offset;
-	}
-
-	CallstackAtOffsetMap.Sort([](const FCallstackAtOffset& LHS,const FCallstackAtOffset& RHS)
-	{
-		return LHS.Offset < RHS.Offset;
-	});
-
-	for (const TPair<uint32, FCallstackData>& Kv : Other.UniqueCallstacks)
-	{
-		if (FCallstackData* Existing = UniqueCallstacks.Find(Kv.Key))
-		{
-			if (LastSerializeCallstack == Existing->Callstack.Get())
-			{
-				LastSerializeCallstack = nullptr;
-			}
-			UniqueCallstacks.Remove(Kv.Key);
-		}
-		UniqueCallstacks.Emplace(Kv.Key, Kv.Value.Clone());
-	}
-
-	TotalSize = FMath::Max(TotalSize, Other.TotalSize);
-}
-
-FArchiveStackTraceWriter::FArchiveStackTraceWriter(
-	FArchive& InInner,
-	FArchiveCallstacks& InCallstacks,
-	const FArchiveDiffMap* InDiffMap,
-	int64 InDiffMapOffset)
-		: FArchiveProxy(InInner)
-		, Callstacks(InCallstacks)
-		, DiffMap(InDiffMap)
-		, DiffMapOffset(InDiffMapOffset)
-		, bInnerArchiveDisabled(false)
-{
-}
-
-FArchiveStackTraceWriter::~FArchiveStackTraceWriter() = default;
-
-void FArchiveStackTraceWriter::Serialize(void* Data, int64 Length)
-{
-	static struct FBreakAtOffsetSettings
-	{
-		FString PackageToBreakOn;
-		int64 OffsetToBreakOn;
-
-		FBreakAtOffsetSettings()
-			: OffsetToBreakOn(-1)
-		{
-			if (!FParse::Param(FCommandLine::Get(), TEXT("cooksinglepackage")))
-			{
-				return;
-			}
-
-			FString Package;
-			if (!FParse::Value(FCommandLine::Get(), TEXT("map="), Package))
-			{
-				return;
-			}
-
-			int64 Offset;
-			if (!FParse::Value(FCommandLine::Get(), TEXT("diffonlybreakoffset="), Offset) || Offset <= 0)
-			{
-				return;
-			}
-
-			OffsetToBreakOn = Offset;
-			PackageToBreakOn = TEXT("/") + FPackageName::GetShortName(Package);
-		}
-	} BreakAtOffsetSettings;
-
-	const int64 CurrentOffset = DiffMapOffset + Tell();
-
-	if (BreakAtOffsetSettings.OffsetToBreakOn >= 0 && BreakAtOffsetSettings.OffsetToBreakOn >= CurrentOffset && BreakAtOffsetSettings.OffsetToBreakOn < CurrentOffset + Length)
-	{
-		FString ArcName = GetArchiveName();
-		int32 SubnameIndex = ArcName.Find(BreakAtOffsetSettings.PackageToBreakOn, ESearchCase::IgnoreCase, ESearchDir::FromEnd);
-		if (SubnameIndex >= 0)
-		{
-			int32 SubnameEndIndex = SubnameIndex + BreakAtOffsetSettings.PackageToBreakOn.Len();
-			if (SubnameEndIndex == ArcName.Len() || ArcName[SubnameEndIndex] == TEXT('.'))
-			{
-				UE_DEBUG_BREAK();
-			}
-		}
-	}
-
-	if (Length > 0)
-	{
-		UObject* SerializedObject = FUObjectThreadContext::Get().GetSerializeContext()->SerializedObject;
-		TArrayView<const FName> DebugStack;
-#if WITH_EDITOR
-		DebugStack = DebugDataStack;
-#endif
-		const bool bIsCollectingCallstacks = DiffMap != nullptr;
-		const bool bCollectCurrentCallstack = bIsCollectingCallstacks && DiffMap->ContainsOffset(CurrentOffset);
-
-		Callstacks.Add(
-			CurrentOffset,
-			Length,
-			SerializedObject,
-			GetSerializedProperty(),
-			DebugStack,
-			bIsCollectingCallstacks,
-			bCollectCurrentCallstack,
-			StackIgnoreCount);
-	}
-
-	if (bInnerArchiveDisabled == false)
-	{
-		InnerArchive.Serialize(Data, Length);
-	}
-}
-
-FArchiveStackTraceMemoryWriter::FArchiveStackTraceMemoryWriter(
-	FArchiveCallstacks& Callstacks,
-	const FArchiveDiffMap* DiffMap,
-	const int64 DiffMapOffset,
-	const int64 PreAllocateBytes,
-	bool bIsPersistent,
-	const TCHAR* Filename)
-		: FLargeMemoryWriter(PreAllocateBytes, bIsPersistent, Filename)
-		, StackTraceWriter(*this, Callstacks, DiffMap, DiffMapOffset)
-{
-	// Hack to prevent recursive calls to serialize when passing in this.
-	StackTraceWriter.SetDisableInnerArchive(true);
-	StackTraceWriter.SetStackIgnoreCount(StackTraceWriter.GetStackIgnoreCount() + 1);
-}
-
-void FArchiveStackTraceMemoryWriter::Serialize(void* Memory, int64 Length)
-{
-	StackTraceWriter.Serialize(Memory, Length);
-	FLargeMemoryWriter::Serialize(Memory, Length);
-}
-
-FArchiveStackTrace::FArchiveStackTrace(UObject* InAsset, const TCHAR* InFilename, bool bInCollectCallstacks, const FArchiveDiffMap* InDiffMap)
-	: FLargeMemoryWriter(0, false, InFilename)
-	, Callstacks(InAsset)
-	, StackTraceWriter(*this, Callstacks, InDiffMap, 0)
-{
-	// Hack to prevent recursive calls to serialize when passing in this.
-	StackTraceWriter.SetDisableInnerArchive(true);
-	StackTraceWriter.SetStackIgnoreCount(StackTraceWriter.GetStackIgnoreCount() + 1);
-}
-
-FArchiveStackTrace::~FArchiveStackTrace()
-{
-}
-
-void FArchiveStackTrace::Serialize(void* Memory, int64 Length)
-{
-	StackTraceWriter.Serialize(Memory, Length);
-	FLargeMemoryWriter::Serialize(Memory, Length);
-}
-
-bool FArchiveStackTrace::LoadPackageIntoMemory(const TCHAR* InFilename, FPackageData& OutPackageData,
-	TUniquePtr<uint8, UE::ArchiveStackTrace::FDeleteByFree>& OutLoadedBytes)
-{
-	UE::ArchiveStackTrace::FPackageData PackageData;
-	bool bResult = UE::ArchiveStackTrace::LoadPackageIntoMemory(InFilename, PackageData, OutLoadedBytes);
-	OutPackageData.Data = PackageData.Data;
-	OutPackageData.Size = PackageData.Size;
-	OutPackageData.HeaderSize = PackageData.HeaderSize;
-	OutPackageData.StartOffset = PackageData.StartOffset;
-	return bResult;
-}
-
-namespace
-{
-	bool ShouldDumpPropertyValueState(FProperty* Prop)
-	{
-		if (Prop->IsA<FNumericProperty>()
-			|| Prop->IsA<FStrProperty>()
-			|| Prop->IsA<FBoolProperty>()
-			|| Prop->IsA<FNameProperty>())
-		{
-			return true;
-		}
-
-		if (FArrayProperty* ArrayProp = CastField<FArrayProperty>(Prop))
-		{
-			return ShouldDumpPropertyValueState(ArrayProp->Inner);
-		}
-
-		if (FMapProperty* MapProp = CastField<FMapProperty>(Prop))
-		{
-			return ShouldDumpPropertyValueState(MapProp->KeyProp) && ShouldDumpPropertyValueState(MapProp->ValueProp);
-		}
-
-		if (FSetProperty* SetProp = CastField<FSetProperty>(Prop))
-		{
-			return ShouldDumpPropertyValueState(SetProp->ElementProp);
-		}
-
-		if (FStructProperty* StructProp = CastField<FStructProperty>(Prop))
-		{
-			if (StructProp->Struct == TBaseStructure<FVector>::Get()
-				|| StructProp->Struct == TBaseStructure<FGuid>::Get())
-			{
-				return true;
-			}
-		}
-
-		if (FOptionalProperty* OptionalProp = CastField<FOptionalProperty>(Prop))
-		{
-			return ShouldDumpPropertyValueState(OptionalProp->GetValueProperty());
-		}
-
-		return false;
-	}
-}
-
-void FArchiveStackTraceWriter::Compare(
-	const FPackageData& SourcePackage,
-	const FPackageData& DestPackage,
-	const FArchiveCallstacks& Callstacks,
-	const FArchiveDiffMap& DiffMap,
-	const TCHAR* AssetFilename,
-	const TCHAR* CallstackCutoffText,
-	const int64 MaxDiffsToLog,
-	int32& InOutDiffsLogged,
-	TMap<FName, FArchiveDiffStats>& OutStats,
-	bool bSuppressLogging)
-{
-#if !NO_LOGGING
-	const TCHAR* const Indent = FDiffFormatHelper::Get().Indent;
-	const TCHAR* const LineTerminator = FDiffFormatHelper::Get().LineTerminator;
-	const int64 SourceSize = SourcePackage.Size - SourcePackage.StartOffset;
-	const int64 DestSize = DestPackage.Size - DestPackage.StartOffset;
-	const int64 SizeToCompare = FMath::Min(SourceSize, DestSize);
-	const FName AssetClass = Callstacks.GetAssetClass();
-	
-	if (SourceSize != DestSize)
-	{
-		UE_CLOG(!bSuppressLogging, LogArchiveDiff, Warning, TEXT("%s: Size mismatch: on disk: %lld vs memory: %lld"), AssetFilename, SourceSize, DestSize);
-		int64 SizeDiff = DestPackage.Size - SourcePackage.Size;
-		OutStats.FindOrAdd(AssetClass).DiffSize += SizeDiff;
-	}
-
-	FString LastDifferenceCallstackDataText;
-	int32 LastDifferenceCallstackOffsetIndex = -1;
-	int64 NumDiffsLocal = 0;
-	int64 NumDiffsLoggedLocal = 0;
-	int64 FirstUnreportedDiffIndex = -1;
-
-	for (int64 LocalOffset = 0; LocalOffset < SizeToCompare; ++LocalOffset)
-	{
-		const int64 SourceAbsoluteOffset = LocalOffset + SourcePackage.StartOffset;
-		const int64 DestAbsoluteOffset = LocalOffset + DestPackage.StartOffset;
-
-		const uint8 SourceByte = SourcePackage.Data[SourceAbsoluteOffset];
-		const uint8 DestByte   = DestPackage  .Data[DestAbsoluteOffset];
-		if (SourceByte == DestByte)
-		{
-			continue;
-		}
-
-		bool bDifferenceLogged = false;
-		ON_SCOPE_EXIT
-		{
-			if (bDifferenceLogged)
-			{
-				InOutDiffsLogged++;
-				NumDiffsLoggedLocal++;
-			}
-		};
-
-		if (DiffMap.ContainsOffset(DestAbsoluteOffset))
-		{
-			int32 DifferenceCallstackoffsetIndex = Callstacks.GetCallstackIndexAtOffset(DestAbsoluteOffset, FMath::Max(LastDifferenceCallstackOffsetIndex, 0));
-			ON_SCOPE_EXIT
-			{
-				LastDifferenceCallstackOffsetIndex = DifferenceCallstackoffsetIndex;
-			};
-
-			if (DifferenceCallstackoffsetIndex < 0)
-			{
-				UE_CLOG(!bSuppressLogging, LogArchiveDiff, Warning, TEXT("%s: Difference at offset %lld (absolute offset: %lld), unknown callstack"), AssetFilename, LocalOffset, DestAbsoluteOffset);
-				continue;
-			}
-
-			if (DifferenceCallstackoffsetIndex == LastDifferenceCallstackOffsetIndex)
-			{
-				continue;
-			}
-
-			const FArchiveCallstacks::FCallstackAtOffset& CallstackAtOffset = Callstacks.GetCallstack(DifferenceCallstackoffsetIndex);
-			const FArchiveCallstacks::FCallstackData& DifferenceCallstackData = Callstacks.GetCallstackData(CallstackAtOffset);
-			FString DifferenceCallstackDataText = DifferenceCallstackData.ToString(CallstackCutoffText);
-			if (LastDifferenceCallstackDataText.Compare(DifferenceCallstackDataText, ESearchCase::CaseSensitive) == 0)
-			{
-				continue;
-			}
-			ON_SCOPE_EXIT
-			{
-				LastDifferenceCallstackDataText = MoveTemp(DifferenceCallstackDataText);
-			};
-
-			if (!CallstackAtOffset.bIgnore && (MaxDiffsToLog < 0 || InOutDiffsLogged < MaxDiffsToLog))
-			{
-				FString BeforePropertyVal;
-				FString AfterPropertyVal;
-				if (FProperty* SerProp = DifferenceCallstackData.SerializedProp)
-				{
-					if (SourceSize == DestSize && ShouldDumpPropertyValueState(SerProp))
-					{
-						// Walk backwards until we find a callstack which wasn't from the given property
-						int64 OffsetX = DestAbsoluteOffset;
-						for (;;)
-						{
-							if (OffsetX == 0)
-							{
-								break;
-							}
-
-							const int32 CallstackIndex = Callstacks.GetCallstackIndexAtOffset(OffsetX - 1, 0);
-							const FArchiveCallstacks::FCallstackAtOffset& PreviousCallstack = Callstacks.GetCallstack(CallstackIndex);
-							const FArchiveCallstacks::FCallstackData& PreviousCallstackData = Callstacks.GetCallstackData(PreviousCallstack);
-							if (PreviousCallstackData.SerializedProp != SerProp)
-							{
-								break;
-							}
-
-							--OffsetX;
-						}
-
-						FPropertyTempVal SourceVal(SerProp);
-						FPropertyTempVal DestVal  (SerProp);
-
-						FStaticMemoryReader SourceReader(&SourcePackage.Data[SourceAbsoluteOffset - (DestAbsoluteOffset - OffsetX)], SourcePackage.Size - SourceAbsoluteOffset);
-						FStaticMemoryReader DestReader(&DestPackage.Data[OffsetX], DestPackage.Size - DestAbsoluteOffset);
-
-						SourceVal.Serialize(SourceReader);
-						DestVal  .Serialize(DestReader);
-
-									if (!SourceReader.IsError() && !DestReader.IsError())
-									{
-										SourceVal.ExportText(BeforePropertyVal);
-										DestVal  .ExportText(AfterPropertyVal);
-									}
-								}
-							}
-
-				FString DiffValues;
-				if (BeforePropertyVal != AfterPropertyVal)
-				{
-					DiffValues = FString::Printf(TEXT("\r\n%sBefore: %s\r\n%sAfter:  %s"), Indent, *BeforePropertyVal, Indent, *AfterPropertyVal);
-				}
-
-				FString DebugDataStackText;
-#if WITH_EDITOR
-				//check for a debug data stack as part of the unique stack entry, and log it out if we find it.
-				FString FullStackText = DifferenceCallstackData.Callstack.Get();
-				int32 DebugDataIndex = FullStackText.Find(ANSI_TO_TCHAR(DebugDataStackMarker), ESearchCase::CaseSensitive);
-				if (DebugDataIndex > 0)
-				{
-					DebugDataStackText = FString::Printf(TEXT("\r\n%s"), FDiffFormatHelper::Get().Indent) + FullStackText.RightChop(DebugDataIndex + 2);
-				}
-#endif
-				UE_CLOG(
-					!bSuppressLogging,
-					LogArchiveDiff,
-					Warning,
-					TEXT("%s: Difference at offset %lld%s (absolute offset: %lld): byte %d on disk, byte %d in memory, callstack:%s%s%s%s%s"),
-					AssetFilename,
-					CallstackAtOffset.Offset - DestPackage.StartOffset,
-					DestAbsoluteOffset > CallstackAtOffset.Offset ? *FString::Printf(TEXT("(+%lld)"), DestAbsoluteOffset - CallstackAtOffset.Offset) : TEXT(""),
-					DestAbsoluteOffset,
-					SourceByte, DestByte,
-					LineTerminator,
-					LineTerminator,
-					*DifferenceCallstackDataText,
-					*DiffValues,
-					*DebugDataStackText
-				);
-
-				const int BytesToLog = 128;
-				UE_CLOG(
-					!bSuppressLogging,
-					LogArchiveDiff,
-					Display,
-					TEXT("%s: Logging %d bytes around absolute offset: %" INT64_FMT " (%016" INT64_X_FMT ") in the on disk (existing) package, (which corresponds to offset %" INT64_FMT " (%016" INT64_X_FMT ") in the in-memory package)"),
-					AssetFilename,
-					BytesToLog,
-					SourceAbsoluteOffset,
-					SourceAbsoluteOffset,
-					DestAbsoluteOffset,
-					DestAbsoluteOffset
-				);
-				FCompressionUtil::LogHexDump(SourcePackage.Data, SourcePackage.Size, SourceAbsoluteOffset - BytesToLog / 2, SourceAbsoluteOffset + BytesToLog / 2);
-
-				UE_CLOG(
-					!bSuppressLogging,
-					LogArchiveDiff,
-					Display,
-					TEXT("%s: Logging %d bytes around absolute offset: %" INT64_FMT " (%016" INT64_X_FMT ") in the in memory (new) package"),
-					AssetFilename,
-					BytesToLog,
-					DestAbsoluteOffset,
-					DestAbsoluteOffset
-				);
-				FCompressionUtil::LogHexDump(DestPackage.Data, DestPackage.Size, DestAbsoluteOffset - BytesToLog / 2, DestAbsoluteOffset + BytesToLog / 2);
-
-				bDifferenceLogged = true;
-			}
-			else if (FirstUnreportedDiffIndex == -1)
-			{
-				FirstUnreportedDiffIndex = DestAbsoluteOffset;
-			}
-			OutStats.FindOrAdd(AssetClass).NumDiffs++;
-			NumDiffsLocal++;
-		}
-		else
-		{
-			// Each byte will count as a difference but without callstack data there's no way around it
-			OutStats.FindOrAdd(AssetClass).NumDiffs++;
-			NumDiffsLocal++;
-			if (FirstUnreportedDiffIndex == -1)
-			{
-				FirstUnreportedDiffIndex = DestAbsoluteOffset;
-			}
-		}
-		OutStats.FindOrAdd(AssetClass).DiffSize++;
-	}
-
-	if (MaxDiffsToLog >= 0 && NumDiffsLocal > NumDiffsLoggedLocal)
-	{
-		if (FirstUnreportedDiffIndex != -1)
-		{
-			UE_CLOG(!bSuppressLogging, LogArchiveDiff, Warning, TEXT("%s: %lld difference(s) not logged (first at absolute offset: %lld)."), AssetFilename, NumDiffsLocal - NumDiffsLoggedLocal, FirstUnreportedDiffIndex);
-		}
-		else
-		{
-			UE_CLOG(!bSuppressLogging, LogArchiveDiff, Warning, TEXT("%s: %lld difference(s) not logged."), AssetFilename, NumDiffsLocal - NumDiffsLoggedLocal);
-		}
-	}
-#endif
-}
-
-void FArchiveStackTrace::CompareWith(const TCHAR* InFilename, const int64 TotalHeaderSize, const TCHAR* CallstackCutoffText, const int32 MaxDiffsToLog, TMap<FName, FArchiveDiffStats>& OutStats)
-{
-	TUniquePtr<uint8, UE::ArchiveStackTrace::FDeleteByFree> SourcePackageBytes;
-	FPackageData SourcePackage;
-	LoadPackageIntoMemory(InFilename, SourcePackage, SourcePackageBytes);
-	CompareWith(SourcePackage, InFilename, TotalHeaderSize, CallstackCutoffText, MaxDiffsToLog, OutStats);
-}
-
-void FArchiveStackTrace::CompareWith(const FPackageData& SourcePackage, const TCHAR* FileDisplayName, const int64 TotalHeaderSize,
-	const TCHAR* CallstackCutoffText, const int32 MaxDiffsToLog, TMap<FName, FArchiveDiffStats>&OutStats,
-	const FArchiveStackTraceWriter::EPackageHeaderFormat PackageHeaderFormat /* = FArchiveStackTraceWriter::EPackageHeaderFormat::PackageFileSummary */)
-{
-	const FName AssetClass = Callstacks.GetAssetClass();
-	OutStats.FindOrAdd(AssetClass).NewFileTotalSize = TotalSize();
-	if (SourcePackage.Size == 0)
-	{
-		UE_LOG(LogArchiveDiff, Warning, TEXT("New package: %s"), *GetArchiveName());
-		OutStats.FindOrAdd(AssetClass).DiffSize = OutStats.FindOrAdd(AssetClass).NewFileTotalSize;
-		return;
-	}
-
-	FPackageData DestPackage;
-	DestPackage.Data = GetData();
-	DestPackage.Size = TotalSize();
-	DestPackage.HeaderSize = TotalHeaderSize;
-	DestPackage.StartOffset = 0;
-
-	UE_LOG(LogArchiveDiff, Display, TEXT("Comparing: %s"), *GetArchiveName());
-	UE_LOG(LogArchiveDiff, Display, TEXT("Asset class: %s"), *AssetClass.ToString());
-
-	int32 NumLoggedDiffs = 0;
-
-	FPackageData SourcePackageHeader = SourcePackage;
-	SourcePackageHeader.Size = SourcePackageHeader.HeaderSize;
-	SourcePackageHeader.HeaderSize = 0;
-	SourcePackageHeader.StartOffset = 0;
-
-	FPackageData DestPackageHeader = DestPackage;
-	DestPackageHeader.Size = TotalHeaderSize;
-	DestPackageHeader.HeaderSize = 0;
-	DestPackageHeader.StartOffset = 0;
-
-	FArchiveStackTraceWriter::Compare(SourcePackageHeader, DestPackageHeader, Callstacks, StackTraceWriter.GetDiffMap(), FileDisplayName, CallstackCutoffText, MaxDiffsToLog, NumLoggedDiffs, OutStats);
-
-	if (TotalHeaderSize > 0 && OutStats.FindOrAdd(AssetClass).NumDiffs > 0)
-	{
-		FArchiveStackTraceWriter::DumpPackageHeaderDiffs(SourcePackage, DestPackage, FileDisplayName, MaxDiffsToLog, PackageHeaderFormat);
-	}
-
-	FPackageData SourcePackageExports = SourcePackage;
-	SourcePackageExports.HeaderSize = 0;
-	SourcePackageExports.StartOffset = SourcePackage.HeaderSize;
-
-	FPackageData DestPackageExports = DestPackage;
-	DestPackageExports.HeaderSize = 0;
-	DestPackageExports.StartOffset = TotalHeaderSize;
-
-	FString AssetName;
-	if (DestPackage.HeaderSize > 0)
-	{
-		AssetName = FPaths::ChangeExtension(FileDisplayName, TEXT("uexp"));
-	}
-	else
-	{
-		AssetName = FileDisplayName;
-	}
-
-	FArchiveStackTraceWriter::Compare(SourcePackageExports, DestPackageExports, Callstacks, StackTraceWriter.GetDiffMap(), *AssetName, CallstackCutoffText, MaxDiffsToLog, NumLoggedDiffs, OutStats);
-
-	// Optionally save out any differences we detected.
-	const FArchiveDiffStats& Stats = OutStats.FindOrAdd(AssetClass);
-	if (Stats.NumDiffs > 0)
-	{
-		static struct FDiffOutputSettings
-		{
-			FString DiffOutputDir;
-
-			FDiffOutputSettings()
-			{
-				FString Dir;
-				if (!FParse::Value(FCommandLine::Get(), TEXT("diffoutputdir="), Dir))
-				{
-					return;
-				}
-
-				FPaths::NormalizeDirectoryName(Dir);
-				DiffOutputDir = MoveTemp(Dir) + TEXT("/");
-			}
-		} DiffOutputSettings;
-
-		// Only save out the differences if we have a -diffoutputdir set.
-		if (!DiffOutputSettings.DiffOutputDir.IsEmpty())
-		{
-			FString OutputFilename = FPaths::ConvertRelativePathToFull(FileDisplayName);
-			FString SavedDir       = FPaths::ConvertRelativePathToFull(FPaths::ProjectSavedDir());
-			if (OutputFilename.StartsWith(SavedDir))
-			{
-				OutputFilename.ReplaceInline(*SavedDir, *DiffOutputSettings.DiffOutputDir);
-
-				IFileManager& FileManager = IFileManager::Get();
-
-				// Copy the original asset as '.before.uasset'.
-				{
-					TUniquePtr<FArchive> DiffUAssetArchive(FileManager.CreateFileWriter(*FPaths::SetExtension(OutputFilename, TEXT(".before.") + FPaths::GetExtension(FileDisplayName))));
-					DiffUAssetArchive->Serialize(SourcePackageHeader.Data + SourcePackageHeader.StartOffset, SourcePackageHeader.Size - SourcePackageHeader.StartOffset);
-				}
-				{
-					TUniquePtr<FArchive> DiffUExpArchive(FileManager.CreateFileWriter(*FPaths::SetExtension(OutputFilename, TEXT(".before.uexp"))));
-					DiffUExpArchive->Serialize(SourcePackageExports.Data + SourcePackageExports.StartOffset, SourcePackageExports.Size - SourcePackageExports.StartOffset);
-				}
-
-				// Save out the in-memory data as '.after.uasset'.
-				{
-					TUniquePtr<FArchive> DiffUAssetArchive(FileManager.CreateFileWriter(*FPaths::SetExtension(OutputFilename, TEXT(".after.") + FPaths::GetExtension(FileDisplayName))));
-					DiffUAssetArchive->Serialize(DestPackageHeader.Data + DestPackageHeader.StartOffset, DestPackageHeader.Size - DestPackageHeader.StartOffset);
-				}
-				{
-					TUniquePtr<FArchive> DiffUExpArchive(FileManager.CreateFileWriter(*FPaths::SetExtension(OutputFilename, TEXT(".after.uexp"))));
-					DiffUExpArchive->Serialize(DestPackageExports.Data + DestPackageExports.StartOffset, DestPackageExports.Size - DestPackageExports.StartOffset);
-				}
-			}
-			else
-			{
-				UE_LOG(LogArchiveDiff, Warning, TEXT("Package '%s' doesn't seem to be writing to the Saved directory - skipping writing diff"), *OutputFilename);
-			}
-		}
-	}
-}
-
-bool FArchiveStackTraceWriter::GenerateDiffMap(const FPackageData& SourcePackage, const FPackageData& DestPackage, const FArchiveCallstacks& Callstacks, int32 MaxDiffsToFind, FArchiveDiffMap& OutDiffMap)
-{
-	bool bIdentical = true;
-	int32 LastDifferenceCallstackOffsetIndex = -1;
-	FArchiveCallstacks::FCallstackData* DifferenceCallstackData = nullptr;
-
-	const int64 SourceSize = SourcePackage.Size - SourcePackage.StartOffset;
-	const int64 DestSize = DestPackage.Size - DestPackage.StartOffset;
-	const int64 SizeToCompare = FMath::Min(SourceSize, DestSize);
-	
-	for (int64 LocalOffset = 0; LocalOffset < SizeToCompare; ++LocalOffset)
-	{
-		const int64 SourceAbsoluteOffset = LocalOffset + SourcePackage.StartOffset;
-		const int64 DestAbsoluteOffset = LocalOffset + DestPackage.StartOffset;
-		if (SourcePackage.Data[SourceAbsoluteOffset] != DestPackage.Data[DestAbsoluteOffset])
-		{
-			bIdentical = false;
-			if (OutDiffMap.Num() < MaxDiffsToFind)
-			{
-				const int32 DifferenceCallstackOffsetIndex = Callstacks.GetCallstackIndexAtOffset(DestAbsoluteOffset, FMath::Max<int32>(LastDifferenceCallstackOffsetIndex, 0));
-				if (DifferenceCallstackOffsetIndex >= 0 && DifferenceCallstackOffsetIndex != LastDifferenceCallstackOffsetIndex)
-				{
-					const FArchiveCallstacks::FCallstackAtOffset& CallstackAtOffset = Callstacks.GetCallstack(DifferenceCallstackOffsetIndex);
-					if (!CallstackAtOffset.bIgnore)
-					{
-						FArchiveDiffInfo OffsetAndSize;
-						OffsetAndSize.Offset = CallstackAtOffset.Offset;
-						OffsetAndSize.Size = Callstacks.GetSerializedDataSizeForOffsetIndex(DifferenceCallstackOffsetIndex);
-						OutDiffMap.Add(OffsetAndSize);
-					}
-				}
-				LastDifferenceCallstackOffsetIndex = DifferenceCallstackOffsetIndex;
-			}
-		}
-	}
-
-	if (SourceSize < DestSize)
-	{
-		bIdentical = false;
-
-		// Add all the remaining callstacks to the diff map
-		for (int32 OffsetIndex = LastDifferenceCallstackOffsetIndex + 1; OffsetIndex < Callstacks.Num() && OutDiffMap.Num() < MaxDiffsToFind; ++OffsetIndex)
-		{
-			const FArchiveCallstacks::FCallstackAtOffset& CallstackAtOffset = Callstacks.GetCallstack(OffsetIndex);
-			// Compare against the size without start offset as all callstack offsets are absolute (from the merged header + exports file)
-			if (CallstackAtOffset.Offset < DestPackage.Size)
-			{
-				if (!CallstackAtOffset.bIgnore)
-				{
-					FArchiveDiffInfo OffsetAndSize;
-					OffsetAndSize.Offset = CallstackAtOffset.Offset;
-					OffsetAndSize.Size = Callstacks.GetSerializedDataSizeForOffsetIndex(OffsetIndex);
-					OutDiffMap.Add(OffsetAndSize);
-				}
-			}
-			else
-			{
-				break;
-			}
-		}
-	}
-	else if (SourceSize > DestSize)
-	{
-		bIdentical = false;
-	}
-	return bIdentical;
-}
-
-bool FArchiveStackTrace::GenerateDiffMap(const TCHAR* InFilename, int64 TotalHeaderSize, int32 MaxDiffsToFind, FArchiveDiffMap& OutDiffMap)
-{
-	TUniquePtr<uint8, UE::ArchiveStackTrace::FDeleteByFree> SourcePackageBytes;
-	FPackageData SourcePackage;
-	if (!LoadPackageIntoMemory(InFilename, SourcePackage, SourcePackageBytes))
-	{
-		return false;
-	}
-	return GenerateDiffMap(SourcePackage, TotalHeaderSize, MaxDiffsToFind, OutDiffMap);
-}
-
-bool FArchiveStackTrace::GenerateDiffMap(const FPackageData& SourcePackage, int64 TotalHeaderSize, int32 MaxDiffsToFind, FArchiveDiffMap& OutDiffMap)
-{
-	check(MaxDiffsToFind > 0);
-
-	bool bIdentical = true;
-	bool bHeaderIdentical = true;
-	bool bExportsIdentical = true;
-
-	FPackageData DestPackage;
-	DestPackage.Data = GetData();
-	DestPackage.Size = TotalSize();
-	DestPackage.HeaderSize = TotalHeaderSize;
-	DestPackage.StartOffset = 0;
-
-	{
-		FPackageData SourcePackageHeader = SourcePackage;
-		SourcePackageHeader.Size = SourcePackageHeader.HeaderSize;
-		SourcePackageHeader.HeaderSize = 0;
-		SourcePackageHeader.StartOffset = 0;
-
-		FPackageData DestPackageHeader = DestPackage;
-		DestPackageHeader.Size = TotalHeaderSize;
-		DestPackageHeader.HeaderSize = 0;
-		DestPackageHeader.StartOffset = 0;
-
-		bHeaderIdentical = FArchiveStackTraceWriter::GenerateDiffMap(SourcePackageHeader, DestPackageHeader, Callstacks, MaxDiffsToFind, OutDiffMap);
-	}
-
-	{
-		FPackageData SourcePackageExports = SourcePackage;
-		SourcePackageExports.HeaderSize = 0;
-		SourcePackageExports.StartOffset = SourcePackage.HeaderSize;
-
-		FPackageData DestPackageExports = DestPackage;
-		DestPackageExports.HeaderSize = 0;
-		DestPackageExports.StartOffset = TotalHeaderSize;
-
-		bExportsIdentical = FArchiveStackTraceWriter::GenerateDiffMap(SourcePackageExports, DestPackageExports, Callstacks, MaxDiffsToFind, OutDiffMap);
-	}
-
-	bIdentical = bHeaderIdentical && bExportsIdentical;
-
-	return bIdentical;
-}
-
-bool FArchiveStackTrace::IsIdentical(const TCHAR* InFilename, int64 BufferSize, const uint8* BufferData)
-{
-	TUniquePtr<uint8, UE::ArchiveStackTrace::FDeleteByFree> SourcePackageBytes;
-	FPackageData SourcePackage;
-	if (!LoadPackageIntoMemory(InFilename, SourcePackage, SourcePackageBytes))
-	{
-		return false;
-	}
-
-	return IsIdentical(SourcePackage, BufferSize, BufferData);
-}
-
-bool FArchiveStackTrace::IsIdentical(const FPackageData& SourcePackage, int64 BufferSize, const uint8* BufferData)
-{
-	bool bIdentical = false;
-	if (BufferSize == SourcePackage.Size)
-	{
-		bIdentical = (FMemory::Memcmp(SourcePackage.Data, BufferData, BufferSize) == 0);
-	}
-	else
-	{
-		bIdentical = false;
-	}
-	return bIdentical;
-}
-
-FLinkerLoad* FArchiveStackTraceWriter::CreateLinkerForPackage(FUObjectSerializeContext* LoadContext, const FString& InPackageName, const FString& InFilename, const FPackageData& PackageData)
-{
-	// First create a temp package to associate the linker with
-	UPackage* Package = FindObjectFast<UPackage>(nullptr, *InPackageName);
-	if (!Package)
-	{
-		Package = CreatePackage(*InPackageName);
-	}
-	// Create an archive for the linker. The linker will take ownership of it.
-	FLargeMemoryReader* PackageReader = new FLargeMemoryReader(PackageData.Data, PackageData.Size, ELargeMemoryReaderFlags::None, *InPackageName);	
-	FLinkerLoad* Linker = FLinkerLoad::CreateLinker(LoadContext, Package, FPackagePath::FromLocalPath(InFilename), LOAD_NoVerify, PackageReader);
-
-	if (Linker && Package)
-	{
-		Package->SetPackageFlags(PKG_ForDiffing);
-	}
-
-	return Linker;
-}
-
-static FString GetTableKey(const FLinkerLoad* Linker, const FObjectExport& Export)
-{
-	FName ClassName = Export.ClassIndex.IsNull() ? FName(NAME_Class) : Linker->ImpExp(Export.ClassIndex).ObjectName;
-	return FString::Printf(TEXT("%s %s.%s"),
-		*ClassName.ToString(),
-		!Export.OuterIndex.IsNull() ? *Linker->ImpExp(Export.OuterIndex).ObjectName.ToString() : *FPackageName::GetShortName(Linker->LinkerRoot),
-		*Export.ObjectName.ToString());
-}
-
-static FString GetTableKey(const FLinkerLoad* Linker, const FObjectImport& Import)
-{
-	return FString::Printf(TEXT("%s %s.%s"),
-		*Import.ClassName.ToString(),
-		!Import.OuterIndex.IsNull() ? *Linker->ImpExp(Import.OuterIndex).ObjectName.ToString() : TEXT("NULL"),
-		*Import.ObjectName.ToString());
-}
-
-static inline FString GetTableKey(const FLinkerLoad* Linker, const FName& Name)
-{
-	return *Name.ToString();
-}
-
-static inline FString GetTableKey(const FLinkerLoad* Linker, FNameEntryId Id)
-{
-	return FName::GetEntry(Id)->GetPlainNameString();
-}
-
-static inline FString GetTableKeyForIndex(const FLinkerLoad* Linker, FPackageIndex Index)
-{
-	if (Index.IsNull())
-	{
-		return TEXT("NULL");
-	}
-	else if (Index.IsExport())
-	{
-		return GetTableKey(Linker, Linker->Exp(Index));
-	}
-	else
-	{
-		return GetTableKey(Linker, Linker->Imp(Index));
-	}
-}
-
-static bool ComparePackageIndices(FLinkerLoad* SourceLinker, FLinkerLoad* DestLinker, const FPackageIndex& SourceIndex, const FPackageIndex& DestIndex);
-
-static bool CompareTableItem(FLinkerLoad* SourceLinker, FLinkerLoad* DestLinker, const FName& SourceName, const FName& DestName)
-{
-	return SourceName == DestName;
-}
-
-static bool CompareTableItem(FLinkerLoad* SourceLinker, FLinkerLoad* DestLinker, FNameEntryId SourceName, FNameEntryId DestName)
-{
-	return SourceName == DestName;
-}
-
-static FString ConvertItemToText(const FName& Name, FLinkerLoad* Linker)
-{
-	return Name.ToString();
-}
-
-static FString ConvertItemToText(FNameEntryId Id, FLinkerLoad* Linker)
-{
-	return FName::GetEntry(Id)->GetPlainNameString();
-}
-
-static bool CompareTableItem(FLinkerLoad* SourceLinker, FLinkerLoad* DestLinker, const FObjectImport& SourceImport, const FObjectImport& DestImport)
-{
-	if (SourceImport.ObjectName != DestImport.ObjectName ||
-		SourceImport.ClassName != DestImport.ClassName ||
-		SourceImport.ClassPackage != DestImport.ClassPackage ||
-		!ComparePackageIndices(SourceLinker, DestLinker, SourceImport.OuterIndex, DestImport.OuterIndex))
-	{
-		return false;
-	}
-	else
-	{
-		return true;
-	}
-}
-
-static FString ConvertItemToText(const FObjectImport& Import, FLinkerLoad* Linker)
-{
-	return FString::Printf(
-		TEXT("%s ClassPackage: %s"),
-		*GetTableKey(Linker, Import),
-		*Import.ClassPackage.ToString()
-	);
-}
-
-static bool CompareTableItem(FLinkerLoad* SourceLinker, FLinkerLoad* DestLinker, const FObjectExport& SourceExport, const FObjectExport& DestExport)
-{
-	if (SourceExport.ObjectName != DestExport.ObjectName ||
-		SourceExport.PackageFlags != DestExport.PackageFlags ||
-		SourceExport.ObjectFlags != DestExport.ObjectFlags ||
-		SourceExport.SerialSize != DestExport.SerialSize ||
-		SourceExport.bForcedExport != DestExport.bForcedExport ||
-		SourceExport.bNotForClient != DestExport.bNotForClient ||
-		SourceExport.bNotForServer != DestExport.bNotForServer ||
-		SourceExport.bNotAlwaysLoadedForEditorGame != DestExport.bNotAlwaysLoadedForEditorGame ||
-		SourceExport.bIsAsset != DestExport.bIsAsset ||
-		SourceExport.bIsInheritedInstance != DestExport.bIsInheritedInstance ||
-		SourceExport.bGeneratePublicHash != DestExport.bGeneratePublicHash ||
-		!ComparePackageIndices(SourceLinker, DestLinker, SourceExport.TemplateIndex, DestExport.TemplateIndex) ||
-		!ComparePackageIndices(SourceLinker, DestLinker, SourceExport.OuterIndex, DestExport.OuterIndex) ||
-		!ComparePackageIndices(SourceLinker, DestLinker, SourceExport.ClassIndex, DestExport.ClassIndex) ||
-		!ComparePackageIndices(SourceLinker, DestLinker, SourceExport.SuperIndex, DestExport.SuperIndex))
-	{
-		return false;
-	}
-	else
-	{
-		return true;
-	}
-}
-
-static bool IsImportMapIdentical(FLinkerLoad* SourceLinker, FLinkerLoad* DestLinker)
-{
-	bool bIdentical = (SourceLinker->ImportMap.Num() == DestLinker->ImportMap.Num());
-	if (bIdentical)
-	{
-		for (int32 ImportIndex = 0; ImportIndex < SourceLinker->ImportMap.Num(); ++ImportIndex)
-		{
-			if (!CompareTableItem(SourceLinker, DestLinker, SourceLinker->ImportMap[ImportIndex], DestLinker->ImportMap[ImportIndex]))
-			{
-				bIdentical = false;
-				break;
-			}
-		}
-	}
-	return bIdentical;
-}
-
-static bool ComparePackageIndices(FLinkerLoad* SourceLinker, FLinkerLoad* DestLinker, const FPackageIndex& SourceIndex, const FPackageIndex& DestIndex)
-{
-	if (SourceIndex.IsNull() && DestIndex.IsNull())
-	{
-		return true;
-	}
-
-	if (SourceIndex.IsExport() && DestIndex.IsExport())
-	{
-		int32 SourceArrayIndex = SourceIndex.ToExport();
-		int32 DestArrayIndex   = DestIndex  .ToExport();
-
-		if (!SourceLinker->ExportMap.IsValidIndex(SourceArrayIndex) || !DestLinker->ExportMap.IsValidIndex(DestArrayIndex))
-		{
-			UE_LOG(LogArchiveDiff, Warning, TEXT("Invalid export indices found, source: %d (of %d), dest: %d (of %d)"), SourceArrayIndex, SourceLinker->ExportMap.Num(), DestArrayIndex, DestLinker->ExportMap.Num());
-			return false;
-		}
-
-		const FObjectExport& SourceOuterExport = SourceLinker->Exp(SourceIndex);
-		const FObjectExport& DestOuterExport   = DestLinker  ->Exp(DestIndex);
-
-		FString SourceOuterExportKey = GetTableKey(SourceLinker, SourceOuterExport);
-		FString DestOuterExportKey   = GetTableKey(DestLinker,   DestOuterExport);
-
-		return SourceOuterExportKey == DestOuterExportKey;
-	}
-
-	if (SourceIndex.IsImport() && DestIndex.IsImport())
-	{
-		int32 SourceArrayIndex = SourceIndex.ToImport();
-		int32 DestArrayIndex   = DestIndex  .ToImport();
-
-		if (!SourceLinker->ImportMap.IsValidIndex(SourceArrayIndex) || !DestLinker->ImportMap.IsValidIndex(DestArrayIndex))
-		{
-			UE_LOG(LogArchiveDiff, Warning, TEXT("Invalid import indices found, source: %d (of %d), dest: %d (of %d)"), SourceArrayIndex, SourceLinker->ImportMap.Num(), DestArrayIndex, DestLinker->ImportMap.Num());
-			return false;
-		}
-
-		const FObjectImport& SourceOuterImport = SourceLinker->Imp(SourceIndex);
-		const FObjectImport& DestOuterImport   = DestLinker  ->Imp(DestIndex);
-
-		FString SourceOuterImportKey = GetTableKey(SourceLinker, SourceOuterImport);
-		FString DestOuterImportKey   = GetTableKey(DestLinker,   DestOuterImport);
-
-		return SourceOuterImportKey == DestOuterImportKey;
-	}
-
-	return false;
-}
-
-static FString ConvertItemToText(const FObjectExport& Export, FLinkerLoad* Linker)
-{
-	FName ClassName = Export.ClassIndex.IsNull() ? FName(NAME_Class) : Linker->ImpExp(Export.ClassIndex).ObjectName;
-	return FString::Printf(TEXT("%s Super: %s, Template: %s, Flags: %d, Size: %lld, PackageFlags: %d, ForcedExport: %d, NotForClient: %d, NotForServer: %d, NotAlwaysLoadedForEditorGame: %d, IsAsset: %d, IsInheritedInstance: %d, GeneratePublicHash: %d"),
-		*GetTableKey(Linker, Export),
-		*GetTableKeyForIndex(Linker, Export.SuperIndex),
-		*GetTableKeyForIndex(Linker, Export.TemplateIndex),
-		(int32)Export.ObjectFlags,
-		Export.SerialSize,
-		Export.PackageFlags,
-		Export.bForcedExport,
-		Export.bNotForClient,
-		Export.bNotForServer,
-		Export.bNotAlwaysLoadedForEditorGame,
-		Export.bIsAsset,
-		Export.bIsInheritedInstance,
-		Export.bGeneratePublicHash);
-}
-
-static bool IsExportMapIdentical(FLinkerLoad* SourceLinker, FLinkerLoad* DestLinker)
-{
-	bool bIdentical = (SourceLinker->ExportMap.Num() == DestLinker->ExportMap.Num());
-	if (bIdentical)
-	{
-		for (int32 ExportIndex = 0; ExportIndex < SourceLinker->ExportMap.Num(); ++ExportIndex)
-		{
-			if (!CompareTableItem(SourceLinker, DestLinker, SourceLinker->ExportMap[ExportIndex], DestLinker->ExportMap[ExportIndex]))
-			{
-				bIdentical = false;
-				break;
-			}
-		}
-	}
-	return bIdentical;
-}
-
-/** Structure that holds an item from the NameMap/ImportMap/ExportMap in a TSet for diffing */
-template <typename T>
-struct TTableItem
-{
-	/** The key generated for this item */
-	FString Key;
-	/** Pointer to the original item */
-	const T* Item;
-	/** Index in the original *Map (table). Only for information purposes. */
-	int32 Index;
-
-	TTableItem(FString&& InKey, const T* InItem, int32 InIndex)
-		: Key(MoveTemp(InKey))
-		, Item(InItem)
-		, Index(InIndex)
-	{
-	}
-
-	FORCENOINLINE friend uint32 GetTypeHash(const TTableItem& TableItem)
-	{
-		return GetTypeHash(TableItem.Key);
-	}
-
-	FORCENOINLINE friend bool operator==(const TTableItem& Lhs, const TTableItem& Rhs)
-	{
-		return Lhs.Key == Rhs.Key;
-	}
-};
-
-/** Dumps differences between Linker tables */
-template <typename T>
-static void DumpTableDifferences(
-	FLinkerLoad* SourceLinker, 
-	FLinkerLoad* DestLinker, 
-	TArray<T>& SourceTable, 
-	TArray<T>& DestTable,
-	const TCHAR* AssetFilename,
-	const TCHAR* ItemName,
-	const int32 MaxDiffsToLog
-)
-{
-#if !NO_LOGGING
-	const TCHAR* const LineTerminator = FDiffFormatHelper::Get().LineTerminator;
-	const TCHAR* const Indent = FDiffFormatHelper::Get().Indent;
-
-	FString HumanReadableString;
-	int32 LoggedDiffs = 0;
-	int32 NumDiffs = 0;
-
-	TSet<TTableItem<T>> SourceSet;
-	TSet<TTableItem<T>> DestSet;
-
-	SourceSet.Reserve(SourceTable.Num());
-	DestSet.Reserve(DestTable.Num());
-
-	for (int32 Index = 0; Index < SourceTable.Num(); ++Index)
-	{
-		const T& Item = SourceTable[Index];
-		SourceSet.Add(TTableItem<T>(GetTableKey(SourceLinker, Item), &Item, Index));
-	}
-	for (int32 Index = 0; Index < DestTable.Num(); ++Index)
-	{
-		const T& Item = DestTable[Index];
-		DestSet.Add(TTableItem<T>(GetTableKey(DestLinker, Item), &Item, Index));
-	}
-
-	// Determine the list of items removed from the source package and added to the dest package
-	TSet<TTableItem<T>> RemovedItems = SourceSet.Difference(DestSet);
-	TSet<TTableItem<T>> AddedItems   = DestSet.Difference(SourceSet);
-
-	// Add changed items as added-and-removed
-	for (const TTableItem<T>& ChangedSourceItem : SourceSet)
-	{
-		if (const TTableItem<T>* ChangedDestItem = DestSet.Find(ChangedSourceItem))
-		{
-			if (!CompareTableItem(SourceLinker, DestLinker, *ChangedSourceItem.Item, *ChangedDestItem->Item))
-			{
-				RemovedItems.Add(ChangedSourceItem);
-				AddedItems  .Add(*ChangedDestItem);
-			}
-		}
-	}
-
-	// Sort all additions and removals by index
-	RemovedItems.Sort([](const TTableItem<T>& Lhs, const TTableItem<T>& Rhs){ return Lhs.Index < Rhs.Index; });
-	AddedItems.  Sort([](const TTableItem<T>& Lhs, const TTableItem<T>& Rhs){ return Lhs.Index < Rhs.Index; });
-
-	// Dump all changes
-	for (const TTableItem<T>& RemovedItem : RemovedItems)
-	{
-		HumanReadableString += Indent;
-		HumanReadableString += FString::Printf(TEXT("-[%d] %s"), RemovedItem.Index, *ConvertItemToText(*RemovedItem.Item, SourceLinker));
-		HumanReadableString += LineTerminator;
-	}
-	for (const TTableItem<T>& AddedItem : AddedItems)
-	{
-		HumanReadableString += Indent;
-		HumanReadableString += FString::Printf(TEXT("+[%d] %s"), AddedItem.Index, *ConvertItemToText(*AddedItem.Item, DestLinker));
-		HumanReadableString += LineTerminator;
-	}
-
-	// For now just log everything out. When this becomes too spammy, respect the MaxDiffsToLog parameter
-	NumDiffs = RemovedItems.Num() + AddedItems.Num();
-	LoggedDiffs = NumDiffs;
-
-	if (NumDiffs > LoggedDiffs)
-	{
-		HumanReadableString += Indent;
-		HumanReadableString += FString::Printf(TEXT("+ %d differences not logged."), (NumDiffs - LoggedDiffs));
-		HumanReadableString += LineTerminator;
-	}
-
-	UE_LOG(
-		LogArchiveDiff,
-		Warning,
-		TEXT("%s: %sMap is different (%d %ss in source package vs %d %ss in dest package):%s%s"),		
-		AssetFilename,
-		ItemName,
-		SourceTable.Num(),
-		ItemName,
-		DestTable.Num(),
-		ItemName,
-		LineTerminator,
-		*HumanReadableString);
-#endif // !NO_LOGGING
-}
-
-static void DumpPackageHeaderDiffs_LinkerLoad(
-	const FArchiveStackTraceWriter::FPackageData& SourcePackage,
-	const FArchiveStackTraceWriter::FPackageData& DestPackage,
-	const FString& AssetFilename,
-	const int32 MaxDiffsToLog)
-{
-#if !NO_LOGGING
-	FString AssetPathName = FPaths::Combine(*FPaths::GetPath(AssetFilename.Mid(AssetFilename.Find(TEXT(":"), ESearchCase::CaseSensitive) + 1)), *FPaths::GetBaseFilename(AssetFilename));
-	// The root directory could have a period in it (d:/Release5.0/EngineTest/Saved/Cooked),
-	// which is not a valid character for a LongPackageName. Remove it.
-	for (TCHAR c : FStringView(INVALID_LONGPACKAGE_CHARACTERS))
-	{
-		AssetPathName.ReplaceCharInline(c, TEXT('_'), ESearchCase::CaseSensitive);
-	}
-	FString SourceAssetPackageName = FPaths::Combine(TEXT("/Memory"), TEXT("/SourceForDiff"), *AssetPathName);
-	FString DestAssetPackageName = FPaths::Combine(TEXT("/Memory"), TEXT("/DestForDiff"), *AssetPathName);
-	check(FPackageName::IsValidLongPackageName(SourceAssetPackageName, true /* bIncludeReadOnlyRoots */));
-
-	TGuardValue<bool> GuardIsSavingPackage(GIsSavingPackage, false);
-	TGuardValue<int32> GuardAllowUnversionedContentInEditor(GAllowUnversionedContentInEditor, 1);
-	TGuardValue<int32> GuardAllowCookedDataInEditorBuilds(GAllowCookedDataInEditorBuilds, 1);
-
-	FLinkerLoad* SourceLinker = nullptr;
-	FLinkerLoad* DestLinker = nullptr;
-	// Create linkers. Note there's no need to clean them up here since they will be removed by the package associated with them
-	{
-		TRefCountPtr<FUObjectSerializeContext> LinkerLoadContext(FUObjectThreadContext::Get().GetSerializeContext());
-		BeginLoad(LinkerLoadContext);
-		SourceLinker = FArchiveStackTraceWriter::CreateLinkerForPackage(LinkerLoadContext, SourceAssetPackageName, AssetFilename, SourcePackage);
-		EndLoad(LinkerLoadContext);
-	}
-
-	{
-		TRefCountPtr<FUObjectSerializeContext> LinkerLoadContext(FUObjectThreadContext::Get().GetSerializeContext());
-		BeginLoad(LinkerLoadContext);
-		DestLinker = FArchiveStackTraceWriter::CreateLinkerForPackage(LinkerLoadContext, DestAssetPackageName, AssetFilename, DestPackage);
-		EndLoad(LinkerLoadContext);
-	}
-
-	if (SourceLinker && DestLinker)
-	{
-		if (SourceLinker->NameMap != DestLinker->NameMap)
-		{
-			DumpTableDifferences<FNameEntryId>(SourceLinker, DestLinker, SourceLinker->NameMap, DestLinker->NameMap, *AssetFilename, TEXT("Name"), MaxDiffsToLog);
-		}
-
-		if (!IsImportMapIdentical(SourceLinker, DestLinker))
-		{
-			DumpTableDifferences<FObjectImport>(SourceLinker, DestLinker, SourceLinker->ImportMap, DestLinker->ImportMap, *AssetFilename, TEXT("Import"), MaxDiffsToLog);
-		}
-
-		if (!IsExportMapIdentical(SourceLinker, DestLinker))
-		{
-			DumpTableDifferences<FObjectExport>(SourceLinker, DestLinker, SourceLinker->ExportMap, DestLinker->ExportMap, *AssetFilename, TEXT("Export"), MaxDiffsToLog);
-		}
-	}
-
-	if (SourceLinker)
-	{
-		UE::ArchiveStackTrace::ForceKillPackageAndLinker(SourceLinker);
-	}
-	if (DestLinker)
-	{
-		UE::ArchiveStackTrace::ForceKillPackageAndLinker(DestLinker);
-	}
-#endif // !NO_LOGGING
-}
-
-static void DumpPackageHeaderDiffs_ZenPackage(
-	const FArchiveStackTraceWriter::FPackageData& SourcePackage,
-	const FArchiveStackTraceWriter::FPackageData& DestPackage,
-	const FString& AssetFilename,
-	const int32 MaxDiffsToLog)
-{
-#if !NO_LOGGING
-	// TODO: Fill in detailed diffing of Zen Package Summary
-#endif // !NO_LOGGING
-}
-
-void FArchiveStackTraceWriter::DumpPackageHeaderDiffs(
-	const FPackageData& SourcePackage,
-	const FPackageData& DestPackage,
-	const FString& AssetFilename,
-	const int32 MaxDiffsToLog,
-	const EPackageHeaderFormat PackageHeaderFormat /* = EPackageHeaderFormat::PackageFileSummary */)
-{
-#if !NO_LOGGING
-	switch (PackageHeaderFormat)
-	{
-	case EPackageHeaderFormat::PackageFileSummary:
-		DumpPackageHeaderDiffs_LinkerLoad(SourcePackage, DestPackage, AssetFilename, MaxDiffsToLog);
-		break;
-	case EPackageHeaderFormat::ZenPackageSummary:
-		DumpPackageHeaderDiffs_ZenPackage(SourcePackage, DestPackage, AssetFilename, MaxDiffsToLog);
-		break;
-	default:
-		unimplemented();
-	}
-#endif // !NO_LOGGING
-}
-
-PRAGMA_ENABLE_DEPRECATION_WARNINGS;
-
 FArchiveStackTraceReader::FSerializeData::FSerializeData(int64 InOffset, int64 InSize, UObject* InObject, FProperty* InProperty)
 : Offset(InOffset)
 , Size(InSize)
diff --git a/Engine/Source/Runtime/CoreUObject/Private/Serialization/BulkData.cpp b/Engine/Source/Runtime/CoreUObject/Private/Serialization/BulkData.cpp
--- a/Engine/Source/Runtime/CoreUObject/Private/Serialization/BulkData.cpp
+++ b/Engine/Source/Runtime/CoreUObject/Private/Serialization/BulkData.cpp
@@ -809,13 +809,6 @@ bool FBulkData::IsAsyncLoadingComplete() const
 	return true;
 }
 
-bool FBulkData::IsAvailableForUse() const
-{
-PRAGMA_DISABLE_DEPRECATION_WARNINGS
-	return BulkMeta.HasAnyFlags(BULKDATA_Unused) == false;
-PRAGMA_ENABLE_DEPRECATION_WARNINGS
-}
-
 /*-----------------------------------------------------------------------------
 	Data retrieval and manipulation.
 -----------------------------------------------------------------------------*/
diff --git a/Engine/Source/Runtime/CoreUObject/Private/UObject/Obj.cpp b/Engine/Source/Runtime/CoreUObject/Private/UObject/Obj.cpp
--- a/Engine/Source/Runtime/CoreUObject/Private/UObject/Obj.cpp
+++ b/Engine/Source/Runtime/CoreUObject/Private/UObject/Obj.cpp
@@ -6065,32 +6065,9 @@ void UObject::PreDestroyFromReplication()
 	Data Validation.
 -----------------------------------------------------------------------------*/
 
-EDataValidationResult UObject::IsDataValid(TArray<FText>& ValidationErrors)
-{
-	return EDataValidationResult::NotValidated;
-}
-
-EDataValidationResult UObject::IsDataValid(FDataValidationContext& Context)
-{
-	// Call the const version
-	return const_cast<const UObject*>(this)->IsDataValid(Context);
-}
-
 EDataValidationResult UObject::IsDataValid(FDataValidationContext& Context) const
 {
-	TArray<FText> ValidationErrors;
-
-PRAGMA_DISABLE_DEPRECATION_WARNINGS
-	// Call the old deprecated TArray<FText> version
-	const EDataValidationResult Result = const_cast<UObject*>(this)->IsDataValid(ValidationErrors);
-PRAGMA_ENABLE_DEPRECATION_WARNINGS
-
-	for (const FText& Text : ValidationErrors)
-	{
-		Context.AddError(Text);
-	}
-
-	return Result;
+	return EDataValidationResult::NotValidated;
 }
 
 #endif // WITH_EDITOR
diff --git a/Engine/Source/Runtime/CoreUObject/Public/Misc/AssetRegistryInterface.h b/Engine/Source/Runtime/CoreUObject/Public/Misc/AssetRegistryInterface.h
--- a/Engine/Source/Runtime/CoreUObject/Public/Misc/AssetRegistryInterface.h
+++ b/Engine/Source/Runtime/CoreUObject/Public/Misc/AssetRegistryInterface.h
@@ -22,40 +22,6 @@ struct FArchiveCookContext;
 struct FARFilter;
 struct FAssetData;
 
-namespace EAssetRegistryDependencyType
-{
-	enum UE_DEPRECATED(5.3, "No longer used") Type
-	{
-		// Dependencies which don't need to be loaded for the object to be used (i.e. soft object paths)
-		Soft = 0x01,
-
-		// Dependencies which are required for correct usage of the source asset, and must be loaded at the same time
-		Hard = 0x02,
-
-		// References to specific SearchableNames inside a package
-		SearchableName = 0x04,
-
-		// Indirect management references, these are set through recursion for Primary Assets that manage packages or other primary assets
-		SoftManage = 0x08,
-
-		// Reference that says one object directly manages another object, set when Primary Assets manage things explicitly
-		HardManage = 0x10,
-
-		// Note: Also update FAssetRegistryDependencyOptions when adding more flags
-	};
-
-	PRAGMA_DISABLE_DEPRECATION_WARNINGS;
-	UE_DEPRECATED(5.3, "No longer used")
-	inline static const Type None = (Type)(0);
-	UE_DEPRECATED(5.3, "No longer used")
-	inline static const Type All = (Type)(Soft | Hard | SearchableName | SoftManage | HardManage);
-	UE_DEPRECATED(5.3, "No longer used")
-	inline static const Type Packages = (Type)(Soft | Hard);
-	UE_DEPRECATED(5.3, "No longer used")
-	inline static const Type Manage = (Type)(SoftManage | HardManage);
-	PRAGMA_ENABLE_DEPRECATION_WARNINGS;
-}
-
 class IAssetRegistry;
 class UAssetRegistryImpl;
 
diff --git a/Engine/Source/Runtime/CoreUObject/Public/Serialization/ArchiveStackTrace.h b/Engine/Source/Runtime/CoreUObject/Public/Serialization/ArchiveStackTrace.h
--- a/Engine/Source/Runtime/CoreUObject/Public/Serialization/ArchiveStackTrace.h
+++ b/Engine/Source/Runtime/CoreUObject/Public/Serialization/ArchiveStackTrace.h
@@ -123,322 +123,6 @@ FArchiveDiffInfo
 	}
 };
 
-class 
-UE_DEPRECATED(5.3, "FArchiveStackTrace was only used by DiffPackageWriter, and has been moved into a private helper class. Contact Epic if you need this class for another reason.")
-FArchiveDiffMap : public TArray<FArchiveDiffInfo>
-{
-public:
-	bool ContainsOffset(int64 Offset) const;
-};
-
-class
-UE_DEPRECATED(5.3, "FArchiveStackTrace was only used by DiffPackageWriter, and has been moved into a private helper class. Contact Epic if you need this class for another reason.")
-FArchiveCallstacks
-{
-public:
-	/** Offset and callstack pair */
-	struct FCallstackAtOffset
-	{
-		/** Offset of a Serialize call */
-		int64 Offset = -1;
-		/** Callstack CRC for the Serialize call */
-		uint32 Callstack = 0;
-		/** Collected inside of skip scope */
-		bool bIgnore = false;
-	};
-
-	/** Struct to hold the actual Serialize call callstack and any associated data */
-	struct FCallstackData
-	{
-		/** Full callstack */
-		TUniquePtr<ANSICHAR[]> Callstack;
-		/** Full name of the currently serialized object */
-		FString SerializedObjectName;
-		/** The currently serialized property */
-		FProperty* SerializedProp = nullptr;
-		/** Name of the currently serialized property */
-		FString SerializedPropertyName;
-
-		FCallstackData() = default;
-		FCallstackData(TUniquePtr<ANSICHAR[]>&& InCallstack, UObject* InSerializedObject, FProperty* InSerializedProperty);
-		FCallstackData(FCallstackData&&) = default;
-		FCallstackData(const FCallstackData&) = delete;
-
-		FCallstackData& operator=(FCallstackData&&) = default;
-		FCallstackData& operator=(const FCallstackData&) = delete;
-
-		/** Converts the callstack and associated data to human readable string */
-		FString ToString(const TCHAR* CallstackCutoffText) const;
-
-		/** Clone the callstack data */
-		FCallstackData Clone() const;
-	};
-	
-	COREUOBJECT_API explicit FArchiveCallstacks(UObject* InAsset);
-
-	/** Returns the asset class name. */
-	COREUOBJECT_API FName GetAssetClass() const;
-
-	/** Returns the total number of callstacks. */
-	int32 Num() const
-	{
-		return CallstackAtOffsetMap.Num();
-	}
-
-	/** Capture and append the current callstack. */
-	COREUOBJECT_API void Add(
-		int64 Offset,
-		int64 Length,
-		UObject* SerializedObject,
-		FProperty* SerializedProperty,
-		TArrayView<const FName> DebugDataStack,
-		bool bIsCollectingCallstacks,
-		bool bCollectCurrentCallstack,
-		int32 StackIgnoreCount);
-
-	/** Append other callstacks. */
-	COREUOBJECT_API void Append(const FArchiveCallstacks& Other, int64 Offset = 0);
-
-	/** Finds a callstack associated with data at the specified offset */
-	COREUOBJECT_API int32 GetCallstackIndexAtOffset(int64 Offset, int32 MinOffsetIndex = 0) const;
-
-	/** Finds a callstack associated with data at the specified offset */
-	const FCallstackAtOffset& GetCallstack(int32 CallstackIndex) const
-	{
-		return CallstackAtOffsetMap[CallstackIndex];
-	}
-	
-	const FCallstackData& GetCallstackData(const FCallstackAtOffset& CallstackOffset) const
-	{
-		return UniqueCallstacks[CallstackOffset.Callstack];
-	}
-
-	/** Returns the size of serialized data at the specified offset. */
-	int64 GetSerializedDataSizeForOffsetIndex(int32 InOffsetIndex) const
-	{
-		if (InOffsetIndex < CallstackAtOffsetMap.Num() - 1)
-		{
-			return CallstackAtOffsetMap[InOffsetIndex + 1].Offset - CallstackAtOffsetMap[InOffsetIndex].Offset;
-		}
-		else
-		{
-			return TotalSize - CallstackAtOffsetMap[InOffsetIndex].Offset;
-		}
-	}
-
-	/** Returns total serialized bytes. */
-	int64 TotalCapturedSize() const
-	{
-		return TotalSize;
-	}
-
-private:
-	/** Adds a unique callstack to UniqueCallstacks map */
-	COREUOBJECT_API ANSICHAR* AddUniqueCallstack(bool bIsCollectingCallstacks, UObject* SerializedObject, FProperty* SerializedProperty, uint32& OutCallstackCRC);
-
-	/** The asset being serialized */
-	UObject* Asset;
-	/** List of offsets and their respective callstacks */
-	TArray<FCallstackAtOffset> CallstackAtOffsetMap;
-	/** Contains all unique callstacks for all Serialize calls */
-	TMap<uint32, FCallstackData> UniqueCallstacks;
-	/** Optimizes callstack comparison. If true the current and the last callstack should be compared as it may have changed */
-	bool bCallstacksDirty;
-	/** Maximum size of the stack trace */
-	const SIZE_T StackTraceSize;
-	/** Buffer for getting the current stack trace */
-	TUniquePtr<ANSICHAR[]> StackTrace;
-	/** Callstack associated with the previous Serialize call */
-	ANSICHAR* LastSerializeCallstack;
-	/** Total serialized bytes */
-	int64 TotalSize;
-};
-
-class
-UE_DEPRECATED(5.3, "FArchiveStackTrace was only used by DiffPackageWriter, and has been moved into a private helper class. Contact Epic if you need this class for another reason.")
-FArchiveStackTraceWriter
-	: public FArchiveProxy
-{
-public:
-	PRAGMA_DISABLE_DEPRECATION_WARNINGS;
-
-	COREUOBJECT_API FArchiveStackTraceWriter(
-		FArchive& InInner,
-		FArchiveCallstacks& InCallstacks,
-		const FArchiveDiffMap* InDiffMap = nullptr,
-		int64 InDiffMapStartOffset = 0);
-
-	COREUOBJECT_API virtual ~FArchiveStackTraceWriter() override;
-
-	COREUOBJECT_API FORCENOINLINE virtual void Serialize(void* Data, int64 Length) override; // FORCENOINLINE so it can be counted during StackTrace
-	
-#if WITH_EDITOR
-	virtual void PushDebugDataString(const FName& DebugData) override
-	{
-		DebugDataStack.Push(DebugData);
-	}
-
-	virtual void PopDebugDataString() override
-	{
-		DebugDataStack.Pop();
-	}
-#endif
-
-	const FArchiveDiffMap& GetDiffMap() const
-	{
-		static FArchiveDiffMap Empty;
-		return DiffMap != nullptr ? *DiffMap : Empty;
-	}
-
-	void SetDisableInnerArchive(bool bDisable)
-	{
-		bInnerArchiveDisabled = bDisable;
-	}
-
-	int32 GetStackIgnoreCount() const { return StackIgnoreCount; }
-	void SetStackIgnoreCount(const int32 IgnoreCount) { StackIgnoreCount = IgnoreCount; }
-
-	struct FPackageData
-	{
-		uint8* Data = nullptr;
-		int64 Size = 0;
-		int64 HeaderSize = 0;
-		int64 StartOffset = 0;
-	};
-
-	using EPackageHeaderFormat = ICookedPackageWriter::EPackageHeaderFormat;
-
-	/** Compares two packages and logs the differences and calltacks. */
-	static COREUOBJECT_API void Compare(
-		const FPackageData& SourcePackage,
-		const FPackageData& DestPackage,
-		const FArchiveCallstacks& Callstacks,
-		const FArchiveDiffMap& DiffMap,
-		const TCHAR* AssetFilename,
-		const TCHAR* CallstackCutoffText,
-		const int64 MaxDiffsToLog,
-		int32& InOutDiffsLogged,
-		TMap<FName, FArchiveDiffStats>& OutStats,
-		bool bSuppressLogging = false);
-
-	/** Creates map with mismatching callstacks. */
-	static COREUOBJECT_API bool GenerateDiffMap(
-		const FPackageData& SourcePackage,
-		const FPackageData& DestPackage,
-		const FArchiveCallstacks& Callstacks,
-		int32 MaxDiffsToFind,
-		FArchiveDiffMap& OutDiffMap);
-
-	/** Logs any mismatching header data. */
-	static COREUOBJECT_API void DumpPackageHeaderDiffs(
-		const FPackageData& SourcePackage,
-		const FPackageData& DestPackage,
-		const FString& AssetFilename,
-		const int32 MaxDiffsToLog,
-		const EPackageHeaderFormat PackageHeaderFormat = EPackageHeaderFormat::PackageFileSummary);
-
-	/** Returns a new linker for loading the specified package. */
-	static COREUOBJECT_API FLinkerLoad* CreateLinkerForPackage(
-		FUObjectSerializeContext* LoadContext,
-		const FString& InPackageName,
-		const FString& InFilename,
-		const FPackageData& PackageData);
-
-private:
-	FArchiveCallstacks& Callstacks;
-	const FArchiveDiffMap* DiffMap;
-#if WITH_EDITOR
-	TArray<FName> DebugDataStack;
-#endif
-	int64 DiffMapOffset;
-	int32 StackIgnoreCount = 2;
-	bool bInnerArchiveDisabled;
-
-PRAGMA_ENABLE_DEPRECATION_WARNINGS;
-
-};
-
-class
-UE_DEPRECATED(5.3, "FArchiveStackTrace was only used by DiffPackageWriter, and has been moved into a private helper class. Contact Epic if you need this class for another reason.")
-FArchiveStackTraceMemoryWriter final
-	: public FLargeMemoryWriter
-{
-public:
-
-PRAGMA_DISABLE_DEPRECATION_WARNINGS;
-
-	COREUOBJECT_API FArchiveStackTraceMemoryWriter(
-		FArchiveCallstacks& Callstacks,
-		const FArchiveDiffMap* DiffMap = nullptr,
-		const int64 DiffMapOffset = 0,
-		const int64 PreAllocateBytes = 0,
-		bool bIsPersistent = false,
-		const TCHAR* Filename = nullptr);
-
-	COREUOBJECT_API FORCENOINLINE virtual void Serialize(void* Memory, int64 Length) override; // FORCENOINLINE so it can be counted during StackTrace
-
-private:
-	FArchiveStackTraceWriter StackTraceWriter;
-
-PRAGMA_ENABLE_DEPRECATION_WARNINGS;
-};
-
-/**
- * Archive that stores a callstack for each of the Serialize calls and has the ability to compare itself to an existing
- * package on disk and dump all the differences to log.
- */
-class
-UE_DEPRECATED(5.3, "FArchiveStackTrace was only used by DiffPackageWriter, and has been moved into a private helper class. Contact Epic if you need this class for another reason.")
-FArchiveStackTrace
-	: public FLargeMemoryWriter
-{
-public:
-	PRAGMA_DISABLE_DEPRECATION_WARNINGS;
-
-	using FPackageData = FArchiveStackTraceWriter::FPackageData;
-
-	COREUOBJECT_API FArchiveStackTrace(UObject* InAsset, const TCHAR* InFilename, bool bInCollectCallstacks = true, const FArchiveDiffMap* InDiffMap = nullptr);
-	COREUOBJECT_API virtual ~FArchiveStackTrace();
-
-	FArchiveCallstacks& GetCallstacks()
-	{
-		return Callstacks;
-	}
-
-	const FArchiveCallstacks& GetCallstacks() const
-	{
-		return Callstacks;
-	}
-
-	COREUOBJECT_API FORCENOINLINE virtual void Serialize(void* Memory, int64 Length) override; // FORCENOINLINE so it can be counted during StackTrace
-
-	/** Compares this archive with the given bytes from disk or FPackageData. Dumps all differences to log. */
-	COREUOBJECT_API void CompareWith(const TCHAR* InFilename, const int64 TotalHeaderSize, const TCHAR* CallstackCutoffText,
-		const int32 MaxDiffsToLog, TMap<FName, FArchiveDiffStats>& OutStats);
-	COREUOBJECT_API void CompareWith(const FPackageData& SourcePackage, const TCHAR* FileDisplayName, const int64 TotalHeaderSize,
-		const TCHAR* CallstackCutoffText, const int32 MaxDiffsToLog, TMap<FName, FArchiveDiffStats>& OutStats,
-		const FArchiveStackTraceWriter::EPackageHeaderFormat PackageHeaderFormat = FArchiveStackTraceWriter::EPackageHeaderFormat::PackageFileSummary);
-
-	/** Generates a map of all differences between this archive and the given bytes from disk or FPackageData. */
-	COREUOBJECT_API bool GenerateDiffMap(const TCHAR* InFilename, int64 TotalHeaderSize, int32 MaxDiffsToFind, FArchiveDiffMap& OutDiffMap);
-	COREUOBJECT_API bool GenerateDiffMap(const FPackageData& SourcePackage, int64 TotalHeaderSize, int32 MaxDiffsToFind,
-		FArchiveDiffMap& OutDiffMap);
-
-	/** Compares the provided buffer with the given bytes from disk or FPackageData. */
-	static COREUOBJECT_API bool IsIdentical(const TCHAR* InFilename, int64 BufferSize, const uint8* BufferData);
-	static COREUOBJECT_API bool IsIdentical(const FPackageData& SourcePackage, int64 BufferSize, const uint8* BufferData);
-
-	/** Helper function to load package contents into memory. Supports EDL packages. */
-	static COREUOBJECT_API bool LoadPackageIntoMemory(const TCHAR* InFilename, FPackageData& OutPackageData,
-		TUniquePtr<uint8, UE::ArchiveStackTrace::FDeleteByFree>& OutLoadedBytes);
-
-private:
-	FArchiveCallstacks Callstacks;
-	FArchiveStackTraceWriter StackTraceWriter;
-
-PRAGMA_ENABLE_DEPRECATION_WARNINGS;
-};
-
 class FArchiveStackTraceReader : public FLargeMemoryReader
 {
 public:
diff --git a/Engine/Source/Runtime/CoreUObject/Public/Serialization/BulkData.h b/Engine/Source/Runtime/CoreUObject/Public/Serialization/BulkData.h
--- a/Engine/Source/Runtime/CoreUObject/Public/Serialization/BulkData.h
+++ b/Engine/Source/Runtime/CoreUObject/Public/Serialization/BulkData.h
@@ -618,13 +618,6 @@ public:
 	UE_DEPRECATED(5.5, "No longer provides any functionality (assume that this will always return true")
 	COREUOBJECT_API bool IsAsyncLoadingComplete() const;
 
-	/**
-	* Returns whether this bulk data is used
-	* @return true if BULKDATA_Unused is not set
-	*/
-	UE_DEPRECATED(5.3, "The feature is being removed, it is assumed that all bulkdata is available for use")
-	COREUOBJECT_API bool IsAvailableForUse() const;
-
 	/**
 	* Returns whether this bulk data represents optional data or not
 	* @return true if BULKDATA_OptionalPayload is set
diff --git a/Engine/Source/Runtime/CoreUObject/Public/StructUtils/StructArrayView.h b/Engine/Source/Runtime/CoreUObject/Public/StructUtils/StructArrayView.h
--- a/Engine/Source/Runtime/CoreUObject/Public/StructUtils/StructArrayView.h
+++ b/Engine/Source/Runtime/CoreUObject/Public/StructUtils/StructArrayView.h
@@ -55,18 +55,6 @@ struct FStructArrayView
 		check(InData != nullptr || InCount == 0);
 	}
 
-	UE_DEPRECATED(5.3, "Use constructor that takes a reference instead")
-	FStructArrayView(const UScriptStruct* InScriptStruct, void* InData, const int32 InCount)
-		: DataPtr(InData)
-		, ScriptStruct(InScriptStruct)
-		, ElementSize(InScriptStruct->GetStructureSize())
-		, ArrayNum(InCount)
-	{
-		check(InScriptStruct);
-		check(InCount >= 0);
-		check(InData != nullptr || InCount == 0);
-	}
-
 	/**
 	 * Checks array invariants: if array size is greater than or equal to zero.
 	 */
@@ -144,9 +132,6 @@ struct FStructArrayView
 		return (uint8*)DataPtr + ((SIZE_T)Index * ElementSize); 
 	}
 
-	UE_DEPRECATED(5.3, "Use GetDataAt() instead")
-	void* GetMutableDataAt(const int32 Index) const { return GetDataAt(Index); }
-
 	/**
 	 * Helper function returning the size of the inner type.
 	 *
@@ -157,9 +142,6 @@ struct FStructArrayView
 		return ElementSize;
 	}
 
-	UE_DEPRECATED(5.3, "Removed to bring in to line with ArrayView, use GetTypeSize() instead")
-	SIZE_T GetElementSize() const { return ElementSize; }
-
 	/**
 	 * Returns number of elements in array.
 	 *
@@ -167,12 +149,6 @@ struct FStructArrayView
 	 */
 	int32 Num() const { return ArrayNum; }
 
-	UE_DEPRECATED(5.3, "Use GetScriptStruct() instead")
-	const UScriptStruct& GetElementType() const { check(ScriptStruct); return *ScriptStruct; }
-
-	UE_DEPRECATED(5.3, "Use GetScriptStruct() instead")
-	const UScriptStruct& GetFragmentType() const { check(ScriptStruct);  return *ScriptStruct; }
-	
 	/** Returns struct type. */
 	const UScriptStruct* GetScriptStruct() const { return ScriptStruct; }
 
@@ -203,36 +179,6 @@ struct FStructArrayView
 		return UE::StructUtils::GetStructRef<T>(ScriptStruct, GetDataAt(Index)); // GetDataAt() calls RangeCheck().
 	}
 
-	template<typename T>
-	UE_DEPRECATED(5.3, "Use GetAt() instead!")
-	const T& GetElementAt(const int32 Index) const
-	{
-		return *((T*)GetDataAt(Index));
-	}
-
-	template<typename T>
-	UE_DEPRECATED(5.3, "Use GetAt() instead!")
-	const T& GetElementAtChecked(const int32 Index) const
-	{
-		check(TBaseStructure<T>::Get() == &ScriptStruct);
-		return *((T*)GetDataAt(Index));
-	}
-
-	template<typename T>
-	UE_DEPRECATED(5.3, "Use GetAt() instead!")
-	T& GetMutableElementAt(const int32 Index) const
-	{
-		return *((T*)GetDataAt(Index));
-	}
-
-	template<typename T>
-	UE_DEPRECATED(5.3, "Use GetAt() instead!")
-	T& GetMutableElementAtChecked(const int32 Index) const
-	{
-		check(TBaseStructure<T>::Get() == &ScriptStruct);
-		return *((T*)GetDataAt(Index));
-	}
-
 	/**
 	 * Array bracket operator. Returns FStructView to element at given index. Index must be valid.
 	 *
diff --git a/Engine/Source/Runtime/CoreUObject/Public/UObject/Object.h b/Engine/Source/Runtime/CoreUObject/Public/UObject/Object.h
--- a/Engine/Source/Runtime/CoreUObject/Public/UObject/Object.h
+++ b/Engine/Source/Runtime/CoreUObject/Public/UObject/Object.h
@@ -422,9 +422,6 @@ public:
 	/** After a critical error, perform any mission-critical cleanup, such as restoring the video mode or releasing hardware resources. */
 	virtual void ShutdownAfterError() {}
 
-	UE_DEPRECATED(5.3, "Matinee has been removed in 5.0, nothing is calling PostInterpChange anymore")
-	virtual void PostInterpChange(FProperty* PropertyThatChanged) {}
-
 #if WITH_EDITOR
 	/** 
 	 * This is called when a property is about to be modified externally
@@ -1100,19 +1097,6 @@ public:
 	 *         the rules. Returns NotValidated if no rules are set for this object.
 	 */
 	COREUOBJECT_API virtual EDataValidationResult IsDataValid(class FDataValidationContext& Context) const;
-
-	UE_DEPRECATED(5.3, "Please use the the const version of IsDataValid(FDataValidationContext&)")
-	COREUOBJECT_API virtual EDataValidationResult IsDataValid(class FDataValidationContext& Context);
-
-	/**
-	 * Generic function to validate objects during changelist validations, etc.
-	 *
-	 * @param	ValidationErrors	the array of validation errors.
-	 * @return Valid if this object has data validation rules set up for it and the data for this object is valid. Returns Invalid if it does not pass 
-	 *         the rules. Returns NotValidated if no rules are set for this object.
-	 */
-	UE_DEPRECATED(5.3, "Please use IsDataValid(FDataValidationContext&) const as the FDataValidationContext will allow distinguishing between Warnings and Errors.")
-	COREUOBJECT_API virtual EDataValidationResult IsDataValid(TArray<FText>& ValidationErrors);
 #endif // WITH_EDITOR
 
 	/*----------------------------------------------------------
diff --git a/Engine/Source/Runtime/CoreUObject/Public/UObject/PackageTrailer.h b/Engine/Source/Runtime/CoreUObject/Public/UObject/PackageTrailer.h
--- a/Engine/Source/Runtime/CoreUObject/Public/UObject/PackageTrailer.h
+++ b/Engine/Source/Runtime/CoreUObject/Public/UObject/PackageTrailer.h
@@ -392,17 +392,6 @@ private:
 class FPackageTrailer
 {
 public:
-	/** 
-	 * Returns if the feature is enabled or disabled.
-	 * 
-	 * Note that this is for development purposes only and should ship as always enabled!
-	 */
-	UE_DEPRECATED(5.3, "FPackageTrailer::IsEnabled will always return true as the system is no longer optional")
-	[[nodiscard]] static bool IsEnabled()
-	{
-		return true;
-	}
-
 	/** Try to load a trailer from a given package path. Note that it will always try to load the trailer from the workspace domain */
 	[[nodiscard]] static COREUOBJECT_API bool TryLoadFromPackage(const FPackagePath& PackagePath, FPackageTrailer& OutTrailer);
 
diff --git a/Engine/Source/Runtime/CoreUObject/Public/UObject/UObjectAnnotation.h b/Engine/Source/Runtime/CoreUObject/Public/UObject/UObjectAnnotation.h
--- a/Engine/Source/Runtime/CoreUObject/Public/UObject/UObjectAnnotation.h
+++ b/Engine/Source/Runtime/CoreUObject/Public/UObject/UObjectAnnotation.h
@@ -909,17 +909,6 @@ public:
 		return MaxAnnotations;
 	}
 
-	/**
-	* Return the number max capacity of the array
-	* Thread safe, but you know, someone might have added more elements before this even returns
-	* @return	the maximum number of elements in the array
-	**/
-	UE_DEPRECATED(5.3, "Use GetMaxAnnotations instead")
-	UE_FORCEINLINE_HINT int32 GetMaxAnnottations() const TSAN_SAFE
-	{
-		return MaxAnnotations;
-	}
-
 	/**
 	* Return if this index is valid
 	* Thread safe, if it is valid now, it is valid forever. Other threads might be adding during this call.
diff --git a/Engine/Source/Runtime/CoreUObject/Public/UObject/UObjectBase.h b/Engine/Source/Runtime/CoreUObject/Public/UObject/UObjectBase.h
--- a/Engine/Source/Runtime/CoreUObject/Public/UObject/UObjectBase.h
+++ b/Engine/Source/Runtime/CoreUObject/Public/UObject/UObjectBase.h
@@ -270,19 +270,6 @@ public:
 	/** Release ref-count on the object. */
 	COREUOBJECT_API void ReleaseRef() const;
 
-protected:
-	/**
-	 * Set the object flags directly
-	 *
-	 **/
-	UE_DEPRECATED(5.3, "This function is not thread-safe. Use AtomicallySetFlags or AtomicallyClearFlags instead.")
-	FORCEINLINE void SetFlagsTo( EObjectFlags NewFlags )
-	{
-		checkfSlow((NewFlags & ~RF_AllFlags) == 0, TEXT("%s flagged as 0x%x but is trying to set flags to RF_AllFlags"), *GetFName().ToString(), (int)ObjectFlags);
-		ObjectFlags = NewFlags;
-	}
-
-public:
 	/**
 	 * Retrieve the object flags directly
 	 *
