[dSTM] Objects that are pending construction (and therefore do not have Stubs) should be considered Local.

#rnx

#Synced-CL 49956667

diff --git a/Engine/Source/Runtime/CoreUObject/Private/UObject/RemoteObject.cpp b/Engine/Source/Runtime/CoreUObject/Private/UObject/RemoteObject.cpp
--- a/Engine/Source/Runtime/CoreUObject/Private/UObject/RemoteObject.cpp
+++ b/Engine/Source/Runtime/CoreUObject/Private/UObject/RemoteObject.cpp
@@ -103,11 +103,11 @@ public:
 		return Super::FindOrAdd(UE::RemoteObject::Private::FRemoteIdLocalizationHelper::GetLocalized(Id));
 	}
 
-	UE::RemoteObject::Handle::FRemoteObjectStub* Find(FRemoteObjectId Id)
+	UE::RemoteObject::Handle::FRemoteObjectStub* Find(FRemoteObjectId Id) const
 	{
 		using namespace UE::RemoteObject::Handle;
 
-		UE::RemoteObject::Handle::FRemoteObjectStub** ExistingStub = Super::Find(UE::RemoteObject::Private::FRemoteIdLocalizationHelper::GetLocalized(Id));
+		UE::RemoteObject::Handle::FRemoteObjectStub* const* ExistingStub = Super::Find(UE::RemoteObject::Private::FRemoteIdLocalizationHelper::GetLocalized(Id));
 		return ExistingStub ? *ExistingStub : nullptr;
 	}
 };
@@ -590,8 +590,12 @@ EResidence GetResidence(const UObject* Object)
 		return EResidence::Local;
 	}
 	const int32 InternalIndex = GUObjectArray.ObjectToIndex(Object);
-	const bool bIsRemote = InternalIndex >= 0 && GUObjectArray.IndexToObject(InternalIndex)->HasAnyFlags(EInternalObjectFlags::Remote);
-	if (!bIsRemote)
+	FUObjectItem* ObjectItem = InternalIndex >= 0 ? GUObjectArray.IndexToObject(InternalIndex) : nullptr;
+	if (!ObjectItem || ObjectItem->HasAnyFlags(EInternalObjectFlags::PendingConstruction))
+	{
+		return EResidence::Local;
+	}
+	else if (!ObjectItem->HasAnyFlags(EInternalObjectFlags::Remote))
 	{
 		UE::RemoteObject::Handle::FRemoteObjectStub* Stub = FindRemoteObjectStub(FRemoteObjectId(Object));
 		return Stub && Stub->MigratedData ? EResidence::LocalNotReady : EResidence::Local;
