PR #14181: Add support for FDoubleInterval to be a UPROPERTY

#14181

diff --git a/Engine/Config/BaseEditor.ini b/Engine/Config/BaseEditor.ini
--- a/Engine/Config/BaseEditor.ini
+++ b/Engine/Config/BaseEditor.ini
@@ -667,6 +667,7 @@ bDisableImmediateAssetDiscovery=false
 +NoExportTypesWithDirectNativeFieldAccess=/Script/CoreUObject.InterpCurvePointLinearColor
 +NoExportTypesWithDirectNativeFieldAccess=/Script/CoreUObject.InterpCurveLinearColor
 +NoExportTypesWithDirectNativeFieldAccess=/Script/CoreUObject.FloatInterval
++NoExportTypesWithDirectNativeFieldAccess=/Script/CoreUObject.DoubleInterval
 +NoExportTypesWithDirectNativeFieldAccess=/Script/CoreUObject.Int32Interval
 +NoExportTypesWithDirectNativeFieldAccess=/Script/CoreUObject.FrameNumber
 +NoExportTypesWithDirectNativeFieldAccess=/Script/Engine.DistributionLookupTable
diff --git a/Engine/Source/Runtime/CoreUObject/Private/UObject/Class.cpp b/Engine/Source/Runtime/CoreUObject/Private/UObject/Class.cpp
--- a/Engine/Source/Runtime/CoreUObject/Private/UObject/Class.cpp
+++ b/Engine/Source/Runtime/CoreUObject/Private/UObject/Class.cpp
@@ -8063,6 +8063,12 @@ UScriptStruct* TBaseStructure<FFloatInterval>::Get()
 	return ScriptStruct;
 }
 
+UScriptStruct* TBaseStructure<FDoubleInterval>::Get()
+{
+	static UScriptStruct* ScriptStruct = StaticGetBaseStructureInternal(TEXT("DoubleInterval"));
+	return ScriptStruct;
+}
+
 UScriptStruct* TBaseStructure<FInt32Interval>::Get()
 {
 	static UScriptStruct* ScriptStruct = StaticGetBaseStructureInternal(TEXT("Int32Interval"));
diff --git a/Engine/Source/Runtime/CoreUObject/Public/UObject/NoExportTypes.h b/Engine/Source/Runtime/CoreUObject/Public/UObject/NoExportTypes.h
--- a/Engine/Source/Runtime/CoreUObject/Public/UObject/NoExportTypes.h
+++ b/Engine/Source/Runtime/CoreUObject/Public/UObject/NoExportTypes.h
@@ -2455,6 +2455,22 @@ struct FFloatInterval
 	float Max;
 };
 
+/**
+ * An interval of doubles, defined by inclusive min and max values
+ * @note This is a mirror of TInterval<double>, defined in Interval.h
+ */
+USTRUCT(noexport, BlueprintType)
+struct FDoubleInterval
+{
+	/** Values must be >= Min */
+	UPROPERTY(EditAnywhere, BlueprintReadWrite, Category=Interval)
+	double Min;
+
+	/** Values must be <= Max */
+	UPROPERTY(EditAnywhere, BlueprintReadWrite, Category=Interval)
+	double Max;
+};
+
 /**
  * An interval of integers, defined by inclusive min and max values
  * @note This is a mirror of TInterval<int32>, defined in Engine/Source/Runtime/Core/Public/Math/Interval.h
diff --git a/Engine/Source/Runtime/PropertyPath/Public/PropertyTypeCompatibility.h b/Engine/Source/Runtime/PropertyPath/Public/PropertyTypeCompatibility.h
--- a/Engine/Source/Runtime/PropertyPath/Public/PropertyTypeCompatibility.h
+++ b/Engine/Source/Runtime/PropertyPath/Public/PropertyTypeCompatibility.h
@@ -279,6 +279,12 @@ inline bool IsConcreteTypeCompatibleWithReflectedType_Impl<FFloatRange>(FPropert
 	return IsConcreteTypeCompatibleWithReflectedType_BuiltInStruct<FFloatRange>(Property);
 }
 
+template<>
+inline bool IsConcreteTypeCompatibleWithReflectedType_Impl<FDoubleInterval>(FProperty* Property)
+{
+	return IsConcreteTypeCompatibleWithReflectedType_BuiltInStruct<FDoubleInterval>(Property);
+}
+
 template<>
 inline bool IsConcreteTypeCompatibleWithReflectedType_Impl<FInt32RangeBound>(FProperty* Property)
 {
