AutoRTFM:  Fixing issue in LowLevelRename where undoing a rename can cause a DSTM migration outside of a transaction.  The rename can be discarded in the case where we know we've created the object during the transaction.

#rnx

#Synced-CL 49956667

diff --git a/Engine/Source/Runtime/CoreUObject/Private/UObject/UObjectBase.cpp b/Engine/Source/Runtime/CoreUObject/Private/UObject/UObjectBase.cpp
--- a/Engine/Source/Runtime/CoreUObject/Private/UObject/UObjectBase.cpp
+++ b/Engine/Source/Runtime/CoreUObject/Private/UObject/UObjectBase.cpp
@@ -32,6 +32,10 @@
 #include "Interfaces/IPluginManager.h"
 #include "Serialization/LoadTimeTrace.h"
 
+#if UE_WITH_REMOTE_OBJECT_HANDLE
+#include "UObject/RemoteObjectPrivate.h"
+#endif
+
 DEFINE_LOG_CATEGORY_STATIC(LogUObjectBase, Log, All);
 DEFINE_STAT(STAT_UObjectsStatGroupTester);
 
@@ -279,15 +283,36 @@ void UObjectBase::LowLevelRename(FName NewName,UObject *NewOuter)
 	if (AutoRTFM::IsClosed())
 	{
 		FName OldName = GetFName();
-		UObject * OldOuter = OuterPrivate;
-		UE_AUTORTFM_ONABORT(this, NewName, NewOuter, OldName, OldOuter)
+		TObjectPtr<UObject> This = static_cast<UObject*>(this);
+		TObjectPtr<UObject> OldOuter = OuterPrivate;
+		TObjectPtr<UObject> NewOuterPtr = NewOuter;
+		UE_AUTORTFM_ONABORT(This, NewName, OldName, OldOuter, NewOuterPtr)
 		{
-			// Only rename if the new name / owner matches the current name / 
-			// owner, as the name may have been changed due to a whole-object
-			// revert.
-			if (NewName == GetFName() && NewOuter == OuterPrivate)
+#if UE_WITH_REMOTE_OBJECT_HANDLE
+			// Prevent us from ever migrating an object since aborting is probably outside of a transactional scope (we will use old copies which we should have, see ensures below)
+			UE::RemoteObject::Private::FUnsafeToMigrateScope CannotMigrateOutsideOfTransaction;
+
+			// These can be LocalNotReady if they were migrated this transaction (and now undone). We will use their old copies for these operations to avoid migrations.
+			ensureMsgf(This.GetResidence() != EResidence::Remote, TEXT("Somehow in-use Object %s migrated during a transaction"), *This.GetRemoteId().ToString());
+			ensureMsgf(NewOuterPtr.GetResidence() != EResidence::Remote, TEXT("Somehow in-use NewOuter %s migrated during a transaction"), *NewOuterPtr.GetRemoteId().ToString());
+			ensureMsgf(OldOuter.GetResidence() != EResidence::Remote, TEXT("Somehow in-use OldOuter %s migrated during a transaction"), *OldOuter.GetRemoteId().ToString());
+#endif
+			// Due to the unsafe to migrate above, we do need to check for existence
+			if (This)
 			{
-				LowLevelRename(OldName, OldOuter);
+				// If we were constructed this transaction, then we know we aren't going to persist don't bother renaming.
+				if (This->HasAnyInternalFlags(EInternalObjectFlags::PendingConstruction))
+				{
+					return;
+				}
+
+				// Only rename if the new name / owner matches the current name / 
+				// owner, as the name may have been changed due to a whole-object
+				// revert.
+				if (NewName == This->GetFName() && NewOuterPtr == This->OuterPrivate)
+				{
+					This->LowLevelRename(OldName, OldOuter);
+				}
 			}
 		};
 		UE_AUTORTFM_OPEN
