Moved an internal namespace under the UE prefix.  Changed TFunctions in our RemoteExecutor API to instead be TUniqueFunctions in order to allow Move operations in the lambda captures.

#rnx
#Synced-CL 49956667

diff --git a/Engine/Source/Runtime/CoreUObject/Private/UObject/RemoteExecutor.cpp b/Engine/Source/Runtime/CoreUObject/Private/UObject/RemoteExecutor.cpp
--- a/Engine/Source/Runtime/CoreUObject/Private/UObject/RemoteExecutor.cpp
+++ b/Engine/Source/Runtime/CoreUObject/Private/UObject/RemoteExecutor.cpp
@@ -26,24 +26,26 @@ static FAutoConsoleVariableRef CVarRemoteExecutorAllowSpeculationMode(
 
 namespace UE::RemoteExecutor
 {
-TDelegate<void()> TickNetworkDelegate;
-TDelegate<void(FRemoteTransactionId, FRemoteWorkPriority, const TArray<FRemoteServerId>& /*MultiServerCommitRemoteServers*/)> BeginMultiServerCommitDelegate;
-TDelegate<void(FRemoteTransactionId, const TArray<FRemoteServerId>& /*MultiServerCommitRemoteServers*/)> ReadyMultiServerCommitDelegate;
-TDelegate<void(FRemoteTransactionId, const TArray<FRemoteServerId>& /*MultiServerCommitRemoteServers*/)> AbandonMultiServerCommitDelegate;
-TDelegate<void(FRemoteTransactionId, const TArray<FRemoteServerId>& /*MultiServerCommitRemoteServers*/)> EndMultiServerCommitDelegate;
-TDelegate<void(FRemoteTransactionId, FRemoteServerId /*ServerId*/)> ReadyRemoteMultiServerCommitDelegate;
-TDelegate<void(FRemoteTransactionId, FRemoteServerId /*ServerId*/)> AbortRemoteMultiServerCommitDelegate;
-TDelegate<TOptional<TTuple<FName, FRemoteWorkPriority, bool, TFunction<void(void)>>>(void)> FetchNextDeferredRPCDelegate;
-
-TMulticastDelegate<void(FRemoteTransactionId, FName)> OnTransactionQueuedDelegate;
-TMulticastDelegate<void(FRemoteTransactionId, FName)> OnTransactionStartingDelegate;
-TMulticastDelegate<void(FRemoteTransactionId, uint32)> OnTransactionCompletedDelegate;
-TMulticastDelegate<void(FRemoteTransactionId, uint32, const FString&)> OnTransactionAbortedDelegate;
-TMulticastDelegate<void(FRemoteTransactionId)> OnTransactionReleasedDelegate;
-
-TDelegate<void(FName /*Text*/)> OnRegionBeginDelegate;
-TDelegate<void(const FString& /*Text*/)> OnRegionEndDelegate;
-
+// Prevent us from having to match this definition to the declartion in the header just to initialize the delegate
+#define DEFINE_REMOTE_EXECUTOR_DELEGATE(x) decltype(x) x
+	DEFINE_REMOTE_EXECUTOR_DELEGATE(TickNetworkDelegate);
+	DEFINE_REMOTE_EXECUTOR_DELEGATE(BeginMultiServerCommitDelegate);
+	DEFINE_REMOTE_EXECUTOR_DELEGATE(ReadyMultiServerCommitDelegate);
+	DEFINE_REMOTE_EXECUTOR_DELEGATE(AbandonMultiServerCommitDelegate);
+	DEFINE_REMOTE_EXECUTOR_DELEGATE(EndMultiServerCommitDelegate);
+	DEFINE_REMOTE_EXECUTOR_DELEGATE(ReadyRemoteMultiServerCommitDelegate);
+	DEFINE_REMOTE_EXECUTOR_DELEGATE(AbortRemoteMultiServerCommitDelegate);
+	DEFINE_REMOTE_EXECUTOR_DELEGATE(FetchNextDeferredRPCDelegate);
+
+	DEFINE_REMOTE_EXECUTOR_DELEGATE(OnTransactionQueuedDelegate);
+	DEFINE_REMOTE_EXECUTOR_DELEGATE(OnTransactionStartingDelegate);
+	DEFINE_REMOTE_EXECUTOR_DELEGATE(OnTransactionCompletedDelegate);
+	DEFINE_REMOTE_EXECUTOR_DELEGATE(OnTransactionAbortedDelegate);
+	DEFINE_REMOTE_EXECUTOR_DELEGATE(OnTransactionReleasedDelegate);
+
+	DEFINE_REMOTE_EXECUTOR_DELEGATE(OnRegionBeginDelegate);
+	DEFINE_REMOTE_EXECUTOR_DELEGATE(OnRegionEndDelegate);
+#undef DEFINE_REMOTE_EXECUTOR_DELEGATE
 
 const TCHAR* EnumToString(ERemoteExecutorAbortReason AbortReason)
 {
@@ -78,7 +80,7 @@ struct FRemoteExecutorWorkList
 	FRemoteWorkPriority Priority;
 	bool bIsTransactional = false;
 
-	TArray<TFunction<void(void)>> PendingWork;
+	TArray<TUniqueFunction<void(void)>> PendingWork;
 
 	void Reset()
 	{
@@ -97,7 +99,7 @@ struct FRemoteExecutorWork
 	FRemoteWorkPriority Priority;
 	bool bIsTransactional = false;
 
-	TFunction<void(void)> Work;
+	TUniqueFunction<void(void)> Work;
 
 	TArray<void*> SubsystemRequests;
 	bool bRequestCreated = false;
@@ -146,7 +148,7 @@ public:
 	FRemoteServerId ActiveRemoteMultiServerCommitServerId;
 	FRemoteTransactionId ActiveRemoteMultiServerCommitRequestId;
 	FRemoteWorkPriority ActiveRemoteMultiServerCommitPriority;
-	TArray<TFunction<void()>> ActiveRemoteMultiServerCommitDeferredActions;
+	TArray<TUniqueFunction<void()>> ActiveRemoteMultiServerCommitDeferredActions;
 	bool bActiveRemoteMultiServerCommitReady = false;
 
 	// tracking data for executing our local multi-server commit
@@ -165,12 +167,12 @@ public:
 		return FRemoteTransactionId((RawRequestId++ % 0x800000u) + 1000u);
 	}
 
-	void EnqueueWork(FName InWorkName, FRemoteWorkPriority InWorkPriority, bool bIsTransactional, int32 InWorkListIndex, const TFunction<void(void)>& InWork)
+	void EnqueueWork(FName InWorkName, FRemoteWorkPriority InWorkPriority, bool bIsTransactional, int32 InWorkListIndex, TUniqueFunction<void(void)>&& InWork)
 	{
 		FRemoteExecutorWork& NewWork = PendingWorks.Emplace_GetRef();
 		NewWork.Name = InWorkName;
 		NewWork.Priority = InWorkPriority;
-		NewWork.Work = InWork;
+		NewWork.Work = MoveTemp(InWork);
 		NewWork.bIsTransactional = bIsTransactional;
 		NewWork.RequestId = GenerateNextTransactionId();
 		NewWork.WorkListIndex = InWorkListIndex;
@@ -181,12 +183,12 @@ public:
 		}
 	}
 
-	void EnqueueMigrationWork(FName InWorkName, FRemoteWorkPriority InWorkPriority, int32 InWorkListIndex, const TFunction<void(void)>& InWork)
+	void EnqueueMigrationWork(FName InWorkName, FRemoteWorkPriority InWorkPriority, int32 InWorkListIndex, TUniqueFunction<void(void)>&& InWork)
 	{
 		FRemoteExecutorWork& NewWork = PendingMigrationWorks.Emplace_GetRef();
 		NewWork.Name = InWorkName;
 		NewWork.Priority = InWorkPriority;
-		NewWork.Work = InWork;
+		NewWork.Work = MoveTemp(InWork);
 		NewWork.bIsTransactional = true;
 		NewWork.RequestId = GenerateNextTransactionId();
 		NewWork.WorkListIndex = InWorkListIndex;
@@ -194,14 +196,14 @@ public:
 		UE::RemoteExecutor::OnTransactionQueuedDelegate.Broadcast(NewWork.RequestId, NewWork.Name);
 	}
 
-	void EnqueueWorkList(FName WorkName, FRemoteWorkPriority InWorkPriority, bool bIsTransactional, const TArray<TFunction<void(void)>>& InWorkList)
+	void EnqueueWorkList(FName WorkName, FRemoteWorkPriority InWorkPriority, bool bIsTransactional, TArray<TUniqueFunction<void(void)>>&& InWorkList)
 	{
 		FRemoteExecutorWorkList& NewWorkList = PendingWorkLists.Emplace_GetRef();
 
 		NewWorkList.Name = WorkName;
 		NewWorkList.Priority = InWorkPriority;
 		NewWorkList.bIsTransactional = bIsTransactional;
-		NewWorkList.PendingWork = InWorkList;
+		NewWorkList.PendingWork = MoveTemp(InWorkList);
 	}
 
 	void ExecutePendingWork()
@@ -226,7 +228,7 @@ public:
 
 						if (WorkList.PendingWork.Num() > 0)
 						{
-							EnqueueWork(WorkList.Name, WorkList.Priority, WorkList.bIsTransactional, WorkListIndex, WorkList.PendingWork[0]);
+							EnqueueWork(WorkList.Name, WorkList.Priority, WorkList.bIsTransactional, WorkListIndex, MoveTemp(WorkList.PendingWork[0]));
 							WorkList.PendingWork.RemoveAt(0, EAllowShrinking::No);
 						}
 					}
@@ -242,7 +244,7 @@ public:
 
 					if (WorkList.PendingWork.Num() > 0)
 					{
-						EnqueueWork(WorkList.Name, WorkList.Priority, WorkList.bIsTransactional, WorkListIndex, WorkList.PendingWork[0]);
+						EnqueueWork(WorkList.Name, WorkList.Priority, WorkList.bIsTransactional, WorkListIndex, MoveTemp(WorkList.PendingWork[0]));
 						WorkList.PendingWork.RemoveAt(0, EAllowShrinking::No);
 					}
 				}
@@ -335,7 +337,8 @@ public:
 
 			if (PendingMigrationWorks.Num())
 			{
-				PendingMigrationWorks += PendingWorks;
+				// Prepend the Pending Migrations to the PendingWorks list
+				PendingMigrationWorks += MoveTemp(PendingWorks);
 				PendingWorks = MoveTemp(PendingMigrationWorks);
 			}
 
@@ -1052,7 +1055,7 @@ void EndRemoteMultiServerCommit(FRemoteServerId ServerId, FRemoteTransactionId R
 	check(GRemoteExecutor.ActiveRemoteMultiServerCommitServerId == ServerId);
 	check(GRemoteExecutor.ActiveRemoteMultiServerCommitRequestId == RequestId);
 
-	for (TFunction<void()>& DeferredAction : GRemoteExecutor.ActiveRemoteMultiServerCommitDeferredActions)
+	for (TUniqueFunction<void()>& DeferredAction : GRemoteExecutor.ActiveRemoteMultiServerCommitDeferredActions)
 	{
 		DeferredAction();
 	}
@@ -1083,7 +1086,7 @@ void AbandonRemoteMultiServerCommit(FRemoteServerId ServerId, FRemoteTransaction
 	}
 }
 
-void EnqueueRemoteMultiServerCommitAction(FRemoteServerId ServerId, FRemoteTransactionId RequestId, const TFunction<void()>& Action)
+void EnqueueRemoteMultiServerCommitAction(FRemoteServerId ServerId, FRemoteTransactionId RequestId, TUniqueFunction<void()>&& Action)
 {
 #if UE_WITH_REMOTE_OBJECT_HANDLE
 	// THIS has to sit in a holding pen...
@@ -1092,7 +1095,7 @@ void EnqueueRemoteMultiServerCommitAction(FRemoteServerId ServerId, FRemoteTrans
 	{
 		UE_LOG(LogRemoteExec, Verbose, TEXT("EnqueueRemoteMultiServerCommitAction Enqueueing action from %s request %s"), *ServerId.ToString(), *RequestId.ToString());
 
-		GRemoteExecutor.ActiveRemoteMultiServerCommitDeferredActions.Add(Action);
+		GRemoteExecutor.ActiveRemoteMultiServerCommitDeferredActions.Add(MoveTemp(Action));
 	}
 	else
 	{
@@ -1169,49 +1172,49 @@ void ExecuteTransactional(FName WorkName, const TFunctionRef<void(void)>& Work)
 	ExecuteTransactionalWithExplicitPriority(WorkName, RootWorkPriority, Work);
 }
 
-void EnqueueWorkWithExplicitPriority(FName WorkName, FRemoteWorkPriority WorkPriority, bool bIsTransactional, const TFunction<void(void)>& InWork)
+void EnqueueWorkWithExplicitPriority(FName WorkName, FRemoteWorkPriority WorkPriority, bool bIsTransactional, TUniqueFunction<void(void)>&& InWork)
 {
 #if UE_WITH_REMOTE_OBJECT_HANDLE && UE_AUTORTFM
-	GRemoteExecutor.EnqueueWork(WorkName, WorkPriority, bIsTransactional, INDEX_NONE, InWork);
+	GRemoteExecutor.EnqueueWork(WorkName, WorkPriority, bIsTransactional, INDEX_NONE, MoveTemp(InWork));
 #else
 	InWork();
 #endif
 }
 
-void EnqueueWork(FName WorkName, bool bIsTransactional, const TFunction<void(void)>& InWork)
+void EnqueueWork(FName WorkName, bool bIsTransactional, TUniqueFunction<void(void)>&& InWork)
 {
 	FRemoteWorkPriority RootWorkPriority = FRemoteWorkPriority::CreateRootWorkPriority(FRemoteServerId::GetLocalServerId(), GRemoteExecutor.GenerateNextTransactionId());
-	EnqueueWorkWithExplicitPriority(WorkName, RootWorkPriority, bIsTransactional, InWork);
+	EnqueueWorkWithExplicitPriority(WorkName, RootWorkPriority, bIsTransactional, MoveTemp(InWork));
 }
 
-void EnqueueMigrationWork(FName WorkName, const TFunction<void(void)>& InWork)
+void EnqueueMigrationWork(FName WorkName, TUniqueFunction<void(void)>&& InWork)
 {
 #if UE_WITH_REMOTE_OBJECT_HANDLE && UE_AUTORTFM
 	FRemoteWorkPriority RootWorkPriority = FRemoteWorkPriority::CreateRootWorkPriority(FRemoteServerId::GetLocalServerId(), GRemoteExecutor.GenerateNextTransactionId());
-	GRemoteExecutor.EnqueueMigrationWork(WorkName, RootWorkPriority, INDEX_NONE, InWork);
+	GRemoteExecutor.EnqueueMigrationWork(WorkName, RootWorkPriority, INDEX_NONE, MoveTemp(InWork));
 #else
 	InWork();
 #endif
 }
 
 
-void EnqueueWorkListWithExplicitPriority(FName WorkName, FRemoteWorkPriority WorkPriority, bool bIsTransactional, const TArray<TFunction<void(void)>>& InWorkList)
+void EnqueueWorkListWithExplicitPriority(FName WorkName, FRemoteWorkPriority WorkPriority, bool bIsTransactional, TArray<TUniqueFunction<void(void)>>&& InWorkList)
 {
 #if UE_WITH_REMOTE_OBJECT_HANDLE
-	GRemoteExecutor.EnqueueWorkList(WorkName, WorkPriority, bIsTransactional, InWorkList);
+	GRemoteExecutor.EnqueueWorkList(WorkName, WorkPriority, bIsTransactional, MoveTemp(InWorkList));
 #else
 	for(int32 WorkIndex = 0; WorkIndex < InWorkList.Num(); WorkIndex++)
 	{
-		const TFunction<void(void)>& Work = InWorkList[WorkIndex];
+		const TUniqueFunction<void(void)>& Work = InWorkList[WorkIndex];
 		Work();
 	}
 #endif
 }
 
-void EnqueueWorkList(FName WorkName, bool bIsTransactional, const TArray<TFunction<void(void)>>& InWorkList)
+void EnqueueWorkList(FName WorkName, bool bIsTransactional, TArray<TUniqueFunction<void(void)>>&& InWorkList)
 {
 	FRemoteWorkPriority RootWorkPriority = FRemoteWorkPriority::CreateRootWorkPriority(FRemoteServerId::GetLocalServerId(), GRemoteExecutor.GenerateNextTransactionId());
-	EnqueueWorkListWithExplicitPriority(WorkName, RootWorkPriority, bIsTransactional, InWorkList);
+	EnqueueWorkListWithExplicitPriority(WorkName, RootWorkPriority, bIsTransactional, MoveTemp(InWorkList));
 }
 
 } // namespace UE::RemoteExecutor
diff --git a/Engine/Source/Runtime/CoreUObject/Private/UObject/RemoteObject.cpp b/Engine/Source/Runtime/CoreUObject/Private/UObject/RemoteObject.cpp
--- a/Engine/Source/Runtime/CoreUObject/Private/UObject/RemoteObject.cpp
+++ b/Engine/Source/Runtime/CoreUObject/Private/UObject/RemoteObject.cpp
@@ -313,7 +313,7 @@ void InitRemoteObjects()
 
 	if (!UE::RemoteExecutor::FetchNextDeferredRPCDelegate.IsBound())
 	{
-		UE::RemoteExecutor::FetchNextDeferredRPCDelegate.BindStatic([]() { return TOptional<TTuple<FName, FRemoteWorkPriority, bool, TFunction<void(void)>>>(); });
+		UE::RemoteExecutor::FetchNextDeferredRPCDelegate.BindStatic([]() { return TOptional<TTuple<FName, FRemoteWorkPriority, bool, TUniqueFunction<void(void)>>>(); });
 	}
 
 #if UE_WITH_REMOTE_OBJECT_HANDLE
diff --git a/Engine/Source/Runtime/CoreUObject/Public/UObject/RemoteExecutor.h b/Engine/Source/Runtime/CoreUObject/Public/UObject/RemoteExecutor.h
--- a/Engine/Source/Runtime/CoreUObject/Public/UObject/RemoteExecutor.h
+++ b/Engine/Source/Runtime/CoreUObject/Public/UObject/RemoteExecutor.h
@@ -234,7 +234,7 @@ namespace UE::RemoteExecutor
 	extern COREUOBJECT_API TDelegate<void(FRemoteTransactionId, FRemoteServerId)> AbortRemoteMultiServerCommitDelegate;
 
 	/** Fetch the next deferred RPC request **/
-	extern COREUOBJECT_API TDelegate<TOptional<TTuple<FName, FRemoteWorkPriority, bool, TFunction<void(void)>>>(void)> FetchNextDeferredRPCDelegate;
+	extern COREUOBJECT_API TDelegate<TOptional<TTuple<FName, FRemoteWorkPriority, bool, TUniqueFunction<void(void)>>>(void)> FetchNextDeferredRPCDelegate;
 
 	/** Delegate executed when a transaction is queued in the executor for execution */
 	extern COREUOBJECT_API TMulticastDelegate<void(FRemoteTransactionId, FName)> OnTransactionQueuedDelegate;
@@ -262,11 +262,11 @@ namespace UE::RemoteExecutor
 
 	COREUOBJECT_API void ExecuteTransactional(FName WorkName, const TFunctionRef<void(void)>& Work);
 	COREUOBJECT_API void ExecuteTransactionalWithExplicitPriority(FName WorkName, FRemoteWorkPriority WorkPriority, const TFunctionRef<void(void)>& Work);
-	COREUOBJECT_API void EnqueueWorkWithExplicitPriority(FName WorkName, FRemoteWorkPriority WorkPriority, bool bIsTransactional, const TFunction<void(void)>& Work);
-	COREUOBJECT_API void EnqueueWork(FName WorkName, bool bIsTransactional, const TFunction<void(void)>& Work);
-	COREUOBJECT_API void EnqueueMigrationWork(FName WorkName, const TFunction<void(void)>& InWork);
-	COREUOBJECT_API void EnqueueWorkListWithExplicitPriority(FName WorkName, FRemoteWorkPriority WorkPriority, bool bIsTransactional, const TArray<TFunction<void(void)>>& WorkList);
-	COREUOBJECT_API void EnqueueWorkList(FName WorkName, bool bIsTransactional, const TArray<TFunction<void(void)>>& WorkList);
+	COREUOBJECT_API void EnqueueWorkWithExplicitPriority(FName WorkName, FRemoteWorkPriority WorkPriority, bool bIsTransactional, TUniqueFunction<void(void)>&& Work);
+	COREUOBJECT_API void EnqueueWork(FName WorkName, bool bIsTransactional, TUniqueFunction<void(void)>&& Work);
+	COREUOBJECT_API void EnqueueMigrationWork(FName WorkName, TUniqueFunction<void(void)>&& InWork);
+	COREUOBJECT_API void EnqueueWorkListWithExplicitPriority(FName WorkName, FRemoteWorkPriority WorkPriority, bool bIsTransactional, TArray<TUniqueFunction<void(void)>>&& WorkList);
+	COREUOBJECT_API void EnqueueWorkList(FName WorkName, bool bIsTransactional, TArray<TUniqueFunction<void(void)>>&& WorkList);
 	COREUOBJECT_API void ExecutePendingWork();
 
 	COREUOBJECT_API void AbortCurrentTransaction(EAbortTransactionBehavior AbortBehavior, ERemoteExecutorAbortReason AbortReason, FStringView AbortDescription);
@@ -280,7 +280,7 @@ namespace UE::RemoteExecutor
 	COREUOBJECT_API void ReadyRemoteMultiServerCommit(FRemoteServerId ServerId, FRemoteTransactionId RequestId);
 	COREUOBJECT_API void AbandonRemoteMultiServerCommit(FRemoteServerId ServerId, FRemoteTransactionId RequestId);
 	COREUOBJECT_API void EndRemoteMultiServerCommit(FRemoteServerId ServerId, FRemoteTransactionId RequestId);
-	COREUOBJECT_API void EnqueueRemoteMultiServerCommitAction(FRemoteServerId ServerId, FRemoteTransactionId RequestId, const TFunction<void()>& Action);
+	COREUOBJECT_API void EnqueueRemoteMultiServerCommitAction(FRemoteServerId ServerId, FRemoteTransactionId RequestId, TUniqueFunction<void()>&& Action);
 
 	COREUOBJECT_API void ReadyMultiServerCommitResponse(FRemoteServerId ServerId, FRemoteTransactionId RequestId);
 	COREUOBJECT_API void AbortMultiServerCommit(FRemoteServerId ServerId, FRemoteTransactionId RequestId);
diff --git a/Engine/Source/Runtime/Engine/Private/DataReplication.cpp b/Engine/Source/Runtime/Engine/Private/DataReplication.cpp
--- a/Engine/Source/Runtime/Engine/Private/DataReplication.cpp
+++ b/Engine/Source/Runtime/Engine/Private/DataReplication.cpp
@@ -2655,7 +2655,7 @@ void FObjectReplicator::EnqueuePendingTransactionalRPCs()
 	TArray<FRPCPendingLocalCall> ServiceTheseRPCs;
 	Swap(PendingTransactionalRPCs, ServiceTheseRPCs);
 
-	TArray<TFunction<void(void)>> TransactionsToEnqueue;
+	TArray<TUniqueFunction<void(void)>> TransactionsToEnqueue;
 	TransactionsToEnqueue.Reserve(ServiceTheseRPCs.Num());
 
 	for (const FRPCPendingLocalCall& Pending : ServiceTheseRPCs)
@@ -2723,11 +2723,11 @@ void FObjectReplicator::EnqueuePendingTransactionalRPCs()
 
 	if (CVarEnqueueTransactionalRpcs.GetValueOnAnyThread(true))
 	{
-		UE::RemoteExecutor::EnqueueWorkList(Object->GetFName(), true, TransactionsToEnqueue);
+		UE::RemoteExecutor::EnqueueWorkList(Object->GetFName(), true, MoveTemp(TransactionsToEnqueue));
 	}
 	else
 	{
-		for (const TFunction<void(void)>& RpcFunc : TransactionsToEnqueue)
+		for (const TUniqueFunction<void(void)>& RpcFunc : TransactionsToEnqueue)
 		{
 			RpcFunc();
 		}
