[UObject]

Disable object detaching from the linker when running renaming loaded objects in commandlets

#rnx

#Synced-CL 49988537

diff --git a/Engine/Source/Runtime/CoreUObject/Private/UObject/Obj.cpp b/Engine/Source/Runtime/CoreUObject/Private/UObject/Obj.cpp
--- a/Engine/Source/Runtime/CoreUObject/Private/UObject/Obj.cpp
+++ b/Engine/Source/Runtime/CoreUObject/Private/UObject/Obj.cpp
@@ -321,12 +321,15 @@ bool UObject::Rename(const TCHAR* InName, UObject* NewOuter, ERenameFlags Flags)
 	// Users who don't want such reload behaviour for the export should be calling InvalidateExport() on the object before renaming.
 	else if (!bIsTest && ((Flags & (REN_AllowPackageLinkerMismatch | REN_DontCreateRedirectors)) == REN_DontCreateRedirectors) && HasAnyFlags(RF_WasLoaded))
 	{
-		if (FLinkerLoad* Linker = GetLinker())
+		if (!IsRunningCommandlet())
 		{
-			const int32 LinkerIndex = GetLinkerIndex();
-			if (LinkerIndex != INDEX_NONE)
+			if (FLinkerLoad* Linker = GetLinker())
 			{
-				Linker->DetachExport(LinkerIndex);
+				const int32 LinkerIndex = GetLinkerIndex();
+				if (LinkerIndex != INDEX_NONE)
+				{
+					Linker->DetachExport(LinkerIndex);
+				}
 			}
 		}
 	}
@@ -915,7 +918,7 @@ bool UObject::IsSelectedInEditor() const
 void UObject::PostRename(UObject* OldOuter, const FName OldName)
 {
 #if WITH_EDITOR
-	if (HasAnyFlags(RF_WasLoaded))
+	if (HasAnyFlags(RF_WasLoaded) && !IsRunningCommandlet())
 	{
 		FCoreUObjectDelegates::OnLoadedObjectRenamed.Broadcast(this, OldOuter, OldName);
 	}
