[LinkerLoad/ALT]
- Fix ownerthread assert in debug builds when using async loading. Will deprecate ownerthread in another CL.
- Will deprecate ownerthread in another CL to keep this fix trivial.

#Synced-CL 49920339

diff --git a/Engine/Source/Runtime/CoreUObject/Public/UObject/LinkerLoad.h b/Engine/Source/Runtime/CoreUObject/Public/UObject/LinkerLoad.h
--- a/Engine/Source/Runtime/CoreUObject/Public/UObject/LinkerLoad.h
+++ b/Engine/Source/Runtime/CoreUObject/Public/UObject/LinkerLoad.h
@@ -1097,7 +1097,6 @@ private:
 		checkf(!bDetachedLoader, TEXT("Attempting to call ::Serialize on a FLinkerLoad that has previously called ::DetachLoader"));
 #endif // WITH_EDITOR
 
-		checkSlow(FPlatformTLS::GetCurrentThreadId() == OwnerThread);
 #if WITH_EDITOR && !UE_INTERNAL_UNIVERSAL_PLACEHOLDERS
 		// This handles the case where we're only serializing to a property bag and we don't immediately know where the script
 		// property data starts/ends in the object's data stream. Allows for loading from packages saved with an older version.
