Minor cleanup + optimization of GetResidence(UObject*) functions

#Synced-CL 50039487

diff --git a/Engine/Source/Runtime/CoreUObject/Private/UObject/RemoteObject.cpp b/Engine/Source/Runtime/CoreUObject/Private/UObject/RemoteObject.cpp
--- a/Engine/Source/Runtime/CoreUObject/Private/UObject/RemoteObject.cpp
+++ b/Engine/Source/Runtime/CoreUObject/Private/UObject/RemoteObject.cpp
@@ -543,7 +543,7 @@ EResidence GetResidence(FRemoteObjectId ObjectId)
 {
 	using namespace UE::RemoteObject::Private;
 
-	if (!ObjectId.IsValid() || !UE::RemoteObject::Private::ObjectMaps)
+	if (!ObjectId.IsValid() || !ObjectMaps)
 	{
 		return EResidence::Local;
 	}
@@ -554,7 +554,7 @@ EResidence GetResidence(FRemoteObjectId ObjectId)
 	}
 
 	// if a stub exists but we don't have the object's memory (the UObject doesn't exist) then the object must be remote
-	if (UE::RemoteObject::Handle::FRemoteObjectStub* Stub = FindRemoteObjectStub(ObjectId))
+	if (UE::RemoteObject::Handle::FRemoteObjectStub* Stub = ObjectMaps->FindStub(ObjectId))
 	{
 		return Stub->MigratedData ? EResidence::LocalNotReady : EResidence::Remote;
 	}
@@ -585,25 +585,23 @@ EResidence GetResidence(const UObject* Object)
 	using namespace UE::RemoteObject::Private;
 
 #if UE_WITH_REMOTE_OBJECT_HANDLE
-	if (!Object || !UE::RemoteObject::Private::ObjectMaps)
+	if (Object && ObjectMaps)
 	{
-		return EResidence::Local;
-	}
-	const int32 InternalIndex = GUObjectArray.ObjectToIndex(Object);
-	FUObjectItem* ObjectItem = InternalIndex >= 0 ? GUObjectArray.IndexToObject(InternalIndex) : nullptr;
-	if (!ObjectItem || ObjectItem->HasAnyFlags(EInternalObjectFlags::PendingConstruction))
-	{
-		return EResidence::Local;
-	}
-	else if (!ObjectItem->HasAnyFlags(EInternalObjectFlags::Remote))
-	{
-		UE::RemoteObject::Handle::FRemoteObjectStub* Stub = FindRemoteObjectStub(FRemoteObjectId(Object));
-		return Stub && Stub->MigratedData ? EResidence::LocalNotReady : EResidence::Local;
+		const int32 InternalIndex = GUObjectArray.ObjectToIndex(Object);
+		const FUObjectItem* ObjectItem = InternalIndex >= 0 ? GUObjectArray.IndexToObject(InternalIndex) : nullptr;
+		if (!ObjectItem || ObjectItem->HasAnyFlags(EInternalObjectFlags::PendingConstruction))
+		{
+			return EResidence::Local;
+		}
+		else if (!ObjectItem->HasAnyFlags(EInternalObjectFlags::Remote))
+		{
+			UE::RemoteObject::Handle::FRemoteObjectStub* Stub = ObjectMaps->FindStub(ObjectItem->GetRemoteId());
+			return Stub && Stub->MigratedData ? EResidence::LocalNotReady : EResidence::Local;
+		}
+		return EResidence::Remote;
 	}
-	return EResidence::Remote;
-#else
-	return EResidence::Local;
 #endif
+	return EResidence::Local;
 }
 
 bool IsRemote(FRemoteObjectId ObjectId)
