Fix UE_WITH_REMOTE_OBJECT_HANDLE builds

#Synced-CL 50046963

diff --git a/Engine/Source/Runtime/CoreUObject/Private/UObject/PropertyObject.cpp b/Engine/Source/Runtime/CoreUObject/Private/UObject/PropertyObject.cpp
--- a/Engine/Source/Runtime/CoreUObject/Private/UObject/PropertyObject.cpp
+++ b/Engine/Source/Runtime/CoreUObject/Private/UObject/PropertyObject.cpp
@@ -389,7 +389,7 @@ void FObjectProperty::PostSerializeObjectItem(FArchive& SerializingArchive, void
 }
 
 #if UE_WITH_REMOTE_OBJECT_HANDLE
-bool FObjectProperty::NetSerializeItem(FArchive& Ar, UPackageMap* Map, void* Data, TArray<uint8>* MetaData) const
+bool FObjectProperty::NetSerializeItem(FArchive& Ar, UPackageMap* Map, TNotNull<void*> Data, TArray<uint8>* MetaData) const
 {
 	const bool bUseRemoteObjectReference = Map && Map->IsUsingRemoteObjectReferences() && HasAnyPropertyFlags(CPF_TObjectPtr);
 	if (bUseRemoteObjectReference)
diff --git a/Engine/Source/Runtime/CoreUObject/Private/UObject/PropertyWeakObjectPtr.cpp b/Engine/Source/Runtime/CoreUObject/Private/UObject/PropertyWeakObjectPtr.cpp
--- a/Engine/Source/Runtime/CoreUObject/Private/UObject/PropertyWeakObjectPtr.cpp
+++ b/Engine/Source/Runtime/CoreUObject/Private/UObject/PropertyWeakObjectPtr.cpp
@@ -116,7 +116,7 @@ bool FWeakObjectProperty::Identical(TNotNull<const void*> A, const void* B, uint
 	return Super::Identical(A, B, PortFlags);
 }
 
-bool FWeakObjectProperty::NetSerializeItem(FArchive& Ar, UPackageMap* Map, void* Data, TArray<uint8>* MetaData) const
+bool FWeakObjectProperty::NetSerializeItem(FArchive& Ar, UPackageMap* Map, TNotNull<void*> Data, TArray<uint8>* MetaData) const
 {
 	const bool bUseRemoteObjectReference = Map && Map->IsUsingRemoteObjectReferences();
 	if (bUseRemoteObjectReference)
