[Backout] - CL50034055 (changes are unnecessary as the theorized regression was caused elsewhere)
#rnx
Original CL Desc
-----------------------------------------------------------------
[UObject]

Disable object detaching from the linker when running renaming loaded objects in commandlets

#rnx

#Synced-CL 49988537

#Synced-CL 50040740

diff --git a/Engine/Source/Runtime/CoreUObject/Private/UObject/Obj.cpp b/Engine/Source/Runtime/CoreUObject/Private/UObject/Obj.cpp
--- a/Engine/Source/Runtime/CoreUObject/Private/UObject/Obj.cpp
+++ b/Engine/Source/Runtime/CoreUObject/Private/UObject/Obj.cpp
@@ -316,23 +316,22 @@ bool UObject::Rename(const TCHAR* InName, UObject* NewOuter, ERenameFlags Flags)
 			}
 		}
 	}
+#if WITH_EDITOR // until we validate the runtime should also detach (renaming is uncommon)
 	// If we are renaming a loaded object while keeping its outer and won't be creating a redirector, we need to detach from the linker if the user
 	// hasn't specifically requested to allow for the linker mismatch. This way, we can reload the renamed object if the package is rerequested
 	// Users who don't want such reload behaviour for the export should be calling InvalidateExport() on the object before renaming.
 	else if (!bIsTest && ((Flags & (REN_AllowPackageLinkerMismatch | REN_DontCreateRedirectors)) == REN_DontCreateRedirectors) && HasAnyFlags(RF_WasLoaded))
 	{
-		if (!IsRunningCommandlet())
+		if (FLinkerLoad* Linker = GetLinker())
 		{
-			if (FLinkerLoad* Linker = GetLinker())
+			const int32 LinkerIndex = GetLinkerIndex();
+			if (LinkerIndex != INDEX_NONE)
 			{
-				const int32 LinkerIndex = GetLinkerIndex();
-				if (LinkerIndex != INDEX_NONE)
-				{
-					Linker->DetachExport(LinkerIndex);
-				}
+				Linker->DetachExport(LinkerIndex);
 			}
 		}
 	}
+#endif
 
 	// find an object with the same name and same class in the new outer
 	bool bIsCaseOnlyChange = false;
@@ -918,7 +917,7 @@ bool UObject::IsSelectedInEditor() const
 void UObject::PostRename(UObject* OldOuter, const FName OldName)
 {
 #if WITH_EDITOR
-	if (HasAnyFlags(RF_WasLoaded) && !IsRunningCommandlet())
+	if (HasAnyFlags(RF_WasLoaded))
 	{
 		FCoreUObjectDelegates::OnLoadedObjectRenamed.Broadcast(this, OldOuter, OldName);
 	}
