Make TInstancedStruct::InitializeAs<> return a reference to the initialized struct data to match the underlying FInstancedStruct::InitializeAs<>.

#rnx

#Synced-CL 50104405

diff --git a/Engine/Source/Runtime/CoreUObject/Public/StructUtils/InstancedStruct.h b/Engine/Source/Runtime/CoreUObject/Public/StructUtils/InstancedStruct.h
--- a/Engine/Source/Runtime/CoreUObject/Public/StructUtils/InstancedStruct.h
+++ b/Engine/Source/Runtime/CoreUObject/Public/StructUtils/InstancedStruct.h
@@ -370,9 +370,9 @@ public:
 
 	/** Initializes from struct type and emplace construct. */
 	template<typename T = BaseStructT, typename... TArgs, typename = std::enable_if_t<std::is_base_of_v<BaseStructT, std::decay_t<T>>>>
-	void InitializeAs(TArgs&&... InArgs)
+	T& InitializeAs(TArgs&&... InArgs)
 	{
-		InstancedStruct.InitializeAs<T>(Forward<TArgs>(InArgs)...);
+		return InstancedStruct.InitializeAs<T>(Forward<TArgs>(InArgs)...);
 	}
 
 	/** Creates a new TInstancedStruct from templated struct type. */
