Fixed delegate out parameters of scalar type.  Added lots of new tests for out parameters.

diff --git a/Engine/Source/Runtime/Core/Public/UObject/ScriptDelegates.h b/Engine/Source/Runtime/Core/Public/UObject/ScriptDelegates.h
--- a/Engine/Source/Runtime/Core/Public/UObject/ScriptDelegates.h
+++ b/Engine/Source/Runtime/Core/Public/UObject/ScriptDelegates.h
@@ -3,6 +3,8 @@
 #pragma once
 
 #include "Concepts/CompleteType.h"
+#include "Concepts/Const.h"
+#include "Concepts/Reference.h"
 #include "Concepts/Scalar.h"
 #include "Containers/Array.h"
 #include "Containers/ContainerAllocationPolicies.h"
@@ -269,7 +271,7 @@ namespace UE::Core::Private
 	};
 
 	template <typename T>
-		requires UE::CScalar<std::remove_reference_t<T>>
+		requires (!UE::CReference<T> || UE::CConst<std::remove_reference_t<T>>) && UE::CScalar<std::remove_reference_t<T>>
 	struct TDelegateCallTraits<T>
 	{
 		using Type = std::remove_cv_t<std::remove_reference_t<T>>;
diff --git a/Engine/Source/Runtime/CoreUObject/Public/UObject/DynamicDelegateTestObject.cpp b/Engine/Source/Runtime/CoreUObject/Public/UObject/DynamicDelegateTestObject.cpp
--- a/Engine/Source/Runtime/CoreUObject/Public/UObject/DynamicDelegateTestObject.cpp
+++ b/Engine/Source/Runtime/CoreUObject/Public/UObject/DynamicDelegateTestObject.cpp
@@ -10,12 +10,23 @@
 // Set by the name of the function to say which one was called.
 FAnsiString GCalledFunctionName;
 
-int32                                      GIntArg        = 5;
-FString                                    GStrArg        = "Hello";
-const UDynamicDelegateTestObject*          GPtrArg        = nullptr;
-FName                                      GNameArg       = "Package";
-bool                                       bGBoolArg      = true;
+int32 GIntArg = 5;
+int32 GIntOut = 97;
+
+FString GStrArg = "Hello";
+FString GStrOut = "Goodbye";
+
+const UDynamicDelegateTestObject* GPtrArg = nullptr;
+const UDynamicDelegateTestObject* GPtrOut = nullptr;
+
+FName GNameArg = "Package";
+FName GNameOut = "Object";
+
+bool bGBoolArg = true;
+bool bGBoolOut = false;
+
 TWeakObjectPtr<UDynamicDelegateTestObject> GWeakObjPtrArg = nullptr;
+TWeakObjectPtr<UDynamicDelegateTestObject> GWeakObjPtrOut = nullptr;
 
 int32   GIntRet = 99;
 FString GStrRet = "World";
@@ -358,7 +369,113 @@ FString UDynamicDelegateTestObject::ConstFuncRetStr6(int Int, FString Str, const
 	return GStrRet;
 }
 
-#define BIND_CALL_CHECK_DELEGATE(DelegateType, bMulticast, FuncPtr, CallArgs, ...) /* CallArgs is a parentheses-wrapped set of arguments to invoke the delegate, and the ellipsis is the payload arguments */ \
+void UDynamicDelegateTestObject::FuncOut1(int& Int)
+{
+	CHECK(GCalledFunctionName.IsEmpty());
+	GCalledFunctionName = "FuncOut1";
+	CHECK(Int == GIntArg);
+	Int = GIntOut;
+}
+
+void UDynamicDelegateTestObject::FuncOut2(int& Int, FString& Str)
+{
+	FuncOut1(Int);
+	GCalledFunctionName = "FuncOut2";
+	CHECK(Str == GStrArg);
+	Str = GStrOut;
+}
+
+void UDynamicDelegateTestObject::FuncOut3(int& Int, FString& Str, FName& Name)
+{
+	FuncOut2(Int, Str);
+	GCalledFunctionName = "FuncOut3";
+	CHECK(Name == GNameArg);
+	Name = GNameOut;
+}
+
+void UDynamicDelegateTestObject::FuncOut4(int& Int, FString& Str, FName& Name, bool& bBool)
+{
+	FuncOut3(Int, Str, Name);
+	GCalledFunctionName = "FuncOut4";
+	CHECK(bBool == bGBoolArg);
+	bBool = bGBoolOut;
+}
+
+void UDynamicDelegateTestObject::FuncOut5(int& Int, FString& Str, FName& Name, bool& bBool, TWeakObjectPtr<UDynamicDelegateTestObject>& WeakObjPtr)
+{
+	FuncOut4(Int, Str, Name, bBool);
+	GCalledFunctionName = "FuncOut5";
+	CHECK(WeakObjPtr == GWeakObjPtrArg);
+	WeakObjPtr = GWeakObjPtrOut;
+}
+
+void UDynamicDelegateTestObject::ConstFuncOut1(int& Int) const
+{
+	CHECK(GCalledFunctionName.IsEmpty());
+	GCalledFunctionName = "ConstFuncOut1";
+	CHECK(Int == GIntArg);
+	Int = GIntOut;
+}
+
+void UDynamicDelegateTestObject::ConstFuncOut2(int& Int, FString& Str) const
+{
+	ConstFuncOut1(Int);
+	GCalledFunctionName = "ConstFuncOut2";
+	CHECK(Str == GStrArg);
+	Str = GStrOut;
+}
+
+void UDynamicDelegateTestObject::ConstFuncOut3(int& Int, FString& Str, FName& Name) const
+{
+	ConstFuncOut2(Int, Str);
+	GCalledFunctionName = "ConstFuncOut3";
+	CHECK(Name == GNameArg);
+	Name = GNameOut;
+}
+
+void UDynamicDelegateTestObject::ConstFuncOut4(int& Int, FString& Str, FName& Name, bool& bBool) const
+{
+	ConstFuncOut3(Int, Str, Name);
+	GCalledFunctionName = "ConstFuncOut4";
+	CHECK(bBool == bGBoolArg);
+	bBool = bGBoolOut;
+}
+
+void UDynamicDelegateTestObject::ConstFuncOut5(int& Int, FString& Str, FName& Name, bool& bBool, TWeakObjectPtr<UDynamicDelegateTestObject>& WeakObjPtr) const
+{
+	ConstFuncOut4(Int, Str, Name, bBool);
+	GCalledFunctionName = "ConstFuncOut5";
+	CHECK(WeakObjPtr == GWeakObjPtrArg);
+	WeakObjPtr = GWeakObjPtrOut;
+}
+
+#define DEFINE_ARG_VARIABLES_1(Arg1Val)                                              auto Arg1 = Arg1Val;
+#define DEFINE_ARG_VARIABLES_2(Arg1Val, Arg2Val)                                     auto Arg1 = Arg1Val; auto Arg2 = Arg2Val;
+#define DEFINE_ARG_VARIABLES_3(Arg1Val, Arg2Val, Arg3Val)                            auto Arg1 = Arg1Val; auto Arg2 = Arg2Val; auto Arg3 = Arg3Val;
+#define DEFINE_ARG_VARIABLES_4(Arg1Val, Arg2Val, Arg3Val, Arg4Val)                   auto Arg1 = Arg1Val; auto Arg2 = Arg2Val; auto Arg3 = Arg3Val; auto Arg4 = Arg4Val;
+#define DEFINE_ARG_VARIABLES_5(Arg1Val, Arg2Val, Arg3Val, Arg4Val, Arg5Val)          auto Arg1 = Arg1Val; auto Arg2 = Arg2Val; auto Arg3 = Arg3Val; auto Arg4 = Arg4Val; auto Arg5 = Arg5Val;
+#define DEFINE_ARG_VARIABLES_6(Arg1Val, Arg2Val, Arg3Val, Arg4Val, Arg5Val, Arg6Val) auto Arg1 = Arg1Val; auto Arg2 = Arg2Val; auto Arg3 = Arg3Val; auto Arg4 = Arg4Val; auto Arg5 = Arg5Val; auto Arg6 = Arg6Val;
+
+#define EXPAND_ARG_VARIABLES_1( ...) (Arg1)
+#define EXPAND_ARG_VARIABLES_2( ...) (Arg1, Arg2)
+#define EXPAND_ARG_VARIABLES_3( ...) (Arg1, Arg2, Arg3)
+#define EXPAND_ARG_VARIABLES_4( ...) (Arg1, Arg2, Arg3, Arg4)
+#define EXPAND_ARG_VARIABLES_5( ...) (Arg1, Arg2, Arg3, Arg4, Arg5)
+#define EXPAND_ARG_VARIABLES_6( ...) (Arg1, Arg2, Arg3, Arg4, Arg5, Arg6)
+
+#define CHECK_ARG_VARIABLES_1(Arg1Val)                                              CHECK(Arg1 == Arg1Val);
+#define CHECK_ARG_VARIABLES_2(Arg1Val, Arg2Val)                                     CHECK(Arg1 == Arg1Val); CHECK(Arg2 == Arg2Val);
+#define CHECK_ARG_VARIABLES_3(Arg1Val, Arg2Val, Arg3Val)                            CHECK(Arg1 == Arg1Val); CHECK(Arg2 == Arg2Val); CHECK(Arg3 == Arg3Val);
+#define CHECK_ARG_VARIABLES_4(Arg1Val, Arg2Val, Arg3Val, Arg4Val)                   CHECK(Arg1 == Arg1Val); CHECK(Arg2 == Arg2Val); CHECK(Arg3 == Arg3Val); CHECK(Arg4 == Arg4Val);
+#define CHECK_ARG_VARIABLES_5(Arg1Val, Arg2Val, Arg3Val, Arg4Val, Arg5Val)          CHECK(Arg1 == Arg1Val); CHECK(Arg2 == Arg2Val); CHECK(Arg3 == Arg3Val); CHECK(Arg4 == Arg4Val); CHECK(Arg5 == Arg5Val);
+#define CHECK_ARG_VARIABLES_6(Arg1Val, Arg2Val, Arg3Val, Arg4Val, Arg5Val, Arr6Val) CHECK(Arg1 == Arg1Val); CHECK(Arg2 == Arg2Val); CHECK(Arg3 == Arg3Val); CHECK(Arg4 == Arg4Val); CHECK(Arg5 == Arg5Val); CHECK(Arg6 == Arg6Val);
+
+#define DEFINE_ARG_VARIABLES(...) UE_APPEND_VA_ARG_COUNT(DEFINE_ARG_VARIABLES_, __VA_ARGS__)(__VA_ARGS__)
+#define EXPAND_ARG_VARIABLES(...) UE_APPEND_VA_ARG_COUNT(EXPAND_ARG_VARIABLES_, __VA_ARGS__)(ObjectPtr, FuncPtr, __VA_ARGS__)
+#define CHECK_ARG_VARIABLES(...)  UE_APPEND_VA_ARG_COUNT(CHECK_ARG_VARIABLES_,  __VA_ARGS__)(__VA_ARGS__)
+
+// CallArgs is a parentheses-wrapped set of arguments to invoke the delegate, and the ellipsis is the payload arguments
+#define BIND_CALL_CHECK_DELEGATE(DelegateType, bMulticast, FuncPtr, CallArgs, ...) \
 	{ \
 		GCalledFunctionName.Reset(); \
 		DelegateType Delegate; \
@@ -367,7 +484,8 @@ FString UDynamicDelegateTestObject::ConstFuncRetStr6(int Int, FString Str, const
 		CHECK(GCalledFunctionName == UE::Core::Private::TrimmedMemberFunctionName<#FuncPtr>.Str); \
 	}
 
-#define BIND_CALL_CHECK_DELEGATE_RETVAL(ExpectedRetVal, DelegateType, bMulticast, FuncPtr, CallArgs, ...) /* CallArgs is a parentheses-wrapped set of arguments to invoke the delegate, and the ellipsis is the payload arguments */ \
+// CallArgs is a parentheses-wrapped set of arguments to invoke the delegate, and the ellipsis is the payload arguments
+#define BIND_CALL_CHECK_DELEGATE_RETVAL(ExpectedRetVal, DelegateType, bMulticast, FuncPtr, CallArgs, ...) \
 	{ \
 		GCalledFunctionName.Reset(); \
 		DelegateType Delegate; \
@@ -377,54 +495,70 @@ FString UDynamicDelegateTestObject::ConstFuncRetStr6(int Int, FString Str, const
 		CHECK(GCalledFunctionName == UE::Core::Private::TrimmedMemberFunctionName<#FuncPtr>.Str); \
 	}
 
+// CallArgs is a parentheses-wrapped set of arguments to invoke the delegate, OutVals is the set of values to test being returned from the function.  There are no payload arguments, as out payload parameters are captured by value and thus not supported.
+#define BIND_CALL_CHECK_DELEGATE_OUT(DelegateType, bMulticast, FuncPtr, CallArgs, OutVals) \
+	{ \
+		GCalledFunctionName.Reset(); \
+		DelegateType Delegate; \
+		DEFINE_ARG_VARIABLES CallArgs \
+		Delegate.UE_IF(bMulticast, AddDynamic, BindDynamic)(this, FuncPtr); \
+		Delegate.UE_IF(bMulticast, Broadcast, Execute) EXPAND_ARG_VARIABLES CallArgs; \
+		CHECK_ARG_VARIABLES OutVals \
+		CHECK(GCalledFunctionName == UE::Core::Private::TrimmedMemberFunctionName<#FuncPtr>.Str); \
+	}
+
 void UDynamicDelegateTestObject::RunTests()
 {
+	UDynamicDelegateTestObject* TestOut = NewObject<UDynamicDelegateTestObject>();
+
 	// Update ptr args, because they can't be given a reasonable default as a global
 	TGuardValue<const UDynamicDelegateTestObject*>          PtrArgGuard(GPtrArg, this);
 	TGuardValue<TWeakObjectPtr<UDynamicDelegateTestObject>> WeakPtrArgGuard(GWeakObjPtrArg, this);
+	TGuardValue<const UDynamicDelegateTestObject*>          OutPtrArgGuard(GPtrOut, TestOut);
+	TGuardValue<TWeakObjectPtr<UDynamicDelegateTestObject>> WeakPtrOutGuard(GWeakObjPtrOut, TestOut);
 
 	// 0 param function
 	{
 		// 0 arg payload
-		BIND_CALL_CHECK_DELEGATE       (         FTestDynamicMulticastDelegate0::FDelegate, 0, &UDynamicDelegateTestObject::Func0,            ());
-		BIND_CALL_CHECK_DELEGATE       (         FTestDynamicDelegate0,                     0, &UDynamicDelegateTestObject::Func0,            ());
-		BIND_CALL_CHECK_DELEGATE       (         FTestDynamicMulticastDelegate0,            1, &UDynamicDelegateTestObject::Func0,            ());
-		BIND_CALL_CHECK_DELEGATE       (         FTestDynamicMulticastDelegate0::FDelegate, 0, &UDynamicDelegateTestObject::ConstFunc0,       ());
-		BIND_CALL_CHECK_DELEGATE       (         FTestDynamicDelegate0,                     0, &UDynamicDelegateTestObject::ConstFunc0,       ());
-		BIND_CALL_CHECK_DELEGATE       (         FTestDynamicMulticastDelegate0,            1, &UDynamicDelegateTestObject::ConstFunc0,       ());
-		BIND_CALL_CHECK_DELEGATE_RETVAL(GIntRet, FTestDynamicDelegateRetInt0,               0, &UDynamicDelegateTestObject::FuncRetInt0,      ());
-		BIND_CALL_CHECK_DELEGATE_RETVAL(GIntRet, FTestDynamicDelegateRetInt0,               0, &UDynamicDelegateTestObject::ConstFuncRetInt0, ());
-		BIND_CALL_CHECK_DELEGATE_RETVAL(GStrRet, FTestDynamicDelegateRetStr0,               0, &UDynamicDelegateTestObject::FuncRetStr0,      ());
-		BIND_CALL_CHECK_DELEGATE_RETVAL(GStrRet, FTestDynamicDelegateRetStr0,               0, &UDynamicDelegateTestObject::ConstFuncRetStr0, ());
+		BIND_CALL_CHECK_DELEGATE(FTestDynamicMulticastDelegate0::FDelegate, 0, &UDynamicDelegateTestObject::Func0, ());
+		BIND_CALL_CHECK_DELEGATE(FTestDynamicDelegate0, 0, &UDynamicDelegateTestObject::Func0, ());
+		BIND_CALL_CHECK_DELEGATE(FTestDynamicMulticastDelegate0, 1, &UDynamicDelegateTestObject::Func0, ());
+		BIND_CALL_CHECK_DELEGATE(FTestDynamicMulticastDelegate0::FDelegate, 0, &UDynamicDelegateTestObject::ConstFunc0, ());
+		BIND_CALL_CHECK_DELEGATE(FTestDynamicDelegate0, 0, &UDynamicDelegateTestObject::ConstFunc0, ());
+		BIND_CALL_CHECK_DELEGATE(FTestDynamicMulticastDelegate0, 1, &UDynamicDelegateTestObject::ConstFunc0, ());
+		BIND_CALL_CHECK_DELEGATE_RETVAL(GIntRet, FTestDynamicDelegateRetInt0, 0, &UDynamicDelegateTestObject::FuncRetInt0, ());
+		BIND_CALL_CHECK_DELEGATE_RETVAL(GIntRet, FTestDynamicDelegateRetInt0, 0, &UDynamicDelegateTestObject::ConstFuncRetInt0, ());
+		BIND_CALL_CHECK_DELEGATE_RETVAL(GStrRet, FTestDynamicDelegateRetStr0, 0, &UDynamicDelegateTestObject::FuncRetStr0, ());
+		BIND_CALL_CHECK_DELEGATE_RETVAL(GStrRet, FTestDynamicDelegateRetStr0, 0, &UDynamicDelegateTestObject::ConstFuncRetStr0, ());
 	}
 
 	// 1 param function
 	{
 		// 0 arg payload
-		BIND_CALL_CHECK_DELEGATE       (         FTestDynamicMulticastDelegate1::FDelegate, 0, &UDynamicDelegateTestObject::Func1,            (GIntArg));
-		BIND_CALL_CHECK_DELEGATE       (         FTestDynamicDelegate1,                     0, &UDynamicDelegateTestObject::Func1,            (GIntArg));
-		BIND_CALL_CHECK_DELEGATE       (         FTestDynamicMulticastDelegate1,            1, &UDynamicDelegateTestObject::Func1,            (GIntArg));
-		BIND_CALL_CHECK_DELEGATE       (         FTestDynamicMulticastDelegate1::FDelegate, 0, &UDynamicDelegateTestObject::ConstFunc1,       (GIntArg));
-		BIND_CALL_CHECK_DELEGATE       (         FTestDynamicDelegate1,                     0, &UDynamicDelegateTestObject::ConstFunc1,       (GIntArg));
-		BIND_CALL_CHECK_DELEGATE       (         FTestDynamicMulticastDelegate1,            1, &UDynamicDelegateTestObject::ConstFunc1,       (GIntArg));
-		BIND_CALL_CHECK_DELEGATE_RETVAL(GIntRet, FTestDynamicDelegateRetInt1,               0, &UDynamicDelegateTestObject::FuncRetInt1,      (GIntArg));
-		BIND_CALL_CHECK_DELEGATE_RETVAL(GIntRet, FTestDynamicDelegateRetInt1,               0, &UDynamicDelegateTestObject::ConstFuncRetInt1, (GIntArg));
-		BIND_CALL_CHECK_DELEGATE_RETVAL(GStrRet, FTestDynamicDelegateRetStr1,               0, &UDynamicDelegateTestObject::FuncRetStr1,      (GIntArg));
-		BIND_CALL_CHECK_DELEGATE_RETVAL(GStrRet, FTestDynamicDelegateRetStr1,               0, &UDynamicDelegateTestObject::ConstFuncRetStr1, (GIntArg));
+		BIND_CALL_CHECK_DELEGATE(FTestDynamicMulticastDelegate1::FDelegate, 0, &UDynamicDelegateTestObject::Func1, (GIntArg));
+		BIND_CALL_CHECK_DELEGATE(FTestDynamicDelegate1, 0, &UDynamicDelegateTestObject::Func1, (GIntArg));
+		BIND_CALL_CHECK_DELEGATE(FTestDynamicMulticastDelegate1, 1, &UDynamicDelegateTestObject::Func1, (GIntArg));
+		BIND_CALL_CHECK_DELEGATE(FTestDynamicMulticastDelegate1::FDelegate, 0, &UDynamicDelegateTestObject::ConstFunc1, (GIntArg));
+		BIND_CALL_CHECK_DELEGATE(FTestDynamicDelegate1, 0, &UDynamicDelegateTestObject::ConstFunc1, (GIntArg));
+		BIND_CALL_CHECK_DELEGATE(FTestDynamicMulticastDelegate1, 1, &UDynamicDelegateTestObject::ConstFunc1, (GIntArg));
+		BIND_CALL_CHECK_DELEGATE_RETVAL(GIntRet, FTestDynamicDelegateRetInt1, 0, &UDynamicDelegateTestObject::FuncRetInt1, (GIntArg));
+		BIND_CALL_CHECK_DELEGATE_RETVAL(GIntRet, FTestDynamicDelegateRetInt1, 0, &UDynamicDelegateTestObject::ConstFuncRetInt1, (GIntArg));
+		BIND_CALL_CHECK_DELEGATE_RETVAL(GStrRet, FTestDynamicDelegateRetStr1, 0, &UDynamicDelegateTestObject::FuncRetStr1, (GIntArg));
+		BIND_CALL_CHECK_DELEGATE_RETVAL(GStrRet, FTestDynamicDelegateRetStr1, 0, &UDynamicDelegateTestObject::ConstFuncRetStr1, (GIntArg));
 
 #if defined(UE_USE_DYNAMIC_DELEGATE_PAYLOADS) && UE_USE_DYNAMIC_DELEGATE_PAYLOADS
 		// 1 arg payload
-		BIND_CALL_CHECK_DELEGATE       (         FTestDynamicMulticastDelegate0::FDelegate, 0, &UDynamicDelegateTestObject::Func1,            (), GIntArg);
-		BIND_CALL_CHECK_DELEGATE       (         FTestDynamicDelegate0,                     0, &UDynamicDelegateTestObject::Func1,            (), GIntArg);
-		BIND_CALL_CHECK_DELEGATE       (         FTestDynamicMulticastDelegate0,            1, &UDynamicDelegateTestObject::Func1,            (), GIntArg);
-		BIND_CALL_CHECK_DELEGATE       (         FTestDynamicMulticastDelegate0::FDelegate, 0, &UDynamicDelegateTestObject::ConstFunc1,       (), GIntArg);
-		BIND_CALL_CHECK_DELEGATE       (         FTestDynamicDelegate0,                     0, &UDynamicDelegateTestObject::ConstFunc1,       (), GIntArg);
-		BIND_CALL_CHECK_DELEGATE       (         FTestDynamicMulticastDelegate0,            1, &UDynamicDelegateTestObject::ConstFunc1,       (), GIntArg);
+		BIND_CALL_CHECK_DELEGATE(FTestDynamicMulticastDelegate0::FDelegate, 0, &UDynamicDelegateTestObject::Func1, (), GIntArg);
+		BIND_CALL_CHECK_DELEGATE(FTestDynamicDelegate0, 0, &UDynamicDelegateTestObject::Func1, (), GIntArg);
+		BIND_CALL_CHECK_DELEGATE(FTestDynamicMulticastDelegate0, 1, &UDynamicDelegateTestObject::Func1, (), GIntArg);
+		BIND_CALL_CHECK_DELEGATE(FTestDynamicMulticastDelegate0::FDelegate, 0, &UDynamicDelegateTestObject::ConstFunc1, (), GIntArg);
+		BIND_CALL_CHECK_DELEGATE(FTestDynamicDelegate0, 0, &UDynamicDelegateTestObject::ConstFunc1, (), GIntArg);
+		BIND_CALL_CHECK_DELEGATE(FTestDynamicMulticastDelegate0, 1, &UDynamicDelegateTestObject::ConstFunc1, (), GIntArg);
 #if UE_USE_DYNAMIC_DELEGATE_PAYLOADS_RETVAL
-		BIND_CALL_CHECK_DELEGATE_RETVAL(GIntRet, FTestDynamicDelegateRetInt0,               0, &UDynamicDelegateTestObject::FuncRetInt1,      (), GIntArg);
-		BIND_CALL_CHECK_DELEGATE_RETVAL(GIntRet, FTestDynamicDelegateRetInt0,               0, &UDynamicDelegateTestObject::ConstFuncRetInt1, (), GIntArg);
-		BIND_CALL_CHECK_DELEGATE_RETVAL(GStrRet, FTestDynamicDelegateRetStr0,               0, &UDynamicDelegateTestObject::FuncRetStr1,      (), GIntArg);
-		BIND_CALL_CHECK_DELEGATE_RETVAL(GStrRet, FTestDynamicDelegateRetStr0,               0, &UDynamicDelegateTestObject::ConstFuncRetStr1, (), GIntArg);
+		BIND_CALL_CHECK_DELEGATE_RETVAL(GIntRet, FTestDynamicDelegateRetInt0, 0, &UDynamicDelegateTestObject::FuncRetInt1, (), GIntArg);
+		BIND_CALL_CHECK_DELEGATE_RETVAL(GIntRet, FTestDynamicDelegateRetInt0, 0, &UDynamicDelegateTestObject::ConstFuncRetInt1, (), GIntArg);
+		BIND_CALL_CHECK_DELEGATE_RETVAL(GStrRet, FTestDynamicDelegateRetStr0, 0, &UDynamicDelegateTestObject::FuncRetStr1, (), GIntArg);
+		BIND_CALL_CHECK_DELEGATE_RETVAL(GStrRet, FTestDynamicDelegateRetStr0, 0, &UDynamicDelegateTestObject::ConstFuncRetStr1, (), GIntArg);
 #endif
 #endif
 	}
@@ -432,44 +566,44 @@ void UDynamicDelegateTestObject::RunTests()
 	// 2 param function
 	{
 		// 0 arg payload
-		BIND_CALL_CHECK_DELEGATE       (         FTestDynamicMulticastDelegate2::FDelegate, 0, &UDynamicDelegateTestObject::Func2,            (GIntArg, GStrArg));
-		BIND_CALL_CHECK_DELEGATE       (         FTestDynamicDelegate2,                     0, &UDynamicDelegateTestObject::Func2,            (GIntArg, GStrArg));
-		BIND_CALL_CHECK_DELEGATE       (         FTestDynamicMulticastDelegate2,            1, &UDynamicDelegateTestObject::Func2,            (GIntArg, GStrArg));
-		BIND_CALL_CHECK_DELEGATE       (         FTestDynamicMulticastDelegate2::FDelegate, 0, &UDynamicDelegateTestObject::ConstFunc2,       (GIntArg, GStrArg));
-		BIND_CALL_CHECK_DELEGATE       (         FTestDynamicDelegate2,                     0, &UDynamicDelegateTestObject::ConstFunc2,       (GIntArg, GStrArg));
-		BIND_CALL_CHECK_DELEGATE       (         FTestDynamicMulticastDelegate2,            1, &UDynamicDelegateTestObject::ConstFunc2,       (GIntArg, GStrArg));
-		BIND_CALL_CHECK_DELEGATE_RETVAL(GIntRet, FTestDynamicDelegateRetInt2,               0, &UDynamicDelegateTestObject::FuncRetInt2,      (GIntArg, GStrArg));
-		BIND_CALL_CHECK_DELEGATE_RETVAL(GIntRet, FTestDynamicDelegateRetInt2,               0, &UDynamicDelegateTestObject::ConstFuncRetInt2, (GIntArg, GStrArg));
-		BIND_CALL_CHECK_DELEGATE_RETVAL(GStrRet, FTestDynamicDelegateRetStr2,               0, &UDynamicDelegateTestObject::FuncRetStr2,      (GIntArg, GStrArg));
-		BIND_CALL_CHECK_DELEGATE_RETVAL(GStrRet, FTestDynamicDelegateRetStr2,               0, &UDynamicDelegateTestObject::ConstFuncRetStr2, (GIntArg, GStrArg));
+		BIND_CALL_CHECK_DELEGATE(FTestDynamicMulticastDelegate2::FDelegate, 0, &UDynamicDelegateTestObject::Func2, (GIntArg, GStrArg));
+		BIND_CALL_CHECK_DELEGATE(FTestDynamicDelegate2, 0, &UDynamicDelegateTestObject::Func2, (GIntArg, GStrArg));
+		BIND_CALL_CHECK_DELEGATE(FTestDynamicMulticastDelegate2, 1, &UDynamicDelegateTestObject::Func2, (GIntArg, GStrArg));
+		BIND_CALL_CHECK_DELEGATE(FTestDynamicMulticastDelegate2::FDelegate, 0, &UDynamicDelegateTestObject::ConstFunc2, (GIntArg, GStrArg));
+		BIND_CALL_CHECK_DELEGATE(FTestDynamicDelegate2, 0, &UDynamicDelegateTestObject::ConstFunc2, (GIntArg, GStrArg));
+		BIND_CALL_CHECK_DELEGATE(FTestDynamicMulticastDelegate2, 1, &UDynamicDelegateTestObject::ConstFunc2, (GIntArg, GStrArg));
+		BIND_CALL_CHECK_DELEGATE_RETVAL(GIntRet, FTestDynamicDelegateRetInt2, 0, &UDynamicDelegateTestObject::FuncRetInt2, (GIntArg, GStrArg));
+		BIND_CALL_CHECK_DELEGATE_RETVAL(GIntRet, FTestDynamicDelegateRetInt2, 0, &UDynamicDelegateTestObject::ConstFuncRetInt2, (GIntArg, GStrArg));
+		BIND_CALL_CHECK_DELEGATE_RETVAL(GStrRet, FTestDynamicDelegateRetStr2, 0, &UDynamicDelegateTestObject::FuncRetStr2, (GIntArg, GStrArg));
+		BIND_CALL_CHECK_DELEGATE_RETVAL(GStrRet, FTestDynamicDelegateRetStr2, 0, &UDynamicDelegateTestObject::ConstFuncRetStr2, (GIntArg, GStrArg));
 
 #if defined(UE_USE_DYNAMIC_DELEGATE_PAYLOADS) && UE_USE_DYNAMIC_DELEGATE_PAYLOADS
 		// 1 arg payload
-		BIND_CALL_CHECK_DELEGATE       (         FTestDynamicMulticastDelegate1::FDelegate, 0, &UDynamicDelegateTestObject::Func2,            (GIntArg), GStrArg);
-		BIND_CALL_CHECK_DELEGATE       (         FTestDynamicDelegate1,                     0, &UDynamicDelegateTestObject::Func2,            (GIntArg), GStrArg);
-		BIND_CALL_CHECK_DELEGATE       (         FTestDynamicMulticastDelegate1,            1, &UDynamicDelegateTestObject::Func2,            (GIntArg), GStrArg);
-		BIND_CALL_CHECK_DELEGATE       (         FTestDynamicMulticastDelegate1::FDelegate, 0, &UDynamicDelegateTestObject::ConstFunc2,       (GIntArg), GStrArg);
-		BIND_CALL_CHECK_DELEGATE       (         FTestDynamicDelegate1,                     0, &UDynamicDelegateTestObject::ConstFunc2,       (GIntArg), GStrArg);
-		BIND_CALL_CHECK_DELEGATE       (         FTestDynamicMulticastDelegate1,            1, &UDynamicDelegateTestObject::ConstFunc2,       (GIntArg), GStrArg);
+		BIND_CALL_CHECK_DELEGATE(FTestDynamicMulticastDelegate1::FDelegate, 0, &UDynamicDelegateTestObject::Func2, (GIntArg), GStrArg);
+		BIND_CALL_CHECK_DELEGATE(FTestDynamicDelegate1, 0, &UDynamicDelegateTestObject::Func2, (GIntArg), GStrArg);
+		BIND_CALL_CHECK_DELEGATE(FTestDynamicMulticastDelegate1, 1, &UDynamicDelegateTestObject::Func2, (GIntArg), GStrArg);
+		BIND_CALL_CHECK_DELEGATE(FTestDynamicMulticastDelegate1::FDelegate, 0, &UDynamicDelegateTestObject::ConstFunc2, (GIntArg), GStrArg);
+		BIND_CALL_CHECK_DELEGATE(FTestDynamicDelegate1, 0, &UDynamicDelegateTestObject::ConstFunc2, (GIntArg), GStrArg);
+		BIND_CALL_CHECK_DELEGATE(FTestDynamicMulticastDelegate1, 1, &UDynamicDelegateTestObject::ConstFunc2, (GIntArg), GStrArg);
 #if UE_USE_DYNAMIC_DELEGATE_PAYLOADS_RETVAL
-		BIND_CALL_CHECK_DELEGATE_RETVAL(GIntRet, FTestDynamicDelegateRetInt1,               0, &UDynamicDelegateTestObject::FuncRetInt2,      (GIntArg), GStrArg);
-		BIND_CALL_CHECK_DELEGATE_RETVAL(GIntRet, FTestDynamicDelegateRetInt1,               0, &UDynamicDelegateTestObject::ConstFuncRetInt2, (GIntArg), GStrArg);
-		BIND_CALL_CHECK_DELEGATE_RETVAL(GStrRet, FTestDynamicDelegateRetStr1,               0, &UDynamicDelegateTestObject::FuncRetStr2,      (GIntArg), GStrArg);
-		BIND_CALL_CHECK_DELEGATE_RETVAL(GStrRet, FTestDynamicDelegateRetStr1,               0, &UDynamicDelegateTestObject::ConstFuncRetStr2, (GIntArg), GStrArg);
+		BIND_CALL_CHECK_DELEGATE_RETVAL(GIntRet, FTestDynamicDelegateRetInt1, 0, &UDynamicDelegateTestObject::FuncRetInt2, (GIntArg), GStrArg);
+		BIND_CALL_CHECK_DELEGATE_RETVAL(GIntRet, FTestDynamicDelegateRetInt1, 0, &UDynamicDelegateTestObject::ConstFuncRetInt2, (GIntArg), GStrArg);
+		BIND_CALL_CHECK_DELEGATE_RETVAL(GStrRet, FTestDynamicDelegateRetStr1, 0, &UDynamicDelegateTestObject::FuncRetStr2, (GIntArg), GStrArg);
+		BIND_CALL_CHECK_DELEGATE_RETVAL(GStrRet, FTestDynamicDelegateRetStr1, 0, &UDynamicDelegateTestObject::ConstFuncRetStr2, (GIntArg), GStrArg);
 #endif
 
 		// 2 arg payload
-		BIND_CALL_CHECK_DELEGATE       (         FTestDynamicMulticastDelegate0::FDelegate, 0, &UDynamicDelegateTestObject::Func2,            (), GIntArg, GStrArg);
-		BIND_CALL_CHECK_DELEGATE       (         FTestDynamicDelegate0,                     0, &UDynamicDelegateTestObject::Func2,            (), GIntArg, GStrArg);
-		BIND_CALL_CHECK_DELEGATE       (         FTestDynamicMulticastDelegate0,            1, &UDynamicDelegateTestObject::Func2,            (), GIntArg, GStrArg);
-		BIND_CALL_CHECK_DELEGATE       (         FTestDynamicMulticastDelegate0::FDelegate, 0, &UDynamicDelegateTestObject::ConstFunc2,       (), GIntArg, GStrArg);
-		BIND_CALL_CHECK_DELEGATE       (         FTestDynamicDelegate0,                     0, &UDynamicDelegateTestObject::ConstFunc2,       (), GIntArg, GStrArg);
-		BIND_CALL_CHECK_DELEGATE       (         FTestDynamicMulticastDelegate0,            1, &UDynamicDelegateTestObject::ConstFunc2,       (), GIntArg, GStrArg);
+		BIND_CALL_CHECK_DELEGATE(FTestDynamicMulticastDelegate0::FDelegate, 0, &UDynamicDelegateTestObject::Func2, (), GIntArg, GStrArg);
+		BIND_CALL_CHECK_DELEGATE(FTestDynamicDelegate0, 0, &UDynamicDelegateTestObject::Func2, (), GIntArg, GStrArg);
+		BIND_CALL_CHECK_DELEGATE(FTestDynamicMulticastDelegate0, 1, &UDynamicDelegateTestObject::Func2, (), GIntArg, GStrArg);
+		BIND_CALL_CHECK_DELEGATE(FTestDynamicMulticastDelegate0::FDelegate, 0, &UDynamicDelegateTestObject::ConstFunc2, (), GIntArg, GStrArg);
+		BIND_CALL_CHECK_DELEGATE(FTestDynamicDelegate0, 0, &UDynamicDelegateTestObject::ConstFunc2, (), GIntArg, GStrArg);
+		BIND_CALL_CHECK_DELEGATE(FTestDynamicMulticastDelegate0, 1, &UDynamicDelegateTestObject::ConstFunc2, (), GIntArg, GStrArg);
 #if UE_USE_DYNAMIC_DELEGATE_PAYLOADS_RETVAL
-		BIND_CALL_CHECK_DELEGATE_RETVAL(GIntRet, FTestDynamicDelegateRetInt0,               0, &UDynamicDelegateTestObject::FuncRetInt2,      (), GIntArg, GStrArg);
-		BIND_CALL_CHECK_DELEGATE_RETVAL(GIntRet, FTestDynamicDelegateRetInt0,               0, &UDynamicDelegateTestObject::ConstFuncRetInt2, (), GIntArg, GStrArg);
-		BIND_CALL_CHECK_DELEGATE_RETVAL(GStrRet, FTestDynamicDelegateRetStr0,               0, &UDynamicDelegateTestObject::FuncRetStr2,      (), GIntArg, GStrArg);
-		BIND_CALL_CHECK_DELEGATE_RETVAL(GStrRet, FTestDynamicDelegateRetStr0,               0, &UDynamicDelegateTestObject::ConstFuncRetStr2, (), GIntArg, GStrArg);
+		BIND_CALL_CHECK_DELEGATE_RETVAL(GIntRet, FTestDynamicDelegateRetInt0, 0, &UDynamicDelegateTestObject::FuncRetInt2, (), GIntArg, GStrArg);
+		BIND_CALL_CHECK_DELEGATE_RETVAL(GIntRet, FTestDynamicDelegateRetInt0, 0, &UDynamicDelegateTestObject::ConstFuncRetInt2, (), GIntArg, GStrArg);
+		BIND_CALL_CHECK_DELEGATE_RETVAL(GStrRet, FTestDynamicDelegateRetStr0, 0, &UDynamicDelegateTestObject::FuncRetStr2, (), GIntArg, GStrArg);
+		BIND_CALL_CHECK_DELEGATE_RETVAL(GStrRet, FTestDynamicDelegateRetStr0, 0, &UDynamicDelegateTestObject::ConstFuncRetStr2, (), GIntArg, GStrArg);
 #endif
 #endif
 	}
@@ -477,58 +611,58 @@ void UDynamicDelegateTestObject::RunTests()
 	// 3 param function
 	{
 		// 0 arg payload
-		BIND_CALL_CHECK_DELEGATE       (         FTestDynamicMulticastDelegate3::FDelegate, 0, &UDynamicDelegateTestObject::Func3,            (GIntArg, GStrArg, GPtrArg));
-		BIND_CALL_CHECK_DELEGATE       (         FTestDynamicDelegate3,                     0, &UDynamicDelegateTestObject::Func3,            (GIntArg, GStrArg, GPtrArg));
-		BIND_CALL_CHECK_DELEGATE       (         FTestDynamicMulticastDelegate3,            1, &UDynamicDelegateTestObject::Func3,            (GIntArg, GStrArg, GPtrArg));
-		BIND_CALL_CHECK_DELEGATE       (         FTestDynamicMulticastDelegate3::FDelegate, 0, &UDynamicDelegateTestObject::ConstFunc3,       (GIntArg, GStrArg, GPtrArg));
-		BIND_CALL_CHECK_DELEGATE       (         FTestDynamicDelegate3,                     0, &UDynamicDelegateTestObject::ConstFunc3,       (GIntArg, GStrArg, GPtrArg));
-		BIND_CALL_CHECK_DELEGATE       (         FTestDynamicMulticastDelegate3,            1, &UDynamicDelegateTestObject::ConstFunc3,       (GIntArg, GStrArg, GPtrArg));
-		BIND_CALL_CHECK_DELEGATE_RETVAL(GIntRet, FTestDynamicDelegateRetInt3,               0, &UDynamicDelegateTestObject::FuncRetInt3,      (GIntArg, GStrArg, GPtrArg));
-		BIND_CALL_CHECK_DELEGATE_RETVAL(GIntRet, FTestDynamicDelegateRetInt3,               0, &UDynamicDelegateTestObject::ConstFuncRetInt3, (GIntArg, GStrArg, GPtrArg));
-		BIND_CALL_CHECK_DELEGATE_RETVAL(GStrRet, FTestDynamicDelegateRetStr3,               0, &UDynamicDelegateTestObject::FuncRetStr3,      (GIntArg, GStrArg, GPtrArg));
-		BIND_CALL_CHECK_DELEGATE_RETVAL(GStrRet, FTestDynamicDelegateRetStr3,               0, &UDynamicDelegateTestObject::ConstFuncRetStr3, (GIntArg, GStrArg, GPtrArg));
+		BIND_CALL_CHECK_DELEGATE(FTestDynamicMulticastDelegate3::FDelegate, 0, &UDynamicDelegateTestObject::Func3, (GIntArg, GStrArg, GPtrArg));
+		BIND_CALL_CHECK_DELEGATE(FTestDynamicDelegate3, 0, &UDynamicDelegateTestObject::Func3, (GIntArg, GStrArg, GPtrArg));
+		BIND_CALL_CHECK_DELEGATE(FTestDynamicMulticastDelegate3, 1, &UDynamicDelegateTestObject::Func3, (GIntArg, GStrArg, GPtrArg));
+		BIND_CALL_CHECK_DELEGATE(FTestDynamicMulticastDelegate3::FDelegate, 0, &UDynamicDelegateTestObject::ConstFunc3, (GIntArg, GStrArg, GPtrArg));
+		BIND_CALL_CHECK_DELEGATE(FTestDynamicDelegate3, 0, &UDynamicDelegateTestObject::ConstFunc3, (GIntArg, GStrArg, GPtrArg));
+		BIND_CALL_CHECK_DELEGATE(FTestDynamicMulticastDelegate3, 1, &UDynamicDelegateTestObject::ConstFunc3, (GIntArg, GStrArg, GPtrArg));
+		BIND_CALL_CHECK_DELEGATE_RETVAL(GIntRet, FTestDynamicDelegateRetInt3, 0, &UDynamicDelegateTestObject::FuncRetInt3, (GIntArg, GStrArg, GPtrArg));
+		BIND_CALL_CHECK_DELEGATE_RETVAL(GIntRet, FTestDynamicDelegateRetInt3, 0, &UDynamicDelegateTestObject::ConstFuncRetInt3, (GIntArg, GStrArg, GPtrArg));
+		BIND_CALL_CHECK_DELEGATE_RETVAL(GStrRet, FTestDynamicDelegateRetStr3, 0, &UDynamicDelegateTestObject::FuncRetStr3, (GIntArg, GStrArg, GPtrArg));
+		BIND_CALL_CHECK_DELEGATE_RETVAL(GStrRet, FTestDynamicDelegateRetStr3, 0, &UDynamicDelegateTestObject::ConstFuncRetStr3, (GIntArg, GStrArg, GPtrArg));
 
 #if defined(UE_USE_DYNAMIC_DELEGATE_PAYLOADS) && UE_USE_DYNAMIC_DELEGATE_PAYLOADS
 		// 1 arg payload
-		BIND_CALL_CHECK_DELEGATE       (         FTestDynamicMulticastDelegate2::FDelegate, 0, &UDynamicDelegateTestObject::Func3,            (GIntArg, GStrArg), GPtrArg);
-		BIND_CALL_CHECK_DELEGATE       (         FTestDynamicDelegate2,                     0, &UDynamicDelegateTestObject::Func3,            (GIntArg, GStrArg), GPtrArg);
-		BIND_CALL_CHECK_DELEGATE       (         FTestDynamicMulticastDelegate2,            1, &UDynamicDelegateTestObject::Func3,            (GIntArg, GStrArg), GPtrArg);
-		BIND_CALL_CHECK_DELEGATE       (         FTestDynamicMulticastDelegate2::FDelegate, 0, &UDynamicDelegateTestObject::ConstFunc3,       (GIntArg, GStrArg), GPtrArg);
-		BIND_CALL_CHECK_DELEGATE       (         FTestDynamicDelegate2,                     0, &UDynamicDelegateTestObject::ConstFunc3,       (GIntArg, GStrArg), GPtrArg);
-		BIND_CALL_CHECK_DELEGATE       (         FTestDynamicMulticastDelegate2,            1, &UDynamicDelegateTestObject::ConstFunc3,       (GIntArg, GStrArg), GPtrArg);
+		BIND_CALL_CHECK_DELEGATE(FTestDynamicMulticastDelegate2::FDelegate, 0, &UDynamicDelegateTestObject::Func3, (GIntArg, GStrArg), GPtrArg);
+		BIND_CALL_CHECK_DELEGATE(FTestDynamicDelegate2, 0, &UDynamicDelegateTestObject::Func3, (GIntArg, GStrArg), GPtrArg);
+		BIND_CALL_CHECK_DELEGATE(FTestDynamicMulticastDelegate2, 1, &UDynamicDelegateTestObject::Func3, (GIntArg, GStrArg), GPtrArg);
+		BIND_CALL_CHECK_DELEGATE(FTestDynamicMulticastDelegate2::FDelegate, 0, &UDynamicDelegateTestObject::ConstFunc3, (GIntArg, GStrArg), GPtrArg);
+		BIND_CALL_CHECK_DELEGATE(FTestDynamicDelegate2, 0, &UDynamicDelegateTestObject::ConstFunc3, (GIntArg, GStrArg), GPtrArg);
+		BIND_CALL_CHECK_DELEGATE(FTestDynamicMulticastDelegate2, 1, &UDynamicDelegateTestObject::ConstFunc3, (GIntArg, GStrArg), GPtrArg);
 #if UE_USE_DYNAMIC_DELEGATE_PAYLOADS_RETVAL
-		BIND_CALL_CHECK_DELEGATE_RETVAL(GIntRet, FTestDynamicDelegateRetInt2,               0, &UDynamicDelegateTestObject::FuncRetInt3,      (GIntArg, GStrArg), GPtrArg);
-		BIND_CALL_CHECK_DELEGATE_RETVAL(GIntRet, FTestDynamicDelegateRetInt2,               0, &UDynamicDelegateTestObject::ConstFuncRetInt3, (GIntArg, GStrArg), GPtrArg);
-		BIND_CALL_CHECK_DELEGATE_RETVAL(GStrRet, FTestDynamicDelegateRetStr2,               0, &UDynamicDelegateTestObject::FuncRetStr3,      (GIntArg, GStrArg), GPtrArg);
-		BIND_CALL_CHECK_DELEGATE_RETVAL(GStrRet, FTestDynamicDelegateRetStr2,               0, &UDynamicDelegateTestObject::ConstFuncRetStr3, (GIntArg, GStrArg), GPtrArg);
+		BIND_CALL_CHECK_DELEGATE_RETVAL(GIntRet, FTestDynamicDelegateRetInt2, 0, &UDynamicDelegateTestObject::FuncRetInt3, (GIntArg, GStrArg), GPtrArg);
+		BIND_CALL_CHECK_DELEGATE_RETVAL(GIntRet, FTestDynamicDelegateRetInt2, 0, &UDynamicDelegateTestObject::ConstFuncRetInt3, (GIntArg, GStrArg), GPtrArg);
+		BIND_CALL_CHECK_DELEGATE_RETVAL(GStrRet, FTestDynamicDelegateRetStr2, 0, &UDynamicDelegateTestObject::FuncRetStr3, (GIntArg, GStrArg), GPtrArg);
+		BIND_CALL_CHECK_DELEGATE_RETVAL(GStrRet, FTestDynamicDelegateRetStr2, 0, &UDynamicDelegateTestObject::ConstFuncRetStr3, (GIntArg, GStrArg), GPtrArg);
 #endif
 
 		// 2 arg payload
-		BIND_CALL_CHECK_DELEGATE       (         FTestDynamicMulticastDelegate1::FDelegate, 0, &UDynamicDelegateTestObject::Func3,            (GIntArg), GStrArg, GPtrArg);
-		BIND_CALL_CHECK_DELEGATE       (         FTestDynamicDelegate1,                     0, &UDynamicDelegateTestObject::Func3,            (GIntArg), GStrArg, GPtrArg);
-		BIND_CALL_CHECK_DELEGATE       (         FTestDynamicMulticastDelegate1,            1, &UDynamicDelegateTestObject::Func3,            (GIntArg), GStrArg, GPtrArg);
-		BIND_CALL_CHECK_DELEGATE       (         FTestDynamicMulticastDelegate1::FDelegate, 0, &UDynamicDelegateTestObject::ConstFunc3,       (GIntArg), GStrArg, GPtrArg);
-		BIND_CALL_CHECK_DELEGATE       (         FTestDynamicDelegate1,                     0, &UDynamicDelegateTestObject::ConstFunc3,       (GIntArg), GStrArg, GPtrArg);
-		BIND_CALL_CHECK_DELEGATE       (         FTestDynamicMulticastDelegate1,            1, &UDynamicDelegateTestObject::ConstFunc3,       (GIntArg), GStrArg, GPtrArg);
+		BIND_CALL_CHECK_DELEGATE(FTestDynamicMulticastDelegate1::FDelegate, 0, &UDynamicDelegateTestObject::Func3, (GIntArg), GStrArg, GPtrArg);
+		BIND_CALL_CHECK_DELEGATE(FTestDynamicDelegate1, 0, &UDynamicDelegateTestObject::Func3, (GIntArg), GStrArg, GPtrArg);
+		BIND_CALL_CHECK_DELEGATE(FTestDynamicMulticastDelegate1, 1, &UDynamicDelegateTestObject::Func3, (GIntArg), GStrArg, GPtrArg);
+		BIND_CALL_CHECK_DELEGATE(FTestDynamicMulticastDelegate1::FDelegate, 0, &UDynamicDelegateTestObject::ConstFunc3, (GIntArg), GStrArg, GPtrArg);
+		BIND_CALL_CHECK_DELEGATE(FTestDynamicDelegate1, 0, &UDynamicDelegateTestObject::ConstFunc3, (GIntArg), GStrArg, GPtrArg);
+		BIND_CALL_CHECK_DELEGATE(FTestDynamicMulticastDelegate1, 1, &UDynamicDelegateTestObject::ConstFunc3, (GIntArg), GStrArg, GPtrArg);
 #if UE_USE_DYNAMIC_DELEGATE_PAYLOADS_RETVAL
-		BIND_CALL_CHECK_DELEGATE_RETVAL(GIntRet, FTestDynamicDelegateRetInt1,               0, &UDynamicDelegateTestObject::FuncRetInt3,      (GIntArg), GStrArg, GPtrArg);
-		BIND_CALL_CHECK_DELEGATE_RETVAL(GIntRet, FTestDynamicDelegateRetInt1,               0, &UDynamicDelegateTestObject::ConstFuncRetInt3, (GIntArg), GStrArg, GPtrArg);
-		BIND_CALL_CHECK_DELEGATE_RETVAL(GStrRet, FTestDynamicDelegateRetStr1,               0, &UDynamicDelegateTestObject::FuncRetStr3,      (GIntArg), GStrArg, GPtrArg);
-		BIND_CALL_CHECK_DELEGATE_RETVAL(GStrRet, FTestDynamicDelegateRetStr1,               0, &UDynamicDelegateTestObject::ConstFuncRetStr3, (GIntArg), GStrArg, GPtrArg);
+		BIND_CALL_CHECK_DELEGATE_RETVAL(GIntRet, FTestDynamicDelegateRetInt1, 0, &UDynamicDelegateTestObject::FuncRetInt3, (GIntArg), GStrArg, GPtrArg);
+		BIND_CALL_CHECK_DELEGATE_RETVAL(GIntRet, FTestDynamicDelegateRetInt1, 0, &UDynamicDelegateTestObject::ConstFuncRetInt3, (GIntArg), GStrArg, GPtrArg);
+		BIND_CALL_CHECK_DELEGATE_RETVAL(GStrRet, FTestDynamicDelegateRetStr1, 0, &UDynamicDelegateTestObject::FuncRetStr3, (GIntArg), GStrArg, GPtrArg);
+		BIND_CALL_CHECK_DELEGATE_RETVAL(GStrRet, FTestDynamicDelegateRetStr1, 0, &UDynamicDelegateTestObject::ConstFuncRetStr3, (GIntArg), GStrArg, GPtrArg);
 #endif
 
 		// 3 arg payload
-		BIND_CALL_CHECK_DELEGATE       (         FTestDynamicMulticastDelegate0::FDelegate, 0, &UDynamicDelegateTestObject::Func3,            (), GIntArg, GStrArg, GPtrArg);
-		BIND_CALL_CHECK_DELEGATE       (         FTestDynamicDelegate0,                     0, &UDynamicDelegateTestObject::Func3,            (), GIntArg, GStrArg, GPtrArg);
-		BIND_CALL_CHECK_DELEGATE       (         FTestDynamicMulticastDelegate0,            1, &UDynamicDelegateTestObject::Func3,            (), GIntArg, GStrArg, GPtrArg);
-		BIND_CALL_CHECK_DELEGATE       (         FTestDynamicMulticastDelegate0::FDelegate, 0, &UDynamicDelegateTestObject::ConstFunc3,       (), GIntArg, GStrArg, GPtrArg);
-		BIND_CALL_CHECK_DELEGATE       (         FTestDynamicDelegate0,                     0, &UDynamicDelegateTestObject::ConstFunc3,       (), GIntArg, GStrArg, GPtrArg);
-		BIND_CALL_CHECK_DELEGATE       (         FTestDynamicMulticastDelegate0,            1, &UDynamicDelegateTestObject::ConstFunc3,       (), GIntArg, GStrArg, GPtrArg);
+		BIND_CALL_CHECK_DELEGATE(FTestDynamicMulticastDelegate0::FDelegate, 0, &UDynamicDelegateTestObject::Func3, (), GIntArg, GStrArg, GPtrArg);
+		BIND_CALL_CHECK_DELEGATE(FTestDynamicDelegate0, 0, &UDynamicDelegateTestObject::Func3, (), GIntArg, GStrArg, GPtrArg);
+		BIND_CALL_CHECK_DELEGATE(FTestDynamicMulticastDelegate0, 1, &UDynamicDelegateTestObject::Func3, (), GIntArg, GStrArg, GPtrArg);
+		BIND_CALL_CHECK_DELEGATE(FTestDynamicMulticastDelegate0::FDelegate, 0, &UDynamicDelegateTestObject::ConstFunc3, (), GIntArg, GStrArg, GPtrArg);
+		BIND_CALL_CHECK_DELEGATE(FTestDynamicDelegate0, 0, &UDynamicDelegateTestObject::ConstFunc3, (), GIntArg, GStrArg, GPtrArg);
+		BIND_CALL_CHECK_DELEGATE(FTestDynamicMulticastDelegate0, 1, &UDynamicDelegateTestObject::ConstFunc3, (), GIntArg, GStrArg, GPtrArg);
 #if UE_USE_DYNAMIC_DELEGATE_PAYLOADS_RETVAL
-		BIND_CALL_CHECK_DELEGATE_RETVAL(GIntRet, FTestDynamicDelegateRetInt0,               0, &UDynamicDelegateTestObject::FuncRetInt3,      (), GIntArg, GStrArg, GPtrArg);
-		BIND_CALL_CHECK_DELEGATE_RETVAL(GIntRet, FTestDynamicDelegateRetInt0,               0, &UDynamicDelegateTestObject::ConstFuncRetInt3, (), GIntArg, GStrArg, GPtrArg);
-		BIND_CALL_CHECK_DELEGATE_RETVAL(GStrRet, FTestDynamicDelegateRetStr0,               0, &UDynamicDelegateTestObject::FuncRetStr3,      (), GIntArg, GStrArg, GPtrArg);
-		BIND_CALL_CHECK_DELEGATE_RETVAL(GStrRet, FTestDynamicDelegateRetStr0,               0, &UDynamicDelegateTestObject::ConstFuncRetStr3, (), GIntArg, GStrArg, GPtrArg);
+		BIND_CALL_CHECK_DELEGATE_RETVAL(GIntRet, FTestDynamicDelegateRetInt0, 0, &UDynamicDelegateTestObject::FuncRetInt3, (), GIntArg, GStrArg, GPtrArg);
+		BIND_CALL_CHECK_DELEGATE_RETVAL(GIntRet, FTestDynamicDelegateRetInt0, 0, &UDynamicDelegateTestObject::ConstFuncRetInt3, (), GIntArg, GStrArg, GPtrArg);
+		BIND_CALL_CHECK_DELEGATE_RETVAL(GStrRet, FTestDynamicDelegateRetStr0, 0, &UDynamicDelegateTestObject::FuncRetStr3, (), GIntArg, GStrArg, GPtrArg);
+		BIND_CALL_CHECK_DELEGATE_RETVAL(GStrRet, FTestDynamicDelegateRetStr0, 0, &UDynamicDelegateTestObject::ConstFuncRetStr3, (), GIntArg, GStrArg, GPtrArg);
 #endif
 #endif
 	}
@@ -536,72 +670,72 @@ void UDynamicDelegateTestObject::RunTests()
 	// 4 param function
 	{
 		// 0 arg payload
-		BIND_CALL_CHECK_DELEGATE       (         FTestDynamicMulticastDelegate4::FDelegate, 0, &UDynamicDelegateTestObject::Func4,            (GIntArg, GStrArg, GPtrArg, GNameArg));
-		BIND_CALL_CHECK_DELEGATE       (         FTestDynamicDelegate4,                     0, &UDynamicDelegateTestObject::Func4,            (GIntArg, GStrArg, GPtrArg, GNameArg));
-		BIND_CALL_CHECK_DELEGATE       (         FTestDynamicMulticastDelegate4,            1, &UDynamicDelegateTestObject::Func4,            (GIntArg, GStrArg, GPtrArg, GNameArg));
-		BIND_CALL_CHECK_DELEGATE       (         FTestDynamicMulticastDelegate4::FDelegate, 0, &UDynamicDelegateTestObject::ConstFunc4,       (GIntArg, GStrArg, GPtrArg, GNameArg));
-		BIND_CALL_CHECK_DELEGATE       (         FTestDynamicDelegate4,                     0, &UDynamicDelegateTestObject::ConstFunc4,       (GIntArg, GStrArg, GPtrArg, GNameArg));
-		BIND_CALL_CHECK_DELEGATE       (         FTestDynamicMulticastDelegate4,            1, &UDynamicDelegateTestObject::ConstFunc4,       (GIntArg, GStrArg, GPtrArg, GNameArg));
-		BIND_CALL_CHECK_DELEGATE_RETVAL(GIntRet, FTestDynamicDelegateRetInt4,               0, &UDynamicDelegateTestObject::FuncRetInt4,      (GIntArg, GStrArg, GPtrArg, GNameArg));
-		BIND_CALL_CHECK_DELEGATE_RETVAL(GIntRet, FTestDynamicDelegateRetInt4,               0, &UDynamicDelegateTestObject::ConstFuncRetInt4, (GIntArg, GStrArg, GPtrArg, GNameArg));
-		BIND_CALL_CHECK_DELEGATE_RETVAL(GStrRet, FTestDynamicDelegateRetStr4,               0, &UDynamicDelegateTestObject::FuncRetStr4,      (GIntArg, GStrArg, GPtrArg, GNameArg));
-		BIND_CALL_CHECK_DELEGATE_RETVAL(GStrRet, FTestDynamicDelegateRetStr4,               0, &UDynamicDelegateTestObject::ConstFuncRetStr4, (GIntArg, GStrArg, GPtrArg, GNameArg));
+		BIND_CALL_CHECK_DELEGATE(FTestDynamicMulticastDelegate4::FDelegate, 0, &UDynamicDelegateTestObject::Func4, (GIntArg, GStrArg, GPtrArg, GNameArg));
+		BIND_CALL_CHECK_DELEGATE(FTestDynamicDelegate4, 0, &UDynamicDelegateTestObject::Func4, (GIntArg, GStrArg, GPtrArg, GNameArg));
+		BIND_CALL_CHECK_DELEGATE(FTestDynamicMulticastDelegate4, 1, &UDynamicDelegateTestObject::Func4, (GIntArg, GStrArg, GPtrArg, GNameArg));
+		BIND_CALL_CHECK_DELEGATE(FTestDynamicMulticastDelegate4::FDelegate, 0, &UDynamicDelegateTestObject::ConstFunc4, (GIntArg, GStrArg, GPtrArg, GNameArg));
+		BIND_CALL_CHECK_DELEGATE(FTestDynamicDelegate4, 0, &UDynamicDelegateTestObject::ConstFunc4, (GIntArg, GStrArg, GPtrArg, GNameArg));
+		BIND_CALL_CHECK_DELEGATE(FTestDynamicMulticastDelegate4, 1, &UDynamicDelegateTestObject::ConstFunc4, (GIntArg, GStrArg, GPtrArg, GNameArg));
+		BIND_CALL_CHECK_DELEGATE_RETVAL(GIntRet, FTestDynamicDelegateRetInt4, 0, &UDynamicDelegateTestObject::FuncRetInt4, (GIntArg, GStrArg, GPtrArg, GNameArg));
+		BIND_CALL_CHECK_DELEGATE_RETVAL(GIntRet, FTestDynamicDelegateRetInt4, 0, &UDynamicDelegateTestObject::ConstFuncRetInt4, (GIntArg, GStrArg, GPtrArg, GNameArg));
+		BIND_CALL_CHECK_DELEGATE_RETVAL(GStrRet, FTestDynamicDelegateRetStr4, 0, &UDynamicDelegateTestObject::FuncRetStr4, (GIntArg, GStrArg, GPtrArg, GNameArg));
+		BIND_CALL_CHECK_DELEGATE_RETVAL(GStrRet, FTestDynamicDelegateRetStr4, 0, &UDynamicDelegateTestObject::ConstFuncRetStr4, (GIntArg, GStrArg, GPtrArg, GNameArg));
 
 #if defined(UE_USE_DYNAMIC_DELEGATE_PAYLOADS) && UE_USE_DYNAMIC_DELEGATE_PAYLOADS
 		// 1 arg payload
-		BIND_CALL_CHECK_DELEGATE       (         FTestDynamicMulticastDelegate3::FDelegate, 0, &UDynamicDelegateTestObject::Func4,            (GIntArg, GStrArg, GPtrArg), GNameArg);
-		BIND_CALL_CHECK_DELEGATE       (         FTestDynamicDelegate3,                     0, &UDynamicDelegateTestObject::Func4,            (GIntArg, GStrArg, GPtrArg), GNameArg);
-		BIND_CALL_CHECK_DELEGATE       (         FTestDynamicMulticastDelegate3,            1, &UDynamicDelegateTestObject::Func4,            (GIntArg, GStrArg, GPtrArg), GNameArg);
-		BIND_CALL_CHECK_DELEGATE       (         FTestDynamicMulticastDelegate3::FDelegate, 0, &UDynamicDelegateTestObject::ConstFunc4,       (GIntArg, GStrArg, GPtrArg), GNameArg);
-		BIND_CALL_CHECK_DELEGATE       (         FTestDynamicDelegate3,                     0, &UDynamicDelegateTestObject::ConstFunc4,       (GIntArg, GStrArg, GPtrArg), GNameArg);
-		BIND_CALL_CHECK_DELEGATE       (         FTestDynamicMulticastDelegate3,            1, &UDynamicDelegateTestObject::ConstFunc4,       (GIntArg, GStrArg, GPtrArg), GNameArg);
+		BIND_CALL_CHECK_DELEGATE(FTestDynamicMulticastDelegate3::FDelegate, 0, &UDynamicDelegateTestObject::Func4, (GIntArg, GStrArg, GPtrArg), GNameArg);
+		BIND_CALL_CHECK_DELEGATE(FTestDynamicDelegate3, 0, &UDynamicDelegateTestObject::Func4, (GIntArg, GStrArg, GPtrArg), GNameArg);
+		BIND_CALL_CHECK_DELEGATE(FTestDynamicMulticastDelegate3, 1, &UDynamicDelegateTestObject::Func4, (GIntArg, GStrArg, GPtrArg), GNameArg);
+		BIND_CALL_CHECK_DELEGATE(FTestDynamicMulticastDelegate3::FDelegate, 0, &UDynamicDelegateTestObject::ConstFunc4, (GIntArg, GStrArg, GPtrArg), GNameArg);
+		BIND_CALL_CHECK_DELEGATE(FTestDynamicDelegate3, 0, &UDynamicDelegateTestObject::ConstFunc4, (GIntArg, GStrArg, GPtrArg), GNameArg);
+		BIND_CALL_CHECK_DELEGATE(FTestDynamicMulticastDelegate3, 1, &UDynamicDelegateTestObject::ConstFunc4, (GIntArg, GStrArg, GPtrArg), GNameArg);
 #if UE_USE_DYNAMIC_DELEGATE_PAYLOADS_RETVAL
-		BIND_CALL_CHECK_DELEGATE_RETVAL(GIntRet, FTestDynamicDelegateRetInt3,               0, &UDynamicDelegateTestObject::FuncRetInt4,      (GIntArg, GStrArg, GPtrArg), GNameArg);
-		BIND_CALL_CHECK_DELEGATE_RETVAL(GIntRet, FTestDynamicDelegateRetInt3,               0, &UDynamicDelegateTestObject::ConstFuncRetInt4, (GIntArg, GStrArg, GPtrArg), GNameArg);
-		BIND_CALL_CHECK_DELEGATE_RETVAL(GStrRet, FTestDynamicDelegateRetStr3,               0, &UDynamicDelegateTestObject::FuncRetStr4,      (GIntArg, GStrArg, GPtrArg), GNameArg);
-		BIND_CALL_CHECK_DELEGATE_RETVAL(GStrRet, FTestDynamicDelegateRetStr3,               0, &UDynamicDelegateTestObject::ConstFuncRetStr4, (GIntArg, GStrArg, GPtrArg), GNameArg);
+		BIND_CALL_CHECK_DELEGATE_RETVAL(GIntRet, FTestDynamicDelegateRetInt3, 0, &UDynamicDelegateTestObject::FuncRetInt4, (GIntArg, GStrArg, GPtrArg), GNameArg);
+		BIND_CALL_CHECK_DELEGATE_RETVAL(GIntRet, FTestDynamicDelegateRetInt3, 0, &UDynamicDelegateTestObject::ConstFuncRetInt4, (GIntArg, GStrArg, GPtrArg), GNameArg);
+		BIND_CALL_CHECK_DELEGATE_RETVAL(GStrRet, FTestDynamicDelegateRetStr3, 0, &UDynamicDelegateTestObject::FuncRetStr4, (GIntArg, GStrArg, GPtrArg), GNameArg);
+		BIND_CALL_CHECK_DELEGATE_RETVAL(GStrRet, FTestDynamicDelegateRetStr3, 0, &UDynamicDelegateTestObject::ConstFuncRetStr4, (GIntArg, GStrArg, GPtrArg), GNameArg);
 #endif
 
 		// 2 arg payload
-		BIND_CALL_CHECK_DELEGATE       (         FTestDynamicMulticastDelegate2::FDelegate, 0, &UDynamicDelegateTestObject::Func4,            (GIntArg, GStrArg), GPtrArg, GNameArg);
-		BIND_CALL_CHECK_DELEGATE       (         FTestDynamicDelegate2,                     0, &UDynamicDelegateTestObject::Func4,            (GIntArg, GStrArg), GPtrArg, GNameArg);
-		BIND_CALL_CHECK_DELEGATE       (         FTestDynamicMulticastDelegate2,            1, &UDynamicDelegateTestObject::Func4,            (GIntArg, GStrArg), GPtrArg, GNameArg);
-		BIND_CALL_CHECK_DELEGATE       (         FTestDynamicMulticastDelegate2::FDelegate, 0, &UDynamicDelegateTestObject::ConstFunc4,       (GIntArg, GStrArg), GPtrArg, GNameArg);
-		BIND_CALL_CHECK_DELEGATE       (         FTestDynamicDelegate2,                     0, &UDynamicDelegateTestObject::ConstFunc4,       (GIntArg, GStrArg), GPtrArg, GNameArg);
-		BIND_CALL_CHECK_DELEGATE       (         FTestDynamicMulticastDelegate2,            1, &UDynamicDelegateTestObject::ConstFunc4,       (GIntArg, GStrArg), GPtrArg, GNameArg);
+		BIND_CALL_CHECK_DELEGATE(FTestDynamicMulticastDelegate2::FDelegate, 0, &UDynamicDelegateTestObject::Func4, (GIntArg, GStrArg), GPtrArg, GNameArg);
+		BIND_CALL_CHECK_DELEGATE(FTestDynamicDelegate2, 0, &UDynamicDelegateTestObject::Func4, (GIntArg, GStrArg), GPtrArg, GNameArg);
+		BIND_CALL_CHECK_DELEGATE(FTestDynamicMulticastDelegate2, 1, &UDynamicDelegateTestObject::Func4, (GIntArg, GStrArg), GPtrArg, GNameArg);
+		BIND_CALL_CHECK_DELEGATE(FTestDynamicMulticastDelegate2::FDelegate, 0, &UDynamicDelegateTestObject::ConstFunc4, (GIntArg, GStrArg), GPtrArg, GNameArg);
+		BIND_CALL_CHECK_DELEGATE(FTestDynamicDelegate2, 0, &UDynamicDelegateTestObject::ConstFunc4, (GIntArg, GStrArg), GPtrArg, GNameArg);
+		BIND_CALL_CHECK_DELEGATE(FTestDynamicMulticastDelegate2, 1, &UDynamicDelegateTestObject::ConstFunc4, (GIntArg, GStrArg), GPtrArg, GNameArg);
 #if UE_USE_DYNAMIC_DELEGATE_PAYLOADS_RETVAL
-		BIND_CALL_CHECK_DELEGATE_RETVAL(GIntRet, FTestDynamicDelegateRetInt2,               0, &UDynamicDelegateTestObject::FuncRetInt4,      (GIntArg, GStrArg), GPtrArg, GNameArg);
-		BIND_CALL_CHECK_DELEGATE_RETVAL(GIntRet, FTestDynamicDelegateRetInt2,               0, &UDynamicDelegateTestObject::ConstFuncRetInt4, (GIntArg, GStrArg), GPtrArg, GNameArg);
-		BIND_CALL_CHECK_DELEGATE_RETVAL(GStrRet, FTestDynamicDelegateRetStr2,               0, &UDynamicDelegateTestObject::FuncRetStr4,      (GIntArg, GStrArg), GPtrArg, GNameArg);
-		BIND_CALL_CHECK_DELEGATE_RETVAL(GStrRet, FTestDynamicDelegateRetStr2,               0, &UDynamicDelegateTestObject::ConstFuncRetStr4, (GIntArg, GStrArg), GPtrArg, GNameArg);
+		BIND_CALL_CHECK_DELEGATE_RETVAL(GIntRet, FTestDynamicDelegateRetInt2, 0, &UDynamicDelegateTestObject::FuncRetInt4, (GIntArg, GStrArg), GPtrArg, GNameArg);
+		BIND_CALL_CHECK_DELEGATE_RETVAL(GIntRet, FTestDynamicDelegateRetInt2, 0, &UDynamicDelegateTestObject::ConstFuncRetInt4, (GIntArg, GStrArg), GPtrArg, GNameArg);
+		BIND_CALL_CHECK_DELEGATE_RETVAL(GStrRet, FTestDynamicDelegateRetStr2, 0, &UDynamicDelegateTestObject::FuncRetStr4, (GIntArg, GStrArg), GPtrArg, GNameArg);
+		BIND_CALL_CHECK_DELEGATE_RETVAL(GStrRet, FTestDynamicDelegateRetStr2, 0, &UDynamicDelegateTestObject::ConstFuncRetStr4, (GIntArg, GStrArg), GPtrArg, GNameArg);
 #endif
 
 		// 3 arg payload
-		BIND_CALL_CHECK_DELEGATE       (         FTestDynamicMulticastDelegate1::FDelegate, 0, &UDynamicDelegateTestObject::Func4,            (GIntArg), GStrArg, GPtrArg, GNameArg);
-		BIND_CALL_CHECK_DELEGATE       (         FTestDynamicDelegate1,                     0, &UDynamicDelegateTestObject::Func4,            (GIntArg), GStrArg, GPtrArg, GNameArg);
-		BIND_CALL_CHECK_DELEGATE       (         FTestDynamicMulticastDelegate1,            1, &UDynamicDelegateTestObject::Func4,            (GIntArg), GStrArg, GPtrArg, GNameArg);
-		BIND_CALL_CHECK_DELEGATE       (         FTestDynamicMulticastDelegate1::FDelegate, 0, &UDynamicDelegateTestObject::ConstFunc4,       (GIntArg), GStrArg, GPtrArg, GNameArg);
-		BIND_CALL_CHECK_DELEGATE       (         FTestDynamicDelegate1,                     0, &UDynamicDelegateTestObject::ConstFunc4,       (GIntArg), GStrArg, GPtrArg, GNameArg);
-		BIND_CALL_CHECK_DELEGATE       (         FTestDynamicMulticastDelegate1,            1, &UDynamicDelegateTestObject::ConstFunc4,       (GIntArg), GStrArg, GPtrArg, GNameArg);
+		BIND_CALL_CHECK_DELEGATE(FTestDynamicMulticastDelegate1::FDelegate, 0, &UDynamicDelegateTestObject::Func4, (GIntArg), GStrArg, GPtrArg, GNameArg);
+		BIND_CALL_CHECK_DELEGATE(FTestDynamicDelegate1, 0, &UDynamicDelegateTestObject::Func4, (GIntArg), GStrArg, GPtrArg, GNameArg);
+		BIND_CALL_CHECK_DELEGATE(FTestDynamicMulticastDelegate1, 1, &UDynamicDelegateTestObject::Func4, (GIntArg), GStrArg, GPtrArg, GNameArg);
+		BIND_CALL_CHECK_DELEGATE(FTestDynamicMulticastDelegate1::FDelegate, 0, &UDynamicDelegateTestObject::ConstFunc4, (GIntArg), GStrArg, GPtrArg, GNameArg);
+		BIND_CALL_CHECK_DELEGATE(FTestDynamicDelegate1, 0, &UDynamicDelegateTestObject::ConstFunc4, (GIntArg), GStrArg, GPtrArg, GNameArg);
+		BIND_CALL_CHECK_DELEGATE(FTestDynamicMulticastDelegate1, 1, &UDynamicDelegateTestObject::ConstFunc4, (GIntArg), GStrArg, GPtrArg, GNameArg);
 #if UE_USE_DYNAMIC_DELEGATE_PAYLOADS_RETVAL
-		BIND_CALL_CHECK_DELEGATE_RETVAL(GIntRet, FTestDynamicDelegateRetInt1,               0, &UDynamicDelegateTestObject::FuncRetInt4,      (GIntArg), GStrArg, GPtrArg, GNameArg);
-		BIND_CALL_CHECK_DELEGATE_RETVAL(GIntRet, FTestDynamicDelegateRetInt1,               0, &UDynamicDelegateTestObject::ConstFuncRetInt4, (GIntArg), GStrArg, GPtrArg, GNameArg);
-		BIND_CALL_CHECK_DELEGATE_RETVAL(GStrRet, FTestDynamicDelegateRetStr1,               0, &UDynamicDelegateTestObject::FuncRetStr4,      (GIntArg), GStrArg, GPtrArg, GNameArg);
-		BIND_CALL_CHECK_DELEGATE_RETVAL(GStrRet, FTestDynamicDelegateRetStr1,               0, &UDynamicDelegateTestObject::ConstFuncRetStr4, (GIntArg), GStrArg, GPtrArg, GNameArg);
+		BIND_CALL_CHECK_DELEGATE_RETVAL(GIntRet, FTestDynamicDelegateRetInt1, 0, &UDynamicDelegateTestObject::FuncRetInt4, (GIntArg), GStrArg, GPtrArg, GNameArg);
+		BIND_CALL_CHECK_DELEGATE_RETVAL(GIntRet, FTestDynamicDelegateRetInt1, 0, &UDynamicDelegateTestObject::ConstFuncRetInt4, (GIntArg), GStrArg, GPtrArg, GNameArg);
+		BIND_CALL_CHECK_DELEGATE_RETVAL(GStrRet, FTestDynamicDelegateRetStr1, 0, &UDynamicDelegateTestObject::FuncRetStr4, (GIntArg), GStrArg, GPtrArg, GNameArg);
+		BIND_CALL_CHECK_DELEGATE_RETVAL(GStrRet, FTestDynamicDelegateRetStr1, 0, &UDynamicDelegateTestObject::ConstFuncRetStr4, (GIntArg), GStrArg, GPtrArg, GNameArg);
 #endif
 
 		// 4 arg payload
-		BIND_CALL_CHECK_DELEGATE       (         FTestDynamicMulticastDelegate0::FDelegate, 0, &UDynamicDelegateTestObject::Func4,            (), GIntArg, GStrArg, GPtrArg, GNameArg);
-		BIND_CALL_CHECK_DELEGATE       (         FTestDynamicDelegate0,                     0, &UDynamicDelegateTestObject::Func4,            (), GIntArg, GStrArg, GPtrArg, GNameArg);
-		BIND_CALL_CHECK_DELEGATE       (         FTestDynamicMulticastDelegate0,            1, &UDynamicDelegateTestObject::Func4,            (), GIntArg, GStrArg, GPtrArg, GNameArg);
-		BIND_CALL_CHECK_DELEGATE       (         FTestDynamicMulticastDelegate0::FDelegate, 0, &UDynamicDelegateTestObject::ConstFunc4,       (), GIntArg, GStrArg, GPtrArg, GNameArg);
-		BIND_CALL_CHECK_DELEGATE       (         FTestDynamicDelegate0,                     0, &UDynamicDelegateTestObject::ConstFunc4,       (), GIntArg, GStrArg, GPtrArg, GNameArg);
-		BIND_CALL_CHECK_DELEGATE       (         FTestDynamicMulticastDelegate0,            1, &UDynamicDelegateTestObject::ConstFunc4,       (), GIntArg, GStrArg, GPtrArg, GNameArg);
+		BIND_CALL_CHECK_DELEGATE(FTestDynamicMulticastDelegate0::FDelegate, 0, &UDynamicDelegateTestObject::Func4, (), GIntArg, GStrArg, GPtrArg, GNameArg);
+		BIND_CALL_CHECK_DELEGATE(FTestDynamicDelegate0, 0, &UDynamicDelegateTestObject::Func4, (), GIntArg, GStrArg, GPtrArg, GNameArg);
+		BIND_CALL_CHECK_DELEGATE(FTestDynamicMulticastDelegate0, 1, &UDynamicDelegateTestObject::Func4, (), GIntArg, GStrArg, GPtrArg, GNameArg);
+		BIND_CALL_CHECK_DELEGATE(FTestDynamicMulticastDelegate0::FDelegate, 0, &UDynamicDelegateTestObject::ConstFunc4, (), GIntArg, GStrArg, GPtrArg, GNameArg);
+		BIND_CALL_CHECK_DELEGATE(FTestDynamicDelegate0, 0, &UDynamicDelegateTestObject::ConstFunc4, (), GIntArg, GStrArg, GPtrArg, GNameArg);
+		BIND_CALL_CHECK_DELEGATE(FTestDynamicMulticastDelegate0, 1, &UDynamicDelegateTestObject::ConstFunc4, (), GIntArg, GStrArg, GPtrArg, GNameArg);
 #if UE_USE_DYNAMIC_DELEGATE_PAYLOADS_RETVAL
-		BIND_CALL_CHECK_DELEGATE_RETVAL(GIntRet, FTestDynamicDelegateRetInt0,               0, &UDynamicDelegateTestObject::FuncRetInt4,      (), GIntArg, GStrArg, GPtrArg, GNameArg);
-		BIND_CALL_CHECK_DELEGATE_RETVAL(GIntRet, FTestDynamicDelegateRetInt0,               0, &UDynamicDelegateTestObject::ConstFuncRetInt4, (), GIntArg, GStrArg, GPtrArg, GNameArg);
-		BIND_CALL_CHECK_DELEGATE_RETVAL(GStrRet, FTestDynamicDelegateRetStr0,               0, &UDynamicDelegateTestObject::FuncRetStr4,      (), GIntArg, GStrArg, GPtrArg, GNameArg);
-		BIND_CALL_CHECK_DELEGATE_RETVAL(GStrRet, FTestDynamicDelegateRetStr0,               0, &UDynamicDelegateTestObject::ConstFuncRetStr4, (), GIntArg, GStrArg, GPtrArg, GNameArg);
+		BIND_CALL_CHECK_DELEGATE_RETVAL(GIntRet, FTestDynamicDelegateRetInt0, 0, &UDynamicDelegateTestObject::FuncRetInt4, (), GIntArg, GStrArg, GPtrArg, GNameArg);
+		BIND_CALL_CHECK_DELEGATE_RETVAL(GIntRet, FTestDynamicDelegateRetInt0, 0, &UDynamicDelegateTestObject::ConstFuncRetInt4, (), GIntArg, GStrArg, GPtrArg, GNameArg);
+		BIND_CALL_CHECK_DELEGATE_RETVAL(GStrRet, FTestDynamicDelegateRetStr0, 0, &UDynamicDelegateTestObject::FuncRetStr4, (), GIntArg, GStrArg, GPtrArg, GNameArg);
+		BIND_CALL_CHECK_DELEGATE_RETVAL(GStrRet, FTestDynamicDelegateRetStr0, 0, &UDynamicDelegateTestObject::ConstFuncRetStr4, (), GIntArg, GStrArg, GPtrArg, GNameArg);
 #endif
 #endif
 	}
@@ -609,86 +743,86 @@ void UDynamicDelegateTestObject::RunTests()
 	// 5 param function
 	{
 		// 0 arg payload
-		BIND_CALL_CHECK_DELEGATE       (         FTestDynamicMulticastDelegate5::FDelegate, 0, &UDynamicDelegateTestObject::Func5,            (GIntArg, GStrArg, GPtrArg, GNameArg, bGBoolArg));
-		BIND_CALL_CHECK_DELEGATE       (         FTestDynamicDelegate5,                     0, &UDynamicDelegateTestObject::Func5,            (GIntArg, GStrArg, GPtrArg, GNameArg, bGBoolArg));
-		BIND_CALL_CHECK_DELEGATE       (         FTestDynamicMulticastDelegate5,            1, &UDynamicDelegateTestObject::Func5,            (GIntArg, GStrArg, GPtrArg, GNameArg, bGBoolArg));
-		BIND_CALL_CHECK_DELEGATE       (         FTestDynamicMulticastDelegate5::FDelegate, 0, &UDynamicDelegateTestObject::ConstFunc5,       (GIntArg, GStrArg, GPtrArg, GNameArg, bGBoolArg));
-		BIND_CALL_CHECK_DELEGATE       (         FTestDynamicDelegate5,                     0, &UDynamicDelegateTestObject::ConstFunc5,       (GIntArg, GStrArg, GPtrArg, GNameArg, bGBoolArg));
-		BIND_CALL_CHECK_DELEGATE       (         FTestDynamicMulticastDelegate5,            1, &UDynamicDelegateTestObject::ConstFunc5,       (GIntArg, GStrArg, GPtrArg, GNameArg, bGBoolArg));
-		BIND_CALL_CHECK_DELEGATE_RETVAL(GIntRet, FTestDynamicDelegateRetInt5,               0, &UDynamicDelegateTestObject::FuncRetInt5,      (GIntArg, GStrArg, GPtrArg, GNameArg, bGBoolArg));
-		BIND_CALL_CHECK_DELEGATE_RETVAL(GIntRet, FTestDynamicDelegateRetInt5,               0, &UDynamicDelegateTestObject::ConstFuncRetInt5, (GIntArg, GStrArg, GPtrArg, GNameArg, bGBoolArg));
-		BIND_CALL_CHECK_DELEGATE_RETVAL(GStrRet, FTestDynamicDelegateRetStr5,               0, &UDynamicDelegateTestObject::FuncRetStr5,      (GIntArg, GStrArg, GPtrArg, GNameArg, bGBoolArg));
-		BIND_CALL_CHECK_DELEGATE_RETVAL(GStrRet, FTestDynamicDelegateRetStr5,               0, &UDynamicDelegateTestObject::ConstFuncRetStr5, (GIntArg, GStrArg, GPtrArg, GNameArg, bGBoolArg));
+		BIND_CALL_CHECK_DELEGATE(FTestDynamicMulticastDelegate5::FDelegate, 0, &UDynamicDelegateTestObject::Func5, (GIntArg, GStrArg, GPtrArg, GNameArg, bGBoolArg));
+		BIND_CALL_CHECK_DELEGATE(FTestDynamicDelegate5, 0, &UDynamicDelegateTestObject::Func5, (GIntArg, GStrArg, GPtrArg, GNameArg, bGBoolArg));
+		BIND_CALL_CHECK_DELEGATE(FTestDynamicMulticastDelegate5, 1, &UDynamicDelegateTestObject::Func5, (GIntArg, GStrArg, GPtrArg, GNameArg, bGBoolArg));
+		BIND_CALL_CHECK_DELEGATE(FTestDynamicMulticastDelegate5::FDelegate, 0, &UDynamicDelegateTestObject::ConstFunc5, (GIntArg, GStrArg, GPtrArg, GNameArg, bGBoolArg));
+		BIND_CALL_CHECK_DELEGATE(FTestDynamicDelegate5, 0, &UDynamicDelegateTestObject::ConstFunc5, (GIntArg, GStrArg, GPtrArg, GNameArg, bGBoolArg));
+		BIND_CALL_CHECK_DELEGATE(FTestDynamicMulticastDelegate5, 1, &UDynamicDelegateTestObject::ConstFunc5, (GIntArg, GStrArg, GPtrArg, GNameArg, bGBoolArg));
+		BIND_CALL_CHECK_DELEGATE_RETVAL(GIntRet, FTestDynamicDelegateRetInt5, 0, &UDynamicDelegateTestObject::FuncRetInt5, (GIntArg, GStrArg, GPtrArg, GNameArg, bGBoolArg));
+		BIND_CALL_CHECK_DELEGATE_RETVAL(GIntRet, FTestDynamicDelegateRetInt5, 0, &UDynamicDelegateTestObject::ConstFuncRetInt5, (GIntArg, GStrArg, GPtrArg, GNameArg, bGBoolArg));
+		BIND_CALL_CHECK_DELEGATE_RETVAL(GStrRet, FTestDynamicDelegateRetStr5, 0, &UDynamicDelegateTestObject::FuncRetStr5, (GIntArg, GStrArg, GPtrArg, GNameArg, bGBoolArg));
+		BIND_CALL_CHECK_DELEGATE_RETVAL(GStrRet, FTestDynamicDelegateRetStr5, 0, &UDynamicDelegateTestObject::ConstFuncRetStr5, (GIntArg, GStrArg, GPtrArg, GNameArg, bGBoolArg));
 
 #if defined(UE_USE_DYNAMIC_DELEGATE_PAYLOADS) && UE_USE_DYNAMIC_DELEGATE_PAYLOADS
 		// 1 arg payload
-		BIND_CALL_CHECK_DELEGATE       (         FTestDynamicMulticastDelegate4::FDelegate, 0, &UDynamicDelegateTestObject::Func5,            (GIntArg, GStrArg, GPtrArg, GNameArg), bGBoolArg);
-		BIND_CALL_CHECK_DELEGATE       (         FTestDynamicDelegate4,                     0, &UDynamicDelegateTestObject::Func5,            (GIntArg, GStrArg, GPtrArg, GNameArg), bGBoolArg);
-		BIND_CALL_CHECK_DELEGATE       (         FTestDynamicMulticastDelegate4,            1, &UDynamicDelegateTestObject::Func5,            (GIntArg, GStrArg, GPtrArg, GNameArg), bGBoolArg);
-		BIND_CALL_CHECK_DELEGATE       (         FTestDynamicMulticastDelegate4::FDelegate, 0, &UDynamicDelegateTestObject::ConstFunc5,       (GIntArg, GStrArg, GPtrArg, GNameArg), bGBoolArg);
-		BIND_CALL_CHECK_DELEGATE       (         FTestDynamicDelegate4,                     0, &UDynamicDelegateTestObject::ConstFunc5,       (GIntArg, GStrArg, GPtrArg, GNameArg), bGBoolArg);
-		BIND_CALL_CHECK_DELEGATE       (         FTestDynamicMulticastDelegate4,            1, &UDynamicDelegateTestObject::ConstFunc5,       (GIntArg, GStrArg, GPtrArg, GNameArg), bGBoolArg);
+		BIND_CALL_CHECK_DELEGATE(FTestDynamicMulticastDelegate4::FDelegate, 0, &UDynamicDelegateTestObject::Func5, (GIntArg, GStrArg, GPtrArg, GNameArg), bGBoolArg);
+		BIND_CALL_CHECK_DELEGATE(FTestDynamicDelegate4, 0, &UDynamicDelegateTestObject::Func5, (GIntArg, GStrArg, GPtrArg, GNameArg), bGBoolArg);
+		BIND_CALL_CHECK_DELEGATE(FTestDynamicMulticastDelegate4, 1, &UDynamicDelegateTestObject::Func5, (GIntArg, GStrArg, GPtrArg, GNameArg), bGBoolArg);
+		BIND_CALL_CHECK_DELEGATE(FTestDynamicMulticastDelegate4::FDelegate, 0, &UDynamicDelegateTestObject::ConstFunc5, (GIntArg, GStrArg, GPtrArg, GNameArg), bGBoolArg);
+		BIND_CALL_CHECK_DELEGATE(FTestDynamicDelegate4, 0, &UDynamicDelegateTestObject::ConstFunc5, (GIntArg, GStrArg, GPtrArg, GNameArg), bGBoolArg);
+		BIND_CALL_CHECK_DELEGATE(FTestDynamicMulticastDelegate4, 1, &UDynamicDelegateTestObject::ConstFunc5, (GIntArg, GStrArg, GPtrArg, GNameArg), bGBoolArg);
 #if UE_USE_DYNAMIC_DELEGATE_PAYLOADS_RETVAL
-		BIND_CALL_CHECK_DELEGATE_RETVAL(GIntRet, FTestDynamicDelegateRetInt4,               0, &UDynamicDelegateTestObject::FuncRetInt5,      (GIntArg, GStrArg, GPtrArg, GNameArg), bGBoolArg);
-		BIND_CALL_CHECK_DELEGATE_RETVAL(GIntRet, FTestDynamicDelegateRetInt4,               0, &UDynamicDelegateTestObject::ConstFuncRetInt5, (GIntArg, GStrArg, GPtrArg, GNameArg), bGBoolArg);
-		BIND_CALL_CHECK_DELEGATE_RETVAL(GStrRet, FTestDynamicDelegateRetStr4,               0, &UDynamicDelegateTestObject::FuncRetStr5,      (GIntArg, GStrArg, GPtrArg, GNameArg), bGBoolArg);
-		BIND_CALL_CHECK_DELEGATE_RETVAL(GStrRet, FTestDynamicDelegateRetStr4,               0, &UDynamicDelegateTestObject::ConstFuncRetStr5, (GIntArg, GStrArg, GPtrArg, GNameArg), bGBoolArg);
+		BIND_CALL_CHECK_DELEGATE_RETVAL(GIntRet, FTestDynamicDelegateRetInt4, 0, &UDynamicDelegateTestObject::FuncRetInt5, (GIntArg, GStrArg, GPtrArg, GNameArg), bGBoolArg);
+		BIND_CALL_CHECK_DELEGATE_RETVAL(GIntRet, FTestDynamicDelegateRetInt4, 0, &UDynamicDelegateTestObject::ConstFuncRetInt5, (GIntArg, GStrArg, GPtrArg, GNameArg), bGBoolArg);
+		BIND_CALL_CHECK_DELEGATE_RETVAL(GStrRet, FTestDynamicDelegateRetStr4, 0, &UDynamicDelegateTestObject::FuncRetStr5, (GIntArg, GStrArg, GPtrArg, GNameArg), bGBoolArg);
+		BIND_CALL_CHECK_DELEGATE_RETVAL(GStrRet, FTestDynamicDelegateRetStr4, 0, &UDynamicDelegateTestObject::ConstFuncRetStr5, (GIntArg, GStrArg, GPtrArg, GNameArg), bGBoolArg);
 #endif
 
 		// 2 arg payload
-		BIND_CALL_CHECK_DELEGATE       (         FTestDynamicMulticastDelegate3::FDelegate, 0, &UDynamicDelegateTestObject::Func5,            (GIntArg, GStrArg, GPtrArg), GNameArg, bGBoolArg);
-		BIND_CALL_CHECK_DELEGATE       (         FTestDynamicDelegate3,                     0, &UDynamicDelegateTestObject::Func5,            (GIntArg, GStrArg, GPtrArg), GNameArg, bGBoolArg);
-		BIND_CALL_CHECK_DELEGATE       (         FTestDynamicMulticastDelegate3,            1, &UDynamicDelegateTestObject::Func5,            (GIntArg, GStrArg, GPtrArg), GNameArg, bGBoolArg);
-		BIND_CALL_CHECK_DELEGATE       (         FTestDynamicMulticastDelegate3::FDelegate, 0, &UDynamicDelegateTestObject::ConstFunc5,       (GIntArg, GStrArg, GPtrArg), GNameArg, bGBoolArg);
-		BIND_CALL_CHECK_DELEGATE       (         FTestDynamicDelegate3,                     0, &UDynamicDelegateTestObject::ConstFunc5,       (GIntArg, GStrArg, GPtrArg), GNameArg, bGBoolArg);
-		BIND_CALL_CHECK_DELEGATE       (         FTestDynamicMulticastDelegate3,            1, &UDynamicDelegateTestObject::ConstFunc5,       (GIntArg, GStrArg, GPtrArg), GNameArg, bGBoolArg);
+		BIND_CALL_CHECK_DELEGATE(FTestDynamicMulticastDelegate3::FDelegate, 0, &UDynamicDelegateTestObject::Func5, (GIntArg, GStrArg, GPtrArg), GNameArg, bGBoolArg);
+		BIND_CALL_CHECK_DELEGATE(FTestDynamicDelegate3, 0, &UDynamicDelegateTestObject::Func5, (GIntArg, GStrArg, GPtrArg), GNameArg, bGBoolArg);
+		BIND_CALL_CHECK_DELEGATE(FTestDynamicMulticastDelegate3, 1, &UDynamicDelegateTestObject::Func5, (GIntArg, GStrArg, GPtrArg), GNameArg, bGBoolArg);
+		BIND_CALL_CHECK_DELEGATE(FTestDynamicMulticastDelegate3::FDelegate, 0, &UDynamicDelegateTestObject::ConstFunc5, (GIntArg, GStrArg, GPtrArg), GNameArg, bGBoolArg);
+		BIND_CALL_CHECK_DELEGATE(FTestDynamicDelegate3, 0, &UDynamicDelegateTestObject::ConstFunc5, (GIntArg, GStrArg, GPtrArg), GNameArg, bGBoolArg);
+		BIND_CALL_CHECK_DELEGATE(FTestDynamicMulticastDelegate3, 1, &UDynamicDelegateTestObject::ConstFunc5, (GIntArg, GStrArg, GPtrArg), GNameArg, bGBoolArg);
 #if UE_USE_DYNAMIC_DELEGATE_PAYLOADS_RETVAL
-		BIND_CALL_CHECK_DELEGATE_RETVAL(GIntRet, FTestDynamicDelegateRetInt3,               0, &UDynamicDelegateTestObject::FuncRetInt5,      (GIntArg, GStrArg, GPtrArg), GNameArg, bGBoolArg);
-		BIND_CALL_CHECK_DELEGATE_RETVAL(GIntRet, FTestDynamicDelegateRetInt3,               0, &UDynamicDelegateTestObject::ConstFuncRetInt5, (GIntArg, GStrArg, GPtrArg), GNameArg, bGBoolArg);
-		BIND_CALL_CHECK_DELEGATE_RETVAL(GStrRet, FTestDynamicDelegateRetStr3,               0, &UDynamicDelegateTestObject::FuncRetStr5,      (GIntArg, GStrArg, GPtrArg), GNameArg, bGBoolArg);
-		BIND_CALL_CHECK_DELEGATE_RETVAL(GStrRet, FTestDynamicDelegateRetStr3,               0, &UDynamicDelegateTestObject::ConstFuncRetStr5, (GIntArg, GStrArg, GPtrArg), GNameArg, bGBoolArg);
+		BIND_CALL_CHECK_DELEGATE_RETVAL(GIntRet, FTestDynamicDelegateRetInt3, 0, &UDynamicDelegateTestObject::FuncRetInt5, (GIntArg, GStrArg, GPtrArg), GNameArg, bGBoolArg);
+		BIND_CALL_CHECK_DELEGATE_RETVAL(GIntRet, FTestDynamicDelegateRetInt3, 0, &UDynamicDelegateTestObject::ConstFuncRetInt5, (GIntArg, GStrArg, GPtrArg), GNameArg, bGBoolArg);
+		BIND_CALL_CHECK_DELEGATE_RETVAL(GStrRet, FTestDynamicDelegateRetStr3, 0, &UDynamicDelegateTestObject::FuncRetStr5, (GIntArg, GStrArg, GPtrArg), GNameArg, bGBoolArg);
+		BIND_CALL_CHECK_DELEGATE_RETVAL(GStrRet, FTestDynamicDelegateRetStr3, 0, &UDynamicDelegateTestObject::ConstFuncRetStr5, (GIntArg, GStrArg, GPtrArg), GNameArg, bGBoolArg);
 #endif
 
 		// 3 arg payload
-		BIND_CALL_CHECK_DELEGATE       (         FTestDynamicMulticastDelegate2::FDelegate, 0, &UDynamicDelegateTestObject::Func5,            (GIntArg, GStrArg), GPtrArg, GNameArg, bGBoolArg);
-		BIND_CALL_CHECK_DELEGATE       (         FTestDynamicDelegate2,                     0, &UDynamicDelegateTestObject::Func5,            (GIntArg, GStrArg), GPtrArg, GNameArg, bGBoolArg);
-		BIND_CALL_CHECK_DELEGATE       (         FTestDynamicMulticastDelegate2,            1, &UDynamicDelegateTestObject::Func5,            (GIntArg, GStrArg), GPtrArg, GNameArg, bGBoolArg);
-		BIND_CALL_CHECK_DELEGATE       (         FTestDynamicMulticastDelegate2::FDelegate, 0, &UDynamicDelegateTestObject::ConstFunc5,       (GIntArg, GStrArg), GPtrArg, GNameArg, bGBoolArg);
-		BIND_CALL_CHECK_DELEGATE       (         FTestDynamicDelegate2,                     0, &UDynamicDelegateTestObject::ConstFunc5,       (GIntArg, GStrArg), GPtrArg, GNameArg, bGBoolArg);
-		BIND_CALL_CHECK_DELEGATE       (         FTestDynamicMulticastDelegate2,            1, &UDynamicDelegateTestObject::ConstFunc5,       (GIntArg, GStrArg), GPtrArg, GNameArg, bGBoolArg);
+		BIND_CALL_CHECK_DELEGATE(FTestDynamicMulticastDelegate2::FDelegate, 0, &UDynamicDelegateTestObject::Func5, (GIntArg, GStrArg), GPtrArg, GNameArg, bGBoolArg);
+		BIND_CALL_CHECK_DELEGATE(FTestDynamicDelegate2, 0, &UDynamicDelegateTestObject::Func5, (GIntArg, GStrArg), GPtrArg, GNameArg, bGBoolArg);
+		BIND_CALL_CHECK_DELEGATE(FTestDynamicMulticastDelegate2, 1, &UDynamicDelegateTestObject::Func5, (GIntArg, GStrArg), GPtrArg, GNameArg, bGBoolArg);
+		BIND_CALL_CHECK_DELEGATE(FTestDynamicMulticastDelegate2::FDelegate, 0, &UDynamicDelegateTestObject::ConstFunc5, (GIntArg, GStrArg), GPtrArg, GNameArg, bGBoolArg);
+		BIND_CALL_CHECK_DELEGATE(FTestDynamicDelegate2, 0, &UDynamicDelegateTestObject::ConstFunc5, (GIntArg, GStrArg), GPtrArg, GNameArg, bGBoolArg);
+		BIND_CALL_CHECK_DELEGATE(FTestDynamicMulticastDelegate2, 1, &UDynamicDelegateTestObject::ConstFunc5, (GIntArg, GStrArg), GPtrArg, GNameArg, bGBoolArg);
 #if UE_USE_DYNAMIC_DELEGATE_PAYLOADS_RETVAL
-		BIND_CALL_CHECK_DELEGATE_RETVAL(GIntRet, FTestDynamicDelegateRetInt2,               0, &UDynamicDelegateTestObject::FuncRetInt5,      (GIntArg, GStrArg), GPtrArg, GNameArg, bGBoolArg);
-		BIND_CALL_CHECK_DELEGATE_RETVAL(GIntRet, FTestDynamicDelegateRetInt2,               0, &UDynamicDelegateTestObject::ConstFuncRetInt5, (GIntArg, GStrArg), GPtrArg, GNameArg, bGBoolArg);
-		BIND_CALL_CHECK_DELEGATE_RETVAL(GStrRet, FTestDynamicDelegateRetStr2,               0, &UDynamicDelegateTestObject::FuncRetStr5,      (GIntArg, GStrArg), GPtrArg, GNameArg, bGBoolArg);
-		BIND_CALL_CHECK_DELEGATE_RETVAL(GStrRet, FTestDynamicDelegateRetStr2,               0, &UDynamicDelegateTestObject::ConstFuncRetStr5, (GIntArg, GStrArg), GPtrArg, GNameArg, bGBoolArg);
+		BIND_CALL_CHECK_DELEGATE_RETVAL(GIntRet, FTestDynamicDelegateRetInt2, 0, &UDynamicDelegateTestObject::FuncRetInt5, (GIntArg, GStrArg), GPtrArg, GNameArg, bGBoolArg);
+		BIND_CALL_CHECK_DELEGATE_RETVAL(GIntRet, FTestDynamicDelegateRetInt2, 0, &UDynamicDelegateTestObject::ConstFuncRetInt5, (GIntArg, GStrArg), GPtrArg, GNameArg, bGBoolArg);
+		BIND_CALL_CHECK_DELEGATE_RETVAL(GStrRet, FTestDynamicDelegateRetStr2, 0, &UDynamicDelegateTestObject::FuncRetStr5, (GIntArg, GStrArg), GPtrArg, GNameArg, bGBoolArg);
+		BIND_CALL_CHECK_DELEGATE_RETVAL(GStrRet, FTestDynamicDelegateRetStr2, 0, &UDynamicDelegateTestObject::ConstFuncRetStr5, (GIntArg, GStrArg), GPtrArg, GNameArg, bGBoolArg);
 #endif
 
 		// 4 arg payload
-		BIND_CALL_CHECK_DELEGATE       (         FTestDynamicMulticastDelegate1::FDelegate, 0, &UDynamicDelegateTestObject::Func5,            (GIntArg), GStrArg, GPtrArg, GNameArg, bGBoolArg);
-		BIND_CALL_CHECK_DELEGATE       (         FTestDynamicDelegate1,                     0, &UDynamicDelegateTestObject::Func5,            (GIntArg), GStrArg, GPtrArg, GNameArg, bGBoolArg);
-		BIND_CALL_CHECK_DELEGATE       (         FTestDynamicMulticastDelegate1,            1, &UDynamicDelegateTestObject::Func5,            (GIntArg), GStrArg, GPtrArg, GNameArg, bGBoolArg);
-		BIND_CALL_CHECK_DELEGATE       (         FTestDynamicMulticastDelegate1::FDelegate, 0, &UDynamicDelegateTestObject::ConstFunc5,       (GIntArg), GStrArg, GPtrArg, GNameArg, bGBoolArg);
-		BIND_CALL_CHECK_DELEGATE       (         FTestDynamicDelegate1,                     0, &UDynamicDelegateTestObject::ConstFunc5,       (GIntArg), GStrArg, GPtrArg, GNameArg, bGBoolArg);
-		BIND_CALL_CHECK_DELEGATE       (         FTestDynamicMulticastDelegate1,            1, &UDynamicDelegateTestObject::ConstFunc5,       (GIntArg), GStrArg, GPtrArg, GNameArg, bGBoolArg);
+		BIND_CALL_CHECK_DELEGATE(FTestDynamicMulticastDelegate1::FDelegate, 0, &UDynamicDelegateTestObject::Func5, (GIntArg), GStrArg, GPtrArg, GNameArg, bGBoolArg);
+		BIND_CALL_CHECK_DELEGATE(FTestDynamicDelegate1, 0, &UDynamicDelegateTestObject::Func5, (GIntArg), GStrArg, GPtrArg, GNameArg, bGBoolArg);
+		BIND_CALL_CHECK_DELEGATE(FTestDynamicMulticastDelegate1, 1, &UDynamicDelegateTestObject::Func5, (GIntArg), GStrArg, GPtrArg, GNameArg, bGBoolArg);
+		BIND_CALL_CHECK_DELEGATE(FTestDynamicMulticastDelegate1::FDelegate, 0, &UDynamicDelegateTestObject::ConstFunc5, (GIntArg), GStrArg, GPtrArg, GNameArg, bGBoolArg);
+		BIND_CALL_CHECK_DELEGATE(FTestDynamicDelegate1, 0, &UDynamicDelegateTestObject::ConstFunc5, (GIntArg), GStrArg, GPtrArg, GNameArg, bGBoolArg);
+		BIND_CALL_CHECK_DELEGATE(FTestDynamicMulticastDelegate1, 1, &UDynamicDelegateTestObject::ConstFunc5, (GIntArg), GStrArg, GPtrArg, GNameArg, bGBoolArg);
 #if UE_USE_DYNAMIC_DELEGATE_PAYLOADS_RETVAL
-		BIND_CALL_CHECK_DELEGATE_RETVAL(GIntRet, FTestDynamicDelegateRetInt1,               0, &UDynamicDelegateTestObject::FuncRetInt5,      (GIntArg), GStrArg, GPtrArg, GNameArg, bGBoolArg);
-		BIND_CALL_CHECK_DELEGATE_RETVAL(GIntRet, FTestDynamicDelegateRetInt1,               0, &UDynamicDelegateTestObject::ConstFuncRetInt5, (GIntArg), GStrArg, GPtrArg, GNameArg, bGBoolArg);
-		BIND_CALL_CHECK_DELEGATE_RETVAL(GStrRet, FTestDynamicDelegateRetStr1,               0, &UDynamicDelegateTestObject::FuncRetStr5,      (GIntArg), GStrArg, GPtrArg, GNameArg, bGBoolArg);
-		BIND_CALL_CHECK_DELEGATE_RETVAL(GStrRet, FTestDynamicDelegateRetStr1,               0, &UDynamicDelegateTestObject::ConstFuncRetStr5, (GIntArg), GStrArg, GPtrArg, GNameArg, bGBoolArg);
+		BIND_CALL_CHECK_DELEGATE_RETVAL(GIntRet, FTestDynamicDelegateRetInt1, 0, &UDynamicDelegateTestObject::FuncRetInt5, (GIntArg), GStrArg, GPtrArg, GNameArg, bGBoolArg);
+		BIND_CALL_CHECK_DELEGATE_RETVAL(GIntRet, FTestDynamicDelegateRetInt1, 0, &UDynamicDelegateTestObject::ConstFuncRetInt5, (GIntArg), GStrArg, GPtrArg, GNameArg, bGBoolArg);
+		BIND_CALL_CHECK_DELEGATE_RETVAL(GStrRet, FTestDynamicDelegateRetStr1, 0, &UDynamicDelegateTestObject::FuncRetStr5, (GIntArg), GStrArg, GPtrArg, GNameArg, bGBoolArg);
+		BIND_CALL_CHECK_DELEGATE_RETVAL(GStrRet, FTestDynamicDelegateRetStr1, 0, &UDynamicDelegateTestObject::ConstFuncRetStr5, (GIntArg), GStrArg, GPtrArg, GNameArg, bGBoolArg);
 #endif
 
 		// 5 arg payload
-		BIND_CALL_CHECK_DELEGATE       (         FTestDynamicMulticastDelegate0::FDelegate, 0, &UDynamicDelegateTestObject::Func5,            (), GIntArg, GStrArg, GPtrArg, GNameArg, bGBoolArg);
-		BIND_CALL_CHECK_DELEGATE       (         FTestDynamicDelegate0,                     0, &UDynamicDelegateTestObject::Func5,            (), GIntArg, GStrArg, GPtrArg, GNameArg, bGBoolArg);
-		BIND_CALL_CHECK_DELEGATE       (         FTestDynamicMulticastDelegate0,            1, &UDynamicDelegateTestObject::Func5,            (), GIntArg, GStrArg, GPtrArg, GNameArg, bGBoolArg);
-		BIND_CALL_CHECK_DELEGATE       (         FTestDynamicMulticastDelegate0::FDelegate, 0, &UDynamicDelegateTestObject::ConstFunc5,       (), GIntArg, GStrArg, GPtrArg, GNameArg, bGBoolArg);
-		BIND_CALL_CHECK_DELEGATE       (         FTestDynamicDelegate0,                     0, &UDynamicDelegateTestObject::ConstFunc5,       (), GIntArg, GStrArg, GPtrArg, GNameArg, bGBoolArg);
-		BIND_CALL_CHECK_DELEGATE       (         FTestDynamicMulticastDelegate0,            1, &UDynamicDelegateTestObject::ConstFunc5,       (), GIntArg, GStrArg, GPtrArg, GNameArg, bGBoolArg);
+		BIND_CALL_CHECK_DELEGATE(FTestDynamicMulticastDelegate0::FDelegate, 0, &UDynamicDelegateTestObject::Func5, (), GIntArg, GStrArg, GPtrArg, GNameArg, bGBoolArg);
+		BIND_CALL_CHECK_DELEGATE(FTestDynamicDelegate0, 0, &UDynamicDelegateTestObject::Func5, (), GIntArg, GStrArg, GPtrArg, GNameArg, bGBoolArg);
+		BIND_CALL_CHECK_DELEGATE(FTestDynamicMulticastDelegate0, 1, &UDynamicDelegateTestObject::Func5, (), GIntArg, GStrArg, GPtrArg, GNameArg, bGBoolArg);
+		BIND_CALL_CHECK_DELEGATE(FTestDynamicMulticastDelegate0::FDelegate, 0, &UDynamicDelegateTestObject::ConstFunc5, (), GIntArg, GStrArg, GPtrArg, GNameArg, bGBoolArg);
+		BIND_CALL_CHECK_DELEGATE(FTestDynamicDelegate0, 0, &UDynamicDelegateTestObject::ConstFunc5, (), GIntArg, GStrArg, GPtrArg, GNameArg, bGBoolArg);
+		BIND_CALL_CHECK_DELEGATE(FTestDynamicMulticastDelegate0, 1, &UDynamicDelegateTestObject::ConstFunc5, (), GIntArg, GStrArg, GPtrArg, GNameArg, bGBoolArg);
 #if UE_USE_DYNAMIC_DELEGATE_PAYLOADS_RETVAL
-		BIND_CALL_CHECK_DELEGATE_RETVAL(GIntRet, FTestDynamicDelegateRetInt0,               0, &UDynamicDelegateTestObject::FuncRetInt5,      (), GIntArg, GStrArg, GPtrArg, GNameArg, bGBoolArg);
-		BIND_CALL_CHECK_DELEGATE_RETVAL(GIntRet, FTestDynamicDelegateRetInt0,               0, &UDynamicDelegateTestObject::ConstFuncRetInt5, (), GIntArg, GStrArg, GPtrArg, GNameArg, bGBoolArg);
-		BIND_CALL_CHECK_DELEGATE_RETVAL(GStrRet, FTestDynamicDelegateRetStr0,               0, &UDynamicDelegateTestObject::FuncRetStr5,      (), GIntArg, GStrArg, GPtrArg, GNameArg, bGBoolArg);
-		BIND_CALL_CHECK_DELEGATE_RETVAL(GStrRet, FTestDynamicDelegateRetStr0,               0, &UDynamicDelegateTestObject::ConstFuncRetStr5, (), GIntArg, GStrArg, GPtrArg, GNameArg, bGBoolArg);
+		BIND_CALL_CHECK_DELEGATE_RETVAL(GIntRet, FTestDynamicDelegateRetInt0, 0, &UDynamicDelegateTestObject::FuncRetInt5, (), GIntArg, GStrArg, GPtrArg, GNameArg, bGBoolArg);
+		BIND_CALL_CHECK_DELEGATE_RETVAL(GIntRet, FTestDynamicDelegateRetInt0, 0, &UDynamicDelegateTestObject::ConstFuncRetInt5, (), GIntArg, GStrArg, GPtrArg, GNameArg, bGBoolArg);
+		BIND_CALL_CHECK_DELEGATE_RETVAL(GStrRet, FTestDynamicDelegateRetStr0, 0, &UDynamicDelegateTestObject::FuncRetStr5, (), GIntArg, GStrArg, GPtrArg, GNameArg, bGBoolArg);
+		BIND_CALL_CHECK_DELEGATE_RETVAL(GStrRet, FTestDynamicDelegateRetStr0, 0, &UDynamicDelegateTestObject::ConstFuncRetStr5, (), GIntArg, GStrArg, GPtrArg, GNameArg, bGBoolArg);
 #endif
 #endif
 	}
@@ -696,103 +830,153 @@ void UDynamicDelegateTestObject::RunTests()
 	// 6 param function
 	{
 		// 0 arg payload
-		BIND_CALL_CHECK_DELEGATE       (         FTestDynamicMulticastDelegate6::FDelegate, 0, &UDynamicDelegateTestObject::Func6,            (GIntArg, GStrArg, GPtrArg, GNameArg, bGBoolArg, GWeakObjPtrArg));
-		BIND_CALL_CHECK_DELEGATE       (         FTestDynamicDelegate6,                     0, &UDynamicDelegateTestObject::Func6,            (GIntArg, GStrArg, GPtrArg, GNameArg, bGBoolArg, GWeakObjPtrArg));
-		BIND_CALL_CHECK_DELEGATE       (         FTestDynamicMulticastDelegate6,            1, &UDynamicDelegateTestObject::Func6,            (GIntArg, GStrArg, GPtrArg, GNameArg, bGBoolArg, GWeakObjPtrArg));
-		BIND_CALL_CHECK_DELEGATE       (         FTestDynamicMulticastDelegate6::FDelegate, 0, &UDynamicDelegateTestObject::ConstFunc6,       (GIntArg, GStrArg, GPtrArg, GNameArg, bGBoolArg, GWeakObjPtrArg));
-		BIND_CALL_CHECK_DELEGATE       (         FTestDynamicDelegate6,                     0, &UDynamicDelegateTestObject::ConstFunc6,       (GIntArg, GStrArg, GPtrArg, GNameArg, bGBoolArg, GWeakObjPtrArg));
-		BIND_CALL_CHECK_DELEGATE       (         FTestDynamicMulticastDelegate6,            1, &UDynamicDelegateTestObject::ConstFunc6,       (GIntArg, GStrArg, GPtrArg, GNameArg, bGBoolArg, GWeakObjPtrArg));
-		BIND_CALL_CHECK_DELEGATE_RETVAL(GIntRet, FTestDynamicDelegateRetInt6,               0, &UDynamicDelegateTestObject::FuncRetInt6,      (GIntArg, GStrArg, GPtrArg, GNameArg, bGBoolArg, GWeakObjPtrArg));
-		BIND_CALL_CHECK_DELEGATE_RETVAL(GIntRet, FTestDynamicDelegateRetInt6,               0, &UDynamicDelegateTestObject::ConstFuncRetInt6, (GIntArg, GStrArg, GPtrArg, GNameArg, bGBoolArg, GWeakObjPtrArg));
-		BIND_CALL_CHECK_DELEGATE_RETVAL(GStrRet, FTestDynamicDelegateRetStr6,               0, &UDynamicDelegateTestObject::FuncRetStr6,      (GIntArg, GStrArg, GPtrArg, GNameArg, bGBoolArg, GWeakObjPtrArg));
-		BIND_CALL_CHECK_DELEGATE_RETVAL(GStrRet, FTestDynamicDelegateRetStr6,               0, &UDynamicDelegateTestObject::ConstFuncRetStr6, (GIntArg, GStrArg, GPtrArg, GNameArg, bGBoolArg, GWeakObjPtrArg));
+		BIND_CALL_CHECK_DELEGATE(FTestDynamicMulticastDelegate6::FDelegate, 0, &UDynamicDelegateTestObject::Func6, (GIntArg, GStrArg, GPtrArg, GNameArg, bGBoolArg, GWeakObjPtrArg));
+		BIND_CALL_CHECK_DELEGATE(FTestDynamicDelegate6, 0, &UDynamicDelegateTestObject::Func6, (GIntArg, GStrArg, GPtrArg, GNameArg, bGBoolArg, GWeakObjPtrArg));
+		BIND_CALL_CHECK_DELEGATE(FTestDynamicMulticastDelegate6, 1, &UDynamicDelegateTestObject::Func6, (GIntArg, GStrArg, GPtrArg, GNameArg, bGBoolArg, GWeakObjPtrArg));
+		BIND_CALL_CHECK_DELEGATE(FTestDynamicMulticastDelegate6::FDelegate, 0, &UDynamicDelegateTestObject::ConstFunc6, (GIntArg, GStrArg, GPtrArg, GNameArg, bGBoolArg, GWeakObjPtrArg));
+		BIND_CALL_CHECK_DELEGATE(FTestDynamicDelegate6, 0, &UDynamicDelegateTestObject::ConstFunc6, (GIntArg, GStrArg, GPtrArg, GNameArg, bGBoolArg, GWeakObjPtrArg));
+		BIND_CALL_CHECK_DELEGATE(FTestDynamicMulticastDelegate6, 1, &UDynamicDelegateTestObject::ConstFunc6, (GIntArg, GStrArg, GPtrArg, GNameArg, bGBoolArg, GWeakObjPtrArg));
+		BIND_CALL_CHECK_DELEGATE_RETVAL(GIntRet, FTestDynamicDelegateRetInt6, 0, &UDynamicDelegateTestObject::FuncRetInt6, (GIntArg, GStrArg, GPtrArg, GNameArg, bGBoolArg, GWeakObjPtrArg));
+		BIND_CALL_CHECK_DELEGATE_RETVAL(GIntRet, FTestDynamicDelegateRetInt6, 0, &UDynamicDelegateTestObject::ConstFuncRetInt6, (GIntArg, GStrArg, GPtrArg, GNameArg, bGBoolArg, GWeakObjPtrArg));
+		BIND_CALL_CHECK_DELEGATE_RETVAL(GStrRet, FTestDynamicDelegateRetStr6, 0, &UDynamicDelegateTestObject::FuncRetStr6, (GIntArg, GStrArg, GPtrArg, GNameArg, bGBoolArg, GWeakObjPtrArg));
+		BIND_CALL_CHECK_DELEGATE_RETVAL(GStrRet, FTestDynamicDelegateRetStr6, 0, &UDynamicDelegateTestObject::ConstFuncRetStr6, (GIntArg, GStrArg, GPtrArg, GNameArg, bGBoolArg, GWeakObjPtrArg));
 
 #if defined(UE_USE_DYNAMIC_DELEGATE_PAYLOADS) && UE_USE_DYNAMIC_DELEGATE_PAYLOADS
 		// 1 arg payload
-		BIND_CALL_CHECK_DELEGATE       (         FTestDynamicMulticastDelegate5::FDelegate, 0, &UDynamicDelegateTestObject::Func6,            (GIntArg, GStrArg, GPtrArg, GNameArg, bGBoolArg), GWeakObjPtrArg);
-		BIND_CALL_CHECK_DELEGATE       (         FTestDynamicDelegate5,                     0, &UDynamicDelegateTestObject::Func6,            (GIntArg, GStrArg, GPtrArg, GNameArg, bGBoolArg), GWeakObjPtrArg);
-		BIND_CALL_CHECK_DELEGATE       (         FTestDynamicMulticastDelegate5,            1, &UDynamicDelegateTestObject::Func6,            (GIntArg, GStrArg, GPtrArg, GNameArg, bGBoolArg), GWeakObjPtrArg);
-		BIND_CALL_CHECK_DELEGATE       (         FTestDynamicMulticastDelegate5::FDelegate, 0, &UDynamicDelegateTestObject::ConstFunc6,       (GIntArg, GStrArg, GPtrArg, GNameArg, bGBoolArg), GWeakObjPtrArg);
-		BIND_CALL_CHECK_DELEGATE       (         FTestDynamicDelegate5,                     0, &UDynamicDelegateTestObject::ConstFunc6,       (GIntArg, GStrArg, GPtrArg, GNameArg, bGBoolArg), GWeakObjPtrArg);
-		BIND_CALL_CHECK_DELEGATE       (         FTestDynamicMulticastDelegate5,            1, &UDynamicDelegateTestObject::ConstFunc6,       (GIntArg, GStrArg, GPtrArg, GNameArg, bGBoolArg), GWeakObjPtrArg);
+		BIND_CALL_CHECK_DELEGATE(FTestDynamicMulticastDelegate5::FDelegate, 0, &UDynamicDelegateTestObject::Func6, (GIntArg, GStrArg, GPtrArg, GNameArg, bGBoolArg), GWeakObjPtrArg);
+		BIND_CALL_CHECK_DELEGATE(FTestDynamicDelegate5, 0, &UDynamicDelegateTestObject::Func6, (GIntArg, GStrArg, GPtrArg, GNameArg, bGBoolArg), GWeakObjPtrArg);
+		BIND_CALL_CHECK_DELEGATE(FTestDynamicMulticastDelegate5, 1, &UDynamicDelegateTestObject::Func6, (GIntArg, GStrArg, GPtrArg, GNameArg, bGBoolArg), GWeakObjPtrArg);
+		BIND_CALL_CHECK_DELEGATE(FTestDynamicMulticastDelegate5::FDelegate, 0, &UDynamicDelegateTestObject::ConstFunc6, (GIntArg, GStrArg, GPtrArg, GNameArg, bGBoolArg), GWeakObjPtrArg);
+		BIND_CALL_CHECK_DELEGATE(FTestDynamicDelegate5, 0, &UDynamicDelegateTestObject::ConstFunc6, (GIntArg, GStrArg, GPtrArg, GNameArg, bGBoolArg), GWeakObjPtrArg);
+		BIND_CALL_CHECK_DELEGATE(FTestDynamicMulticastDelegate5, 1, &UDynamicDelegateTestObject::ConstFunc6, (GIntArg, GStrArg, GPtrArg, GNameArg, bGBoolArg), GWeakObjPtrArg);
 #if UE_USE_DYNAMIC_DELEGATE_PAYLOADS_RETVAL
-		BIND_CALL_CHECK_DELEGATE_RETVAL(GIntRet, FTestDynamicDelegateRetInt5,               0, &UDynamicDelegateTestObject::FuncRetInt6,      (GIntArg, GStrArg, GPtrArg, GNameArg, bGBoolArg), GWeakObjPtrArg);
-		BIND_CALL_CHECK_DELEGATE_RETVAL(GIntRet, FTestDynamicDelegateRetInt5,               0, &UDynamicDelegateTestObject::ConstFuncRetInt6, (GIntArg, GStrArg, GPtrArg, GNameArg, bGBoolArg), GWeakObjPtrArg);
-		BIND_CALL_CHECK_DELEGATE_RETVAL(GStrRet, FTestDynamicDelegateRetStr5,               0, &UDynamicDelegateTestObject::FuncRetStr6,      (GIntArg, GStrArg, GPtrArg, GNameArg, bGBoolArg), GWeakObjPtrArg);
-		BIND_CALL_CHECK_DELEGATE_RETVAL(GStrRet, FTestDynamicDelegateRetStr5,               0, &UDynamicDelegateTestObject::ConstFuncRetStr6, (GIntArg, GStrArg, GPtrArg, GNameArg, bGBoolArg), GWeakObjPtrArg);
+		BIND_CALL_CHECK_DELEGATE_RETVAL(GIntRet, FTestDynamicDelegateRetInt5, 0, &UDynamicDelegateTestObject::FuncRetInt6, (GIntArg, GStrArg, GPtrArg, GNameArg, bGBoolArg), GWeakObjPtrArg);
+		BIND_CALL_CHECK_DELEGATE_RETVAL(GIntRet, FTestDynamicDelegateRetInt5, 0, &UDynamicDelegateTestObject::ConstFuncRetInt6, (GIntArg, GStrArg, GPtrArg, GNameArg, bGBoolArg), GWeakObjPtrArg);
+		BIND_CALL_CHECK_DELEGATE_RETVAL(GStrRet, FTestDynamicDelegateRetStr5, 0, &UDynamicDelegateTestObject::FuncRetStr6, (GIntArg, GStrArg, GPtrArg, GNameArg, bGBoolArg), GWeakObjPtrArg);
+		BIND_CALL_CHECK_DELEGATE_RETVAL(GStrRet, FTestDynamicDelegateRetStr5, 0, &UDynamicDelegateTestObject::ConstFuncRetStr6, (GIntArg, GStrArg, GPtrArg, GNameArg, bGBoolArg), GWeakObjPtrArg);
 #endif
 
 		// 2 arg payload
-		BIND_CALL_CHECK_DELEGATE       (         FTestDynamicMulticastDelegate4::FDelegate, 0, &UDynamicDelegateTestObject::Func6,            (GIntArg, GStrArg, GPtrArg, GNameArg), bGBoolArg, GWeakObjPtrArg);
-		BIND_CALL_CHECK_DELEGATE       (         FTestDynamicDelegate4,                     0, &UDynamicDelegateTestObject::Func6,            (GIntArg, GStrArg, GPtrArg, GNameArg), bGBoolArg, GWeakObjPtrArg);
-		BIND_CALL_CHECK_DELEGATE       (         FTestDynamicMulticastDelegate4,            1, &UDynamicDelegateTestObject::Func6,            (GIntArg, GStrArg, GPtrArg, GNameArg), bGBoolArg, GWeakObjPtrArg);
-		BIND_CALL_CHECK_DELEGATE       (         FTestDynamicMulticastDelegate4::FDelegate, 0, &UDynamicDelegateTestObject::ConstFunc6,       (GIntArg, GStrArg, GPtrArg, GNameArg), bGBoolArg, GWeakObjPtrArg);
-		BIND_CALL_CHECK_DELEGATE       (         FTestDynamicDelegate4,                     0, &UDynamicDelegateTestObject::ConstFunc6,       (GIntArg, GStrArg, GPtrArg, GNameArg), bGBoolArg, GWeakObjPtrArg);
-		BIND_CALL_CHECK_DELEGATE       (         FTestDynamicMulticastDelegate4,            1, &UDynamicDelegateTestObject::ConstFunc6,       (GIntArg, GStrArg, GPtrArg, GNameArg), bGBoolArg, GWeakObjPtrArg);
+		BIND_CALL_CHECK_DELEGATE(FTestDynamicMulticastDelegate4::FDelegate, 0, &UDynamicDelegateTestObject::Func6, (GIntArg, GStrArg, GPtrArg, GNameArg), bGBoolArg, GWeakObjPtrArg);
+		BIND_CALL_CHECK_DELEGATE(FTestDynamicDelegate4, 0, &UDynamicDelegateTestObject::Func6, (GIntArg, GStrArg, GPtrArg, GNameArg), bGBoolArg, GWeakObjPtrArg);
+		BIND_CALL_CHECK_DELEGATE(FTestDynamicMulticastDelegate4, 1, &UDynamicDelegateTestObject::Func6, (GIntArg, GStrArg, GPtrArg, GNameArg), bGBoolArg, GWeakObjPtrArg);
+		BIND_CALL_CHECK_DELEGATE(FTestDynamicMulticastDelegate4::FDelegate, 0, &UDynamicDelegateTestObject::ConstFunc6, (GIntArg, GStrArg, GPtrArg, GNameArg), bGBoolArg, GWeakObjPtrArg);
+		BIND_CALL_CHECK_DELEGATE(FTestDynamicDelegate4, 0, &UDynamicDelegateTestObject::ConstFunc6, (GIntArg, GStrArg, GPtrArg, GNameArg), bGBoolArg, GWeakObjPtrArg);
+		BIND_CALL_CHECK_DELEGATE(FTestDynamicMulticastDelegate4, 1, &UDynamicDelegateTestObject::ConstFunc6, (GIntArg, GStrArg, GPtrArg, GNameArg), bGBoolArg, GWeakObjPtrArg);
 #if UE_USE_DYNAMIC_DELEGATE_PAYLOADS_RETVAL
-		BIND_CALL_CHECK_DELEGATE_RETVAL(GIntRet, FTestDynamicDelegateRetInt4,               0, &UDynamicDelegateTestObject::FuncRetInt6,      (GIntArg, GStrArg, GPtrArg, GNameArg), bGBoolArg, GWeakObjPtrArg);
-		BIND_CALL_CHECK_DELEGATE_RETVAL(GIntRet, FTestDynamicDelegateRetInt4,               0, &UDynamicDelegateTestObject::ConstFuncRetInt6, (GIntArg, GStrArg, GPtrArg, GNameArg), bGBoolArg, GWeakObjPtrArg);
-		BIND_CALL_CHECK_DELEGATE_RETVAL(GStrRet, FTestDynamicDelegateRetStr4,               0, &UDynamicDelegateTestObject::FuncRetStr6,      (GIntArg, GStrArg, GPtrArg, GNameArg), bGBoolArg, GWeakObjPtrArg);
-		BIND_CALL_CHECK_DELEGATE_RETVAL(GStrRet, FTestDynamicDelegateRetStr4,               0, &UDynamicDelegateTestObject::ConstFuncRetStr6, (GIntArg, GStrArg, GPtrArg, GNameArg), bGBoolArg, GWeakObjPtrArg);
+		BIND_CALL_CHECK_DELEGATE_RETVAL(GIntRet, FTestDynamicDelegateRetInt4, 0, &UDynamicDelegateTestObject::FuncRetInt6, (GIntArg, GStrArg, GPtrArg, GNameArg), bGBoolArg, GWeakObjPtrArg);
+		BIND_CALL_CHECK_DELEGATE_RETVAL(GIntRet, FTestDynamicDelegateRetInt4, 0, &UDynamicDelegateTestObject::ConstFuncRetInt6, (GIntArg, GStrArg, GPtrArg, GNameArg), bGBoolArg, GWeakObjPtrArg);
+		BIND_CALL_CHECK_DELEGATE_RETVAL(GStrRet, FTestDynamicDelegateRetStr4, 0, &UDynamicDelegateTestObject::FuncRetStr6, (GIntArg, GStrArg, GPtrArg, GNameArg), bGBoolArg, GWeakObjPtrArg);
+		BIND_CALL_CHECK_DELEGATE_RETVAL(GStrRet, FTestDynamicDelegateRetStr4, 0, &UDynamicDelegateTestObject::ConstFuncRetStr6, (GIntArg, GStrArg, GPtrArg, GNameArg), bGBoolArg, GWeakObjPtrArg);
 #endif
 
 		// 3 arg payload
-		BIND_CALL_CHECK_DELEGATE       (         FTestDynamicMulticastDelegate3::FDelegate, 0, &UDynamicDelegateTestObject::Func6,            (GIntArg, GStrArg, GPtrArg), GNameArg, bGBoolArg, GWeakObjPtrArg);
-		BIND_CALL_CHECK_DELEGATE       (         FTestDynamicDelegate3,                     0, &UDynamicDelegateTestObject::Func6,            (GIntArg, GStrArg, GPtrArg), GNameArg, bGBoolArg, GWeakObjPtrArg);
-		BIND_CALL_CHECK_DELEGATE       (         FTestDynamicMulticastDelegate3,            1, &UDynamicDelegateTestObject::Func6,            (GIntArg, GStrArg, GPtrArg), GNameArg, bGBoolArg, GWeakObjPtrArg);
-		BIND_CALL_CHECK_DELEGATE       (         FTestDynamicMulticastDelegate3::FDelegate, 0, &UDynamicDelegateTestObject::ConstFunc6,       (GIntArg, GStrArg, GPtrArg), GNameArg, bGBoolArg, GWeakObjPtrArg);
-		BIND_CALL_CHECK_DELEGATE       (         FTestDynamicDelegate3,                     0, &UDynamicDelegateTestObject::ConstFunc6,       (GIntArg, GStrArg, GPtrArg), GNameArg, bGBoolArg, GWeakObjPtrArg);
-		BIND_CALL_CHECK_DELEGATE       (         FTestDynamicMulticastDelegate3,            1, &UDynamicDelegateTestObject::ConstFunc6,       (GIntArg, GStrArg, GPtrArg), GNameArg, bGBoolArg, GWeakObjPtrArg);
+		BIND_CALL_CHECK_DELEGATE(FTestDynamicMulticastDelegate3::FDelegate, 0, &UDynamicDelegateTestObject::Func6, (GIntArg, GStrArg, GPtrArg), GNameArg, bGBoolArg, GWeakObjPtrArg);
+		BIND_CALL_CHECK_DELEGATE(FTestDynamicDelegate3, 0, &UDynamicDelegateTestObject::Func6, (GIntArg, GStrArg, GPtrArg), GNameArg, bGBoolArg, GWeakObjPtrArg);
+		BIND_CALL_CHECK_DELEGATE(FTestDynamicMulticastDelegate3, 1, &UDynamicDelegateTestObject::Func6, (GIntArg, GStrArg, GPtrArg), GNameArg, bGBoolArg, GWeakObjPtrArg);
+		BIND_CALL_CHECK_DELEGATE(FTestDynamicMulticastDelegate3::FDelegate, 0, &UDynamicDelegateTestObject::ConstFunc6, (GIntArg, GStrArg, GPtrArg), GNameArg, bGBoolArg, GWeakObjPtrArg);
+		BIND_CALL_CHECK_DELEGATE(FTestDynamicDelegate3, 0, &UDynamicDelegateTestObject::ConstFunc6, (GIntArg, GStrArg, GPtrArg), GNameArg, bGBoolArg, GWeakObjPtrArg);
+		BIND_CALL_CHECK_DELEGATE(FTestDynamicMulticastDelegate3, 1, &UDynamicDelegateTestObject::ConstFunc6, (GIntArg, GStrArg, GPtrArg), GNameArg, bGBoolArg, GWeakObjPtrArg);
 #if UE_USE_DYNAMIC_DELEGATE_PAYLOADS_RETVAL
-		BIND_CALL_CHECK_DELEGATE_RETVAL(GIntRet, FTestDynamicDelegateRetInt3,               0, &UDynamicDelegateTestObject::FuncRetInt6,      (GIntArg, GStrArg, GPtrArg), GNameArg, bGBoolArg, GWeakObjPtrArg);
-		BIND_CALL_CHECK_DELEGATE_RETVAL(GIntRet, FTestDynamicDelegateRetInt3,               0, &UDynamicDelegateTestObject::ConstFuncRetInt6, (GIntArg, GStrArg, GPtrArg), GNameArg, bGBoolArg, GWeakObjPtrArg);
-		BIND_CALL_CHECK_DELEGATE_RETVAL(GStrRet, FTestDynamicDelegateRetStr3,               0, &UDynamicDelegateTestObject::FuncRetStr6,      (GIntArg, GStrArg, GPtrArg), GNameArg, bGBoolArg, GWeakObjPtrArg);
-		BIND_CALL_CHECK_DELEGATE_RETVAL(GStrRet, FTestDynamicDelegateRetStr3,               0, &UDynamicDelegateTestObject::ConstFuncRetStr6, (GIntArg, GStrArg, GPtrArg), GNameArg, bGBoolArg, GWeakObjPtrArg);
+		BIND_CALL_CHECK_DELEGATE_RETVAL(GIntRet, FTestDynamicDelegateRetInt3, 0, &UDynamicDelegateTestObject::FuncRetInt6, (GIntArg, GStrArg, GPtrArg), GNameArg, bGBoolArg, GWeakObjPtrArg);
+		BIND_CALL_CHECK_DELEGATE_RETVAL(GIntRet, FTestDynamicDelegateRetInt3, 0, &UDynamicDelegateTestObject::ConstFuncRetInt6, (GIntArg, GStrArg, GPtrArg), GNameArg, bGBoolArg, GWeakObjPtrArg);
+		BIND_CALL_CHECK_DELEGATE_RETVAL(GStrRet, FTestDynamicDelegateRetStr3, 0, &UDynamicDelegateTestObject::FuncRetStr6, (GIntArg, GStrArg, GPtrArg), GNameArg, bGBoolArg, GWeakObjPtrArg);
+		BIND_CALL_CHECK_DELEGATE_RETVAL(GStrRet, FTestDynamicDelegateRetStr3, 0, &UDynamicDelegateTestObject::ConstFuncRetStr6, (GIntArg, GStrArg, GPtrArg), GNameArg, bGBoolArg, GWeakObjPtrArg);
 #endif
 
 		// 4 arg payload
-		BIND_CALL_CHECK_DELEGATE       (         FTestDynamicMulticastDelegate2::FDelegate, 0, &UDynamicDelegateTestObject::Func6,            (GIntArg, GStrArg), GPtrArg, GNameArg, bGBoolArg, GWeakObjPtrArg);
-		BIND_CALL_CHECK_DELEGATE       (         FTestDynamicDelegate2,                     0, &UDynamicDelegateTestObject::Func6,            (GIntArg, GStrArg), GPtrArg, GNameArg, bGBoolArg, GWeakObjPtrArg);
-		BIND_CALL_CHECK_DELEGATE       (         FTestDynamicMulticastDelegate2,            1, &UDynamicDelegateTestObject::Func6,            (GIntArg, GStrArg), GPtrArg, GNameArg, bGBoolArg, GWeakObjPtrArg);
-		BIND_CALL_CHECK_DELEGATE       (         FTestDynamicMulticastDelegate2::FDelegate, 0, &UDynamicDelegateTestObject::ConstFunc6,       (GIntArg, GStrArg), GPtrArg, GNameArg, bGBoolArg, GWeakObjPtrArg);
-		BIND_CALL_CHECK_DELEGATE       (         FTestDynamicDelegate2,                     0, &UDynamicDelegateTestObject::ConstFunc6,       (GIntArg, GStrArg), GPtrArg, GNameArg, bGBoolArg, GWeakObjPtrArg);
-		BIND_CALL_CHECK_DELEGATE       (         FTestDynamicMulticastDelegate2,            1, &UDynamicDelegateTestObject::ConstFunc6,       (GIntArg, GStrArg), GPtrArg, GNameArg, bGBoolArg, GWeakObjPtrArg);
+		BIND_CALL_CHECK_DELEGATE(FTestDynamicMulticastDelegate2::FDelegate, 0, &UDynamicDelegateTestObject::Func6, (GIntArg, GStrArg), GPtrArg, GNameArg, bGBoolArg, GWeakObjPtrArg);
+		BIND_CALL_CHECK_DELEGATE(FTestDynamicDelegate2, 0, &UDynamicDelegateTestObject::Func6, (GIntArg, GStrArg), GPtrArg, GNameArg, bGBoolArg, GWeakObjPtrArg);
+		BIND_CALL_CHECK_DELEGATE(FTestDynamicMulticastDelegate2, 1, &UDynamicDelegateTestObject::Func6, (GIntArg, GStrArg), GPtrArg, GNameArg, bGBoolArg, GWeakObjPtrArg);
+		BIND_CALL_CHECK_DELEGATE(FTestDynamicMulticastDelegate2::FDelegate, 0, &UDynamicDelegateTestObject::ConstFunc6, (GIntArg, GStrArg), GPtrArg, GNameArg, bGBoolArg, GWeakObjPtrArg);
+		BIND_CALL_CHECK_DELEGATE(FTestDynamicDelegate2, 0, &UDynamicDelegateTestObject::ConstFunc6, (GIntArg, GStrArg), GPtrArg, GNameArg, bGBoolArg, GWeakObjPtrArg);
+		BIND_CALL_CHECK_DELEGATE(FTestDynamicMulticastDelegate2, 1, &UDynamicDelegateTestObject::ConstFunc6, (GIntArg, GStrArg), GPtrArg, GNameArg, bGBoolArg, GWeakObjPtrArg);
 #if UE_USE_DYNAMIC_DELEGATE_PAYLOADS_RETVAL
-		BIND_CALL_CHECK_DELEGATE_RETVAL(GIntRet, FTestDynamicDelegateRetInt2,               0, &UDynamicDelegateTestObject::FuncRetInt6,      (GIntArg, GStrArg), GPtrArg, GNameArg, bGBoolArg, GWeakObjPtrArg);
-		BIND_CALL_CHECK_DELEGATE_RETVAL(GIntRet, FTestDynamicDelegateRetInt2,               0, &UDynamicDelegateTestObject::ConstFuncRetInt6, (GIntArg, GStrArg), GPtrArg, GNameArg, bGBoolArg, GWeakObjPtrArg);
-		BIND_CALL_CHECK_DELEGATE_RETVAL(GStrRet, FTestDynamicDelegateRetStr2,               0, &UDynamicDelegateTestObject::FuncRetStr6,      (GIntArg, GStrArg), GPtrArg, GNameArg, bGBoolArg, GWeakObjPtrArg);
-		BIND_CALL_CHECK_DELEGATE_RETVAL(GStrRet, FTestDynamicDelegateRetStr2,               0, &UDynamicDelegateTestObject::ConstFuncRetStr6, (GIntArg, GStrArg), GPtrArg, GNameArg, bGBoolArg, GWeakObjPtrArg);
+		BIND_CALL_CHECK_DELEGATE_RETVAL(GIntRet, FTestDynamicDelegateRetInt2, 0, &UDynamicDelegateTestObject::FuncRetInt6, (GIntArg, GStrArg), GPtrArg, GNameArg, bGBoolArg, GWeakObjPtrArg);
+		BIND_CALL_CHECK_DELEGATE_RETVAL(GIntRet, FTestDynamicDelegateRetInt2, 0, &UDynamicDelegateTestObject::ConstFuncRetInt6, (GIntArg, GStrArg), GPtrArg, GNameArg, bGBoolArg, GWeakObjPtrArg);
+		BIND_CALL_CHECK_DELEGATE_RETVAL(GStrRet, FTestDynamicDelegateRetStr2, 0, &UDynamicDelegateTestObject::FuncRetStr6, (GIntArg, GStrArg), GPtrArg, GNameArg, bGBoolArg, GWeakObjPtrArg);
+		BIND_CALL_CHECK_DELEGATE_RETVAL(GStrRet, FTestDynamicDelegateRetStr2, 0, &UDynamicDelegateTestObject::ConstFuncRetStr6, (GIntArg, GStrArg), GPtrArg, GNameArg, bGBoolArg, GWeakObjPtrArg);
 #endif
 
 		// 5 arg payload
-		BIND_CALL_CHECK_DELEGATE       (         FTestDynamicMulticastDelegate1::FDelegate, 0, &UDynamicDelegateTestObject::Func6,            (GIntArg), GStrArg, GPtrArg, GNameArg, bGBoolArg, GWeakObjPtrArg);
-		BIND_CALL_CHECK_DELEGATE       (         FTestDynamicDelegate1,                     0, &UDynamicDelegateTestObject::Func6,            (GIntArg), GStrArg, GPtrArg, GNameArg, bGBoolArg, GWeakObjPtrArg);
-		BIND_CALL_CHECK_DELEGATE       (         FTestDynamicMulticastDelegate1,            1, &UDynamicDelegateTestObject::Func6,            (GIntArg), GStrArg, GPtrArg, GNameArg, bGBoolArg, GWeakObjPtrArg);
-		BIND_CALL_CHECK_DELEGATE       (         FTestDynamicMulticastDelegate1::FDelegate, 0, &UDynamicDelegateTestObject::ConstFunc6,       (GIntArg), GStrArg, GPtrArg, GNameArg, bGBoolArg, GWeakObjPtrArg);
-		BIND_CALL_CHECK_DELEGATE       (         FTestDynamicDelegate1,                     0, &UDynamicDelegateTestObject::ConstFunc6,       (GIntArg), GStrArg, GPtrArg, GNameArg, bGBoolArg, GWeakObjPtrArg);
-		BIND_CALL_CHECK_DELEGATE       (         FTestDynamicMulticastDelegate1,            1, &UDynamicDelegateTestObject::ConstFunc6,       (GIntArg), GStrArg, GPtrArg, GNameArg, bGBoolArg, GWeakObjPtrArg);
+		BIND_CALL_CHECK_DELEGATE(FTestDynamicMulticastDelegate1::FDelegate, 0, &UDynamicDelegateTestObject::Func6, (GIntArg), GStrArg, GPtrArg, GNameArg, bGBoolArg, GWeakObjPtrArg);
+		BIND_CALL_CHECK_DELEGATE(FTestDynamicDelegate1, 0, &UDynamicDelegateTestObject::Func6, (GIntArg), GStrArg, GPtrArg, GNameArg, bGBoolArg, GWeakObjPtrArg);
+		BIND_CALL_CHECK_DELEGATE(FTestDynamicMulticastDelegate1, 1, &UDynamicDelegateTestObject::Func6, (GIntArg), GStrArg, GPtrArg, GNameArg, bGBoolArg, GWeakObjPtrArg);
+		BIND_CALL_CHECK_DELEGATE(FTestDynamicMulticastDelegate1::FDelegate, 0, &UDynamicDelegateTestObject::ConstFunc6, (GIntArg), GStrArg, GPtrArg, GNameArg, bGBoolArg, GWeakObjPtrArg);
+		BIND_CALL_CHECK_DELEGATE(FTestDynamicDelegate1, 0, &UDynamicDelegateTestObject::ConstFunc6, (GIntArg), GStrArg, GPtrArg, GNameArg, bGBoolArg, GWeakObjPtrArg);
+		BIND_CALL_CHECK_DELEGATE(FTestDynamicMulticastDelegate1, 1, &UDynamicDelegateTestObject::ConstFunc6, (GIntArg), GStrArg, GPtrArg, GNameArg, bGBoolArg, GWeakObjPtrArg);
 #if UE_USE_DYNAMIC_DELEGATE_PAYLOADS_RETVAL
-		BIND_CALL_CHECK_DELEGATE_RETVAL(GIntRet, FTestDynamicDelegateRetInt1,               0, &UDynamicDelegateTestObject::FuncRetInt6,      (GIntArg), GStrArg, GPtrArg, GNameArg, bGBoolArg, GWeakObjPtrArg);
-		BIND_CALL_CHECK_DELEGATE_RETVAL(GIntRet, FTestDynamicDelegateRetInt1,               0, &UDynamicDelegateTestObject::ConstFuncRetInt6, (GIntArg), GStrArg, GPtrArg, GNameArg, bGBoolArg, GWeakObjPtrArg);
-		BIND_CALL_CHECK_DELEGATE_RETVAL(GStrRet, FTestDynamicDelegateRetStr1,               0, &UDynamicDelegateTestObject::FuncRetStr6,      (GIntArg), GStrArg, GPtrArg, GNameArg, bGBoolArg, GWeakObjPtrArg);
-		BIND_CALL_CHECK_DELEGATE_RETVAL(GStrRet, FTestDynamicDelegateRetStr1,               0, &UDynamicDelegateTestObject::ConstFuncRetStr6, (GIntArg), GStrArg, GPtrArg, GNameArg, bGBoolArg, GWeakObjPtrArg);
+		BIND_CALL_CHECK_DELEGATE_RETVAL(GIntRet, FTestDynamicDelegateRetInt1, 0, &UDynamicDelegateTestObject::FuncRetInt6, (GIntArg), GStrArg, GPtrArg, GNameArg, bGBoolArg, GWeakObjPtrArg);
+		BIND_CALL_CHECK_DELEGATE_RETVAL(GIntRet, FTestDynamicDelegateRetInt1, 0, &UDynamicDelegateTestObject::ConstFuncRetInt6, (GIntArg), GStrArg, GPtrArg, GNameArg, bGBoolArg, GWeakObjPtrArg);
+		BIND_CALL_CHECK_DELEGATE_RETVAL(GStrRet, FTestDynamicDelegateRetStr1, 0, &UDynamicDelegateTestObject::FuncRetStr6, (GIntArg), GStrArg, GPtrArg, GNameArg, bGBoolArg, GWeakObjPtrArg);
+		BIND_CALL_CHECK_DELEGATE_RETVAL(GStrRet, FTestDynamicDelegateRetStr1, 0, &UDynamicDelegateTestObject::ConstFuncRetStr6, (GIntArg), GStrArg, GPtrArg, GNameArg, bGBoolArg, GWeakObjPtrArg);
 #endif
 
 		// 6 arg payload
-		BIND_CALL_CHECK_DELEGATE       (         FTestDynamicMulticastDelegate0::FDelegate, 0, &UDynamicDelegateTestObject::Func6,            (), GIntArg, GStrArg, GPtrArg, GNameArg, bGBoolArg, GWeakObjPtrArg);
-		BIND_CALL_CHECK_DELEGATE       (         FTestDynamicDelegate0,                     0, &UDynamicDelegateTestObject::Func6,            (), GIntArg, GStrArg, GPtrArg, GNameArg, bGBoolArg, GWeakObjPtrArg);
-		BIND_CALL_CHECK_DELEGATE       (         FTestDynamicMulticastDelegate0,            1, &UDynamicDelegateTestObject::Func6,            (), GIntArg, GStrArg, GPtrArg, GNameArg, bGBoolArg, GWeakObjPtrArg);
-		BIND_CALL_CHECK_DELEGATE       (         FTestDynamicMulticastDelegate0::FDelegate, 0, &UDynamicDelegateTestObject::ConstFunc6,       (), GIntArg, GStrArg, GPtrArg, GNameArg, bGBoolArg, GWeakObjPtrArg);
-		BIND_CALL_CHECK_DELEGATE       (         FTestDynamicDelegate0,                     0, &UDynamicDelegateTestObject::ConstFunc6,       (), GIntArg, GStrArg, GPtrArg, GNameArg, bGBoolArg, GWeakObjPtrArg);
-		BIND_CALL_CHECK_DELEGATE       (         FTestDynamicMulticastDelegate0,            1, &UDynamicDelegateTestObject::ConstFunc6,       (), GIntArg, GStrArg, GPtrArg, GNameArg, bGBoolArg, GWeakObjPtrArg);
+		BIND_CALL_CHECK_DELEGATE(FTestDynamicMulticastDelegate0::FDelegate, 0, &UDynamicDelegateTestObject::Func6, (), GIntArg, GStrArg, GPtrArg, GNameArg, bGBoolArg, GWeakObjPtrArg);
+		BIND_CALL_CHECK_DELEGATE(FTestDynamicDelegate0, 0, &UDynamicDelegateTestObject::Func6, (), GIntArg, GStrArg, GPtrArg, GNameArg, bGBoolArg, GWeakObjPtrArg);
+		BIND_CALL_CHECK_DELEGATE(FTestDynamicMulticastDelegate0, 1, &UDynamicDelegateTestObject::Func6, (), GIntArg, GStrArg, GPtrArg, GNameArg, bGBoolArg, GWeakObjPtrArg);
+		BIND_CALL_CHECK_DELEGATE(FTestDynamicMulticastDelegate0::FDelegate, 0, &UDynamicDelegateTestObject::ConstFunc6, (), GIntArg, GStrArg, GPtrArg, GNameArg, bGBoolArg, GWeakObjPtrArg);
+		BIND_CALL_CHECK_DELEGATE(FTestDynamicDelegate0, 0, &UDynamicDelegateTestObject::ConstFunc6, (), GIntArg, GStrArg, GPtrArg, GNameArg, bGBoolArg, GWeakObjPtrArg);
+		BIND_CALL_CHECK_DELEGATE(FTestDynamicMulticastDelegate0, 1, &UDynamicDelegateTestObject::ConstFunc6, (), GIntArg, GStrArg, GPtrArg, GNameArg, bGBoolArg, GWeakObjPtrArg);
 #if UE_USE_DYNAMIC_DELEGATE_PAYLOADS_RETVAL
-		BIND_CALL_CHECK_DELEGATE_RETVAL(GIntRet, FTestDynamicDelegateRetInt0,               0, &UDynamicDelegateTestObject::FuncRetInt6,      (), GIntArg, GStrArg, GPtrArg, GNameArg, bGBoolArg, GWeakObjPtrArg);
-		BIND_CALL_CHECK_DELEGATE_RETVAL(GIntRet, FTestDynamicDelegateRetInt0,               0, &UDynamicDelegateTestObject::ConstFuncRetInt6, (), GIntArg, GStrArg, GPtrArg, GNameArg, bGBoolArg, GWeakObjPtrArg);
-		BIND_CALL_CHECK_DELEGATE_RETVAL(GStrRet, FTestDynamicDelegateRetStr0,               0, &UDynamicDelegateTestObject::FuncRetStr6,      (), GIntArg, GStrArg, GPtrArg, GNameArg, bGBoolArg, GWeakObjPtrArg);
-		BIND_CALL_CHECK_DELEGATE_RETVAL(GStrRet, FTestDynamicDelegateRetStr0,               0, &UDynamicDelegateTestObject::ConstFuncRetStr6, (), GIntArg, GStrArg, GPtrArg, GNameArg, bGBoolArg, GWeakObjPtrArg);
+		BIND_CALL_CHECK_DELEGATE_RETVAL(GIntRet, FTestDynamicDelegateRetInt0, 0, &UDynamicDelegateTestObject::FuncRetInt6, (), GIntArg, GStrArg, GPtrArg, GNameArg, bGBoolArg, GWeakObjPtrArg);
+		BIND_CALL_CHECK_DELEGATE_RETVAL(GIntRet, FTestDynamicDelegateRetInt0, 0, &UDynamicDelegateTestObject::ConstFuncRetInt6, (), GIntArg, GStrArg, GPtrArg, GNameArg, bGBoolArg, GWeakObjPtrArg);
+		BIND_CALL_CHECK_DELEGATE_RETVAL(GStrRet, FTestDynamicDelegateRetStr0, 0, &UDynamicDelegateTestObject::FuncRetStr6, (), GIntArg, GStrArg, GPtrArg, GNameArg, bGBoolArg, GWeakObjPtrArg);
+		BIND_CALL_CHECK_DELEGATE_RETVAL(GStrRet, FTestDynamicDelegateRetStr0, 0, &UDynamicDelegateTestObject::ConstFuncRetStr6, (), GIntArg, GStrArg, GPtrArg, GNameArg, bGBoolArg, GWeakObjPtrArg);
 #endif
 #endif
 	}
+
+	// 1 out param function
+	{
+		BIND_CALL_CHECK_DELEGATE_OUT(FTestDynamicMulticastDelegateOut1::FDelegate, 0, &UDynamicDelegateTestObject::FuncOut1, (GIntArg), (GIntOut));
+		BIND_CALL_CHECK_DELEGATE_OUT(FTestDynamicDelegateOut1, 0, &UDynamicDelegateTestObject::FuncOut1, (GIntArg), (GIntOut));
+		BIND_CALL_CHECK_DELEGATE_OUT(FTestDynamicMulticastDelegateOut1, 1, &UDynamicDelegateTestObject::FuncOut1, (GIntArg), (GIntOut));
+		BIND_CALL_CHECK_DELEGATE_OUT(FTestDynamicMulticastDelegateOut1::FDelegate, 0, &UDynamicDelegateTestObject::ConstFuncOut1, (GIntArg), (GIntOut));
+		BIND_CALL_CHECK_DELEGATE_OUT(FTestDynamicDelegateOut1, 0, &UDynamicDelegateTestObject::ConstFuncOut1, (GIntArg), (GIntOut));
+		BIND_CALL_CHECK_DELEGATE_OUT(FTestDynamicMulticastDelegateOut1, 1, &UDynamicDelegateTestObject::ConstFuncOut1, (GIntArg), (GIntOut));
+	}
+
+	// 2 out param function
+	{
+		BIND_CALL_CHECK_DELEGATE_OUT(FTestDynamicMulticastDelegateOut2::FDelegate, 0, &UDynamicDelegateTestObject::FuncOut2, (GIntArg, GStrArg), (GIntOut, GStrOut));
+		BIND_CALL_CHECK_DELEGATE_OUT(FTestDynamicDelegateOut2, 0, &UDynamicDelegateTestObject::FuncOut2, (GIntArg, GStrArg), (GIntOut, GStrOut));
+		BIND_CALL_CHECK_DELEGATE_OUT(FTestDynamicMulticastDelegateOut2, 1, &UDynamicDelegateTestObject::FuncOut2, (GIntArg, GStrArg), (GIntOut, GStrOut));
+		BIND_CALL_CHECK_DELEGATE_OUT(FTestDynamicMulticastDelegateOut2::FDelegate, 0, &UDynamicDelegateTestObject::ConstFuncOut2, (GIntArg, GStrArg), (GIntOut, GStrOut));
+		BIND_CALL_CHECK_DELEGATE_OUT(FTestDynamicDelegateOut2, 0, &UDynamicDelegateTestObject::ConstFuncOut2, (GIntArg, GStrArg), (GIntOut, GStrOut));
+		BIND_CALL_CHECK_DELEGATE_OUT(FTestDynamicMulticastDelegateOut2, 1, &UDynamicDelegateTestObject::ConstFuncOut2, (GIntArg, GStrArg), (GIntOut, GStrOut));
+	}
+
+	// 3 out param function
+	{
+		BIND_CALL_CHECK_DELEGATE_OUT(FTestDynamicMulticastDelegateOut3::FDelegate, 0, &UDynamicDelegateTestObject::FuncOut3, (GIntArg, GStrArg, GNameArg), (GIntOut, GStrOut, GNameOut));
+		BIND_CALL_CHECK_DELEGATE_OUT(FTestDynamicDelegateOut3, 0, &UDynamicDelegateTestObject::FuncOut3, (GIntArg, GStrArg, GNameArg), (GIntOut, GStrOut, GNameOut));
+		BIND_CALL_CHECK_DELEGATE_OUT(FTestDynamicMulticastDelegateOut3, 1, &UDynamicDelegateTestObject::FuncOut3, (GIntArg, GStrArg, GNameArg), (GIntOut, GStrOut, GNameOut));
+		BIND_CALL_CHECK_DELEGATE_OUT(FTestDynamicMulticastDelegateOut3::FDelegate, 0, &UDynamicDelegateTestObject::ConstFuncOut3, (GIntArg, GStrArg, GNameArg), (GIntOut, GStrOut, GNameOut));
+		BIND_CALL_CHECK_DELEGATE_OUT(FTestDynamicDelegateOut3, 0, &UDynamicDelegateTestObject::ConstFuncOut3, (GIntArg, GStrArg, GNameArg), (GIntOut, GStrOut, GNameOut));
+		BIND_CALL_CHECK_DELEGATE_OUT(FTestDynamicMulticastDelegateOut3, 1, &UDynamicDelegateTestObject::ConstFuncOut3, (GIntArg, GStrArg, GNameArg), (GIntOut, GStrOut, GNameOut));
+	}
+
+	// 4 out param function
+	{
+		BIND_CALL_CHECK_DELEGATE_OUT(FTestDynamicMulticastDelegateOut4::FDelegate, 0, &UDynamicDelegateTestObject::FuncOut4, (GIntArg, GStrArg, GNameArg, bGBoolArg), (GIntOut, GStrOut, GNameOut, bGBoolOut));
+		BIND_CALL_CHECK_DELEGATE_OUT(FTestDynamicDelegateOut4, 0, &UDynamicDelegateTestObject::FuncOut4, (GIntArg, GStrArg, GNameArg, bGBoolArg), (GIntOut, GStrOut, GNameOut, bGBoolOut));
+		BIND_CALL_CHECK_DELEGATE_OUT(FTestDynamicMulticastDelegateOut4, 1, &UDynamicDelegateTestObject::FuncOut4, (GIntArg, GStrArg, GNameArg, bGBoolArg), (GIntOut, GStrOut, GNameOut, bGBoolOut));
+		BIND_CALL_CHECK_DELEGATE_OUT(FTestDynamicMulticastDelegateOut4::FDelegate, 0, &UDynamicDelegateTestObject::ConstFuncOut4, (GIntArg, GStrArg, GNameArg, bGBoolArg), (GIntOut, GStrOut, GNameOut, bGBoolOut));
+		BIND_CALL_CHECK_DELEGATE_OUT(FTestDynamicDelegateOut4, 0, &UDynamicDelegateTestObject::ConstFuncOut4, (GIntArg, GStrArg, GNameArg, bGBoolArg), (GIntOut, GStrOut, GNameOut, bGBoolOut));
+		BIND_CALL_CHECK_DELEGATE_OUT(FTestDynamicMulticastDelegateOut4, 1, &UDynamicDelegateTestObject::ConstFuncOut4, (GIntArg, GStrArg, GNameArg, bGBoolArg), (GIntOut, GStrOut, GNameOut, bGBoolOut));
+	}
+
+	// 5 out param function
+	{
+		BIND_CALL_CHECK_DELEGATE_OUT(FTestDynamicMulticastDelegateOut5::FDelegate, 0, &UDynamicDelegateTestObject::FuncOut5, (GIntArg, GStrArg, GNameArg, bGBoolArg, GWeakObjPtrArg), (GIntOut, GStrOut, GNameOut, bGBoolOut, GWeakObjPtrOut));
+		BIND_CALL_CHECK_DELEGATE_OUT(FTestDynamicDelegateOut5, 0, &UDynamicDelegateTestObject::FuncOut5, (GIntArg, GStrArg, GNameArg, bGBoolArg, GWeakObjPtrArg), (GIntOut, GStrOut, GNameOut, bGBoolOut, GWeakObjPtrOut));
+		BIND_CALL_CHECK_DELEGATE_OUT(FTestDynamicMulticastDelegateOut5, 1, &UDynamicDelegateTestObject::FuncOut5, (GIntArg, GStrArg, GNameArg, bGBoolArg, GWeakObjPtrArg), (GIntOut, GStrOut, GNameOut, bGBoolOut, GWeakObjPtrOut));
+		BIND_CALL_CHECK_DELEGATE_OUT(FTestDynamicMulticastDelegateOut5::FDelegate, 0, &UDynamicDelegateTestObject::ConstFuncOut5, (GIntArg, GStrArg, GNameArg, bGBoolArg, GWeakObjPtrArg), (GIntOut, GStrOut, GNameOut, bGBoolOut, GWeakObjPtrOut));
+		BIND_CALL_CHECK_DELEGATE_OUT(FTestDynamicDelegateOut5, 0, &UDynamicDelegateTestObject::ConstFuncOut5, (GIntArg, GStrArg, GNameArg, bGBoolArg, GWeakObjPtrArg), (GIntOut, GStrOut, GNameOut, bGBoolOut, GWeakObjPtrOut));
+		BIND_CALL_CHECK_DELEGATE_OUT(FTestDynamicMulticastDelegateOut5, 1, &UDynamicDelegateTestObject::ConstFuncOut5, (GIntArg, GStrArg, GNameArg, bGBoolArg, GWeakObjPtrArg), (GIntOut, GStrOut, GNameOut, bGBoolOut, GWeakObjPtrOut));
+	}
 }
 
 TEST_CASE_NAMED(DynamicDelegateTest, "UE::CoreUObject::DynamicDelegateTest", "[ApplicationContextMask][SmokeFilter]")
@@ -801,6 +985,32 @@ TEST_CASE_NAMED(DynamicDelegateTest, "UE::CoreUObject::DynamicDelegateTest", "[A
 	TestObj->RunTests();
 }
 
+#undef CHECK_ARG_VARIABLES
+#undef EXPAND_ARG_VARIABLES
+#undef DEFINE_ARG_VARIABLES
+
+#undef CHECK_ARG_VARIABLES_6
+#undef CHECK_ARG_VARIABLES_5
+#undef CHECK_ARG_VARIABLES_4
+#undef CHECK_ARG_VARIABLES_3
+#undef CHECK_ARG_VARIABLES_2
+#undef CHECK_ARG_VARIABLES_1
+
+#undef EXPAND_ARG_VARIABLES_6
+#undef EXPAND_ARG_VARIABLES_5
+#undef EXPAND_ARG_VARIABLES_4
+#undef EXPAND_ARG_VARIABLES_3
+#undef EXPAND_ARG_VARIABLES_2
+#undef EXPAND_ARG_VARIABLES_1
+
+#undef DEFINE_ARG_VARIABLES_6
+#undef DEFINE_ARG_VARIABLES_5
+#undef DEFINE_ARG_VARIABLES_4
+#undef DEFINE_ARG_VARIABLES_3
+#undef DEFINE_ARG_VARIABLES_2
+#undef DEFINE_ARG_VARIABLES_1
+
+#undef BIND_CALL_CHECK_DELEGATE_OUT
 #undef BIND_CALL_CHECK_DELEGATE_RETVAL
 #undef BIND_CALL_CHECK_DELEGATE
 
diff --git a/Engine/Source/Runtime/CoreUObject/Public/UObject/DynamicDelegateTestObject.h b/Engine/Source/Runtime/CoreUObject/Public/UObject/DynamicDelegateTestObject.h
--- a/Engine/Source/Runtime/CoreUObject/Public/UObject/DynamicDelegateTestObject.h
+++ b/Engine/Source/Runtime/CoreUObject/Public/UObject/DynamicDelegateTestObject.h
@@ -41,6 +41,19 @@ DECLARE_DYNAMIC_DELEGATE_RetVal_FourParams (FString, FTestDynamicDelegateRetStr4
 DECLARE_DYNAMIC_DELEGATE_RetVal_FiveParams (FString, FTestDynamicDelegateRetStr5, int, Int, FString, Str, const UDynamicDelegateTestObject*, Ptr, FName, Name, bool, bBool);
 DECLARE_DYNAMIC_DELEGATE_RetVal_SixParams  (FString, FTestDynamicDelegateRetStr6, int, Int, FString, Str, const UDynamicDelegateTestObject*, Ptr, FName, Name, bool, bBool, TWeakObjectPtr<UDynamicDelegateTestObject>, WeakPtr);
 
+DECLARE_DYNAMIC_DELEGATE_OneParam   (FTestDynamicDelegateOut1, int&, Int);
+DECLARE_DYNAMIC_DELEGATE_TwoParams  (FTestDynamicDelegateOut2, int&, Int, FString&, Str);
+DECLARE_DYNAMIC_DELEGATE_ThreeParams(FTestDynamicDelegateOut3, int&, Int, FString&, Str, FName&, Name);
+DECLARE_DYNAMIC_DELEGATE_FourParams (FTestDynamicDelegateOut4, int&, Int, FString&, Str, FName&, Name, bool&, bBool);
+DECLARE_DYNAMIC_DELEGATE_FiveParams (FTestDynamicDelegateOut5, int&, Int, FString&, Str, FName&, Name, bool&, bBool, TWeakObjectPtr<UDynamicDelegateTestObject>&, WeakPtr);
+
+// Pointer out params aren't currently supported
+DECLARE_DYNAMIC_MULTICAST_DELEGATE_OneParam   (FTestDynamicMulticastDelegateOut1, int&, Int);
+DECLARE_DYNAMIC_MULTICAST_DELEGATE_TwoParams  (FTestDynamicMulticastDelegateOut2, int&, Int, FString&, Str);
+DECLARE_DYNAMIC_MULTICAST_DELEGATE_ThreeParams(FTestDynamicMulticastDelegateOut3, int&, Int, FString&, Str, FName&, Name);
+DECLARE_DYNAMIC_MULTICAST_DELEGATE_FourParams (FTestDynamicMulticastDelegateOut4, int&, Int, FString&, Str, FName&, Name, bool&, bBool);
+DECLARE_DYNAMIC_MULTICAST_DELEGATE_FiveParams (FTestDynamicMulticastDelegateOut5, int&, Int, FString&, Str, FName&, Name, bool&, bBool, TWeakObjectPtr<UDynamicDelegateTestObject>&, WeakPtr);
+
 UCLASS()
 class UDynamicDelegateTestObject : public UObject
 {
@@ -97,6 +110,20 @@ public:
 	UFUNCTION() FString ConstFuncRetStr5(int Int, FString Str, const UDynamicDelegateTestObject* Ptr, FName Name, bool bBool) const;
 	UFUNCTION() FString ConstFuncRetStr6(int Int, FString Str, const UDynamicDelegateTestObject* Ptr, FName Name, bool bBool, TWeakObjectPtr<UDynamicDelegateTestObject> WeakPtr) const;
 
+	// Pointer out params aren't currently supported
+	UFUNCTION() void FuncOut1(int& Int);
+	UFUNCTION() void FuncOut2(int& Int, FString& Str);
+	UFUNCTION() void FuncOut3(int& Int, FString& Str, FName& Name);
+	UFUNCTION() void FuncOut4(int& Int, FString& Str, FName& Name, bool& bBool);
+	UFUNCTION() void FuncOut5(int& Int, FString& Str, FName& Name, bool& bBool, TWeakObjectPtr<UDynamicDelegateTestObject>& WeakPtr);
+
+	// Pointer out params aren't currently supported
+	UFUNCTION() void ConstFuncOut1(int& Int) const;
+	UFUNCTION() void ConstFuncOut2(int& Int, FString& Str) const;
+	UFUNCTION() void ConstFuncOut3(int& Int, FString& Str, FName& Name) const;
+	UFUNCTION() void ConstFuncOut4(int& Int, FString& Str, FName& Name, bool& bBool) const;
+	UFUNCTION() void ConstFuncOut5(int& Int, FString& Str, FName& Name, bool& bBool, TWeakObjectPtr<UDynamicDelegateTestObject>& WeakPtr) const;
+
 	void RunTests();
 };
 
