Use TSparseSet for reference finding to accomodate huge arrays for reference replacement

#rnx

#Synced-CL 50069009

diff --git a/Engine/Source/Runtime/CoreUObject/Private/UObject/ReferencerFinder.cpp b/Engine/Source/Runtime/CoreUObject/Private/UObject/ReferencerFinder.cpp
--- a/Engine/Source/Runtime/CoreUObject/Private/UObject/ReferencerFinder.cpp
+++ b/Engine/Source/Runtime/CoreUObject/Private/UObject/ReferencerFinder.cpp
@@ -12,12 +12,12 @@
 
 class FAllReferencesProcessor : public FSimpleReferenceProcessorBase
 {
-	const TSet<UObject*>& PotentiallyReferencedObjects;
-	TSet<UObject*>& ReferencingObjects;
+	const TSparseSet<UObject*>& PotentiallyReferencedObjects;
+	TSparseSet<UObject*>& ReferencingObjects;
 	EReferencerFinderFlags Flags;
 
 public:
-	FAllReferencesProcessor(const TSet<UObject*>& InPotentiallyReferencedObjects, EReferencerFinderFlags InFlags, TSet<UObject*>& OutReferencingObjects)
+	FAllReferencesProcessor(const TSparseSet<UObject*>& InPotentiallyReferencedObjects, EReferencerFinderFlags InFlags, TSparseSet<UObject*>& OutReferencingObjects)
 		: PotentiallyReferencedObjects(InPotentiallyReferencedObjects)
 		, ReferencingObjects(OutReferencingObjects)
 		, Flags(InFlags)
@@ -88,12 +88,12 @@ struct FReferencerFinderArchive : public FArchiveUObject
 {
 private:
 	UObject* CurrentObject = nullptr;
-	const TSet<UObject*>& Targets;
+	const TSparseSet<UObject*>& Targets;
 	bool bDoesReferenceAnyTargets = false;
 	const bool bSkipInnerReferences;
 
 public:
-	FReferencerFinderArchive(const TSet<UObject*>& InTargets, bool bInSkipInnerReferences)
+	FReferencerFinderArchive(const TSparseSet<UObject*>& InTargets, bool bInSkipInnerReferences)
 		: FArchiveUObject()
 		, Targets(InTargets)
 		, bDoesReferenceAnyTargets(false)
@@ -165,7 +165,7 @@ private:
 };
 
 static TArray<UObject*> GetAllReferencersIncludingWeak(
-	const TSet<UObject*>& Referencees, const TSet<UObject*>* ObjectsToIgnore, const bool bSkipInnerReferences)
+	const TSparseSet<UObject*>& Referencees, const TSet<UObject*>* ObjectsToIgnore, const bool bSkipInnerReferences)
 {
 	// Lock the global array so that nothing can add UObjects while we're iterating over it
 	GUObjectArray.LockInternalArray();
@@ -175,8 +175,8 @@ static TArray<UObject*> GetAllReferencersIncludingWeak(
 	const int32 NumberOfObjectsPerThread = (MaxNumberOfObjects / NumThreads) + 1;
 	
 	// Allocate per thread results, in this case each thread will produce a list of referencers
-	TUniquePtr<TSet<UObject*>[]> ThreadResultsAlloc = MakeUnique<TSet<UObject*>[]>(NumThreads);
-	TArrayView<TSet<UObject*>> ThreadResults(ThreadResultsAlloc.Get(), NumThreads);
+	TUniquePtr<TSparseSet<UObject*>[]> ThreadResultsAlloc = MakeUnique<TSparseSet<UObject*>[]>(NumThreads);
+	TArrayView<TSparseSet<UObject*>> ThreadResults(ThreadResultsAlloc.Get(), NumThreads);
 
 	ParallelFor(NumThreads, 
 		[&Referencees, 
@@ -189,7 +189,7 @@ static TArray<UObject*> GetAllReferencersIncludingWeak(
 	{
 		const EReferencerFinderFlags Flags = bSkipInnerReferences ? EReferencerFinderFlags::SkipInnerReferences : EReferencerFinderFlags::None;
 		FReferencerFinderArchive ReferenceFinderArchive(Referencees, bSkipInnerReferences);
-		TSet<UObject*>& ThreadResult = ThreadResults[ThreadIndex];
+		TSparseSet<UObject*>& ThreadResult = ThreadResults[ThreadIndex];
 		TArray<UObject*> ObjectsToSearch;
 		ObjectsToSearch.Reserve(NumberOfObjectsPerThread);
 		
@@ -245,7 +245,7 @@ static TArray<UObject*> GetAllReferencersIncludingWeak(
 	GUObjectArray.UnlockInternalArray();
 
 	TArray<UObject*> FinalResult;
-	for(const TSet<UObject*>& ThreadResult : ThreadResults)
+	for(const TSparseSet<UObject*>& ThreadResult : ThreadResults)
 	{
 		FinalResult.Append(ThreadResult.Array());
 	}
@@ -254,7 +254,7 @@ static TArray<UObject*> GetAllReferencersIncludingWeak(
 }
 
 static TArray<UObject*> GetAllReferencersExcludingWeak(
-	const TSet<UObject*>& Referencees, const TSet<UObject*>* ObjectsToIgnore, const EReferencerFinderFlags Flags)
+	const TSparseSet<UObject*>& Referencees, const TSet<UObject*>* ObjectsToIgnore, const EReferencerFinderFlags Flags)
 {
 	// We can use the faster GC schema to find references when we don't need to include weak 
 	// references:
@@ -270,7 +270,7 @@ static TArray<UObject*> GetAllReferencersExcludingWeak(
 
 	ParallelFor(NumThreads, [&Referencees, ObjectsToIgnore, &ResultCritical, &Ret, NumberOfObjectsPerThread, NumThreads, MaxNumberOfObjects, Flags](int32 ThreadIndex)
 	{
-		TSet<UObject*> ThreadResult;
+		TSparseSet<UObject*> ThreadResult;
 		FAllReferencesProcessor Processor(Referencees, Flags, ThreadResult);
 		TArray<UObject*> ObjectsToSerialize;
 		ObjectsToSerialize.Reserve(NumberOfObjectsPerThread);
@@ -328,16 +328,9 @@ static FAutoConsoleVariableRef CVar_UseSerializeToFindWeakReferencers(
 	GUseSerializeToFindWeakReferencers,
 	TEXT("If true use Serialize routines to find objects that are referencing only via weak references - set to false to only use GC schema for reference finding")
 );
-}// namespace UE::Private
 
-TArray<UObject*> FReferencerFinder::GetAllReferencers(
-	const TArray<UObject*>& Referencees, const TSet<UObject*>* ObjectsToIgnore, EReferencerFinderFlags Flags)
-{
-	return GetAllReferencers(TSet<UObject*>(Referencees), ObjectsToIgnore, Flags);
-}
-
-TArray<UObject*> FReferencerFinder::GetAllReferencers(
-	const TSet<UObject*>& Referencees, const TSet<UObject*>* ObjectsToIgnore, EReferencerFinderFlags Flags)
+TArray<UObject*> GetAllReferencers(
+	const TSparseSet<UObject*>& Referencees, const TSet<UObject*>* ObjectsToIgnore, EReferencerFinderFlags Flags)
 {
 	if(Referencees.Num() == 0)
 	{
@@ -364,4 +357,22 @@ TArray<UObject*> FReferencerFinder::GetAllReferencers(
 			Flags
 		);
 	}
+}
+}// namespace UE::Private
+
+TArray<UObject*> FReferencerFinder::GetAllReferencers(
+	const TArray<UObject*>& Referencees, const TSet<UObject*>* ObjectsToIgnore, EReferencerFinderFlags Flags)
+{
+	return UE::Private::GetAllReferencers(TSparseSet<UObject*>(Referencees), ObjectsToIgnore, Flags);
+}
+
+TArray<UObject*> FReferencerFinder::GetAllReferencers(
+	const TSet<UObject*>& Referencees, const TSet<UObject*>* ObjectsToIgnore, EReferencerFinderFlags Flags)
+{
+	TSparseSet<UObject*> AsSparse;
+	for(UObject* Obj: Referencees)
+	{
+		AsSparse.Add(Obj);
+	}
+	return UE::Private::GetAllReferencers(AsSparse, ObjectsToIgnore, Flags);
 }
\ No newline at end of file
