PackageName: Move FMountPoint and FLongPackagePathsSingleton out into separate files, for readability improvement in upcoming public interface to FMountPoint.
#rnx

#Synced-CL 50109798

diff --git a/Engine/Source/Runtime/CoreUObject/Private/Misc/PackageName.cpp b/Engine/Source/Runtime/CoreUObject/Private/Misc/PackageName.cpp
--- a/Engine/Source/Runtime/CoreUObject/Private/Misc/PackageName.cpp
+++ b/Engine/Source/Runtime/CoreUObject/Private/Misc/PackageName.cpp
@@ -25,6 +25,8 @@
 #include "Misc/AutomationTest.h"
 #include "Misc/CoreDelegates.h"
 #include "Misc/EnumClassFlags.h"
+#include "Misc/PackageName/MountPointPrivate.h"
+#include "Misc/PackageName/PackageNamePrivate.h"
 #include "Misc/PackagePath.h"
 #include "Misc/PackageSegment.h"
 #include "Misc/Paths.h"
@@ -167,530 +169,437 @@ bool FPackageName::TryConvertGameRelativePackagePathToLocalPath(FStringView Rela
 	}
 }
 
-
-struct FMountPoint
+FLongPackagePathsSingleton& FLongPackagePathsSingleton::Get()
 {
-	enum class EMountFlags
-	{
-		Removable		= 0x1,
-		NotRemovable	= 0x0,
-		ReadOnly		= 0x2,
-		NotReadOnly		= 0x0,
-		Alias			= 0x4,
-		NotAlias		= 0x0,
-	};
-
-	/** The LongPackageName path(e.g., "/Engine/") */
-	const FString RootPath;
-
-	/** The LocalPath relative path(e.g., "../../../Engine/Content/") */
-	const FString ContentPathRelative;
-
-	/** The LocalPath absolute path */
-	const FString ContentPathAbsolute;
-
-	/** Whether the mountpoint was created by InsertMountPoint and can be removed by RemoveMountPoint. */
-	const bool bRemovable : 1;
-
-	/** Whether new Packages can be saved into the mountpoint. */
-	const bool bReadOnly : 1;
-
-	/** Whether this mountpoint is the same RootPath as another mountpoint but bound to a different ContentPath. */
-	const bool bAlias : 1;
-
-	FMountPoint(FString&& InRootPath, FString&& InContentPathRelative,
-		FString&& InContentPathAbsolute, EMountFlags Flags);
-};
-ENUM_CLASS_FLAGS(FMountPoint::EMountFlags);
+	static FLongPackagePathsSingleton Singleton;
+	return Singleton;
+}
 
-FMountPoint::FMountPoint(FString&& InRootPath, FString&& InContentPathRelative,
-	FString&& InContentPathAbsolute, EMountFlags Flags)
-	: RootPath(MoveTemp(InRootPath))
-	, ContentPathRelative(MoveTemp(InContentPathRelative))
-	, ContentPathAbsolute(MoveTemp(InContentPathAbsolute))
-	, bRemovable(!!(Flags & EMountFlags::Removable))
-	, bReadOnly(!!(Flags & EMountFlags::ReadOnly))
-	, bAlias(!!(Flags & EMountFlags::Alias))
+void FLongPackagePathsSingleton::OnCoreUObjectInitialized()
 {
+	// Allow the plugin manager to mount new content paths by exposing access through a delegate.  PluginManager is 
+	// a Core class, but content path functionality is added at the CoreUObject level.
+	IPluginManager& PluginManager = IPluginManager::Get();
+	PluginManager.SetRegisterMountPointDelegate(IPluginManager::FRegisterMountPointDelegate::CreateStatic(&FPackageName::RegisterMountPoint));
+	PluginManager.SetUnRegisterMountPointDelegate(IPluginManager::FRegisterMountPointDelegate::CreateStatic(&FPackageName::UnRegisterMountPoint));
+	UE::CoreUObject::Private::PluginHandler::Install();
 }
 
-struct FLongPackagePathsSingleton
+FMountPoint FLongPackagePathsSingleton::ConstructMountPoint(const FString& RootPath, const FString& ContentPath,
+	FMountPoint::EMountFlags Flags)
 {
-	mutable FTransactionallySafeRWLock MountLock;
-
-	FString ConfigRootPath;
-	FString EngineRootPath;
-	FString GameRootPath;
-	FString ScriptRootPath;
-	FString MemoryRootPath;
-	FString TempRootPath;
+	FString MountPath = ContentPath;
 
-	FString VerseSubPath;
-
-	FString EngineContentPath;
-	FString ContentPathShort;
-	FString EngineShadersPath;
-	FString EngineShadersPathShort;
-	FString GameContentPath;
-	FString GameConfigPath;
-	FString GameScriptPath;
-	FString GameSavedPath;
-	FString GameContentPathRebased;
-	FString GameConfigPathRebased;
-	FString GameScriptPathRebased;
-	FString GameSavedPathRebased;
-
-	TDirectoryTree<FMountPoint*> RootPathTree;
-	TDirectoryTree<FMountPoint*> ContentPathTree;
-	TArray<TUniquePtr<FMountPoint>> MountPoints;
-
-	// singleton
-	static FLongPackagePathsSingleton& Get()
+	// If a relative path is passed, convert to an absolute path 
+	if (FPaths::IsRelative(MountPath))
 	{
-		static FLongPackagePathsSingleton Singleton;
-		return Singleton;
-	}
-
-	/** Initialization function to setup content paths that cannot be run until CoreUObject/PluginManager have been initialized */
-	void OnCoreUObjectInitialized()
-	{
-		// Allow the plugin manager to mount new content paths by exposing access through a delegate.  PluginManager is 
-		// a Core class, but content path functionality is added at the CoreUObject level.
-		IPluginManager& PluginManager = IPluginManager::Get();
-		PluginManager.SetRegisterMountPointDelegate(IPluginManager::FRegisterMountPointDelegate::CreateStatic(&FPackageName::RegisterMountPoint));
-		PluginManager.SetUnRegisterMountPointDelegate(IPluginManager::FRegisterMountPointDelegate::CreateStatic(&FPackageName::UnRegisterMountPoint));
-		UE::CoreUObject::Private::PluginHandler::Install();
-	}
+		MountPath = FPaths::ConvertRelativePathToFull(ContentPath);
 
-	/**
-	 * Given a content path return the consistently-formatted RelativePath and AbsolutePath, suitable for lookup based only on text
-	 * The RelativePath is consistent with FileManager relative paths.
-	 */
-	static FMountPoint ConstructMountPoint(const FString& RootPath, const FString& ContentPath,
-		FMountPoint::EMountFlags Flags)
-	{
-		FString MountPath = ContentPath;
-
-		// If a relative path is passed, convert to an absolute path 
-		if (FPaths::IsRelative(MountPath))
+		// Revert to original path if unable to convert to full path
+		if (MountPath.Len() <= 1)
 		{
-			MountPath = FPaths::ConvertRelativePathToFull(ContentPath);
-
-			// Revert to original path if unable to convert to full path
-			if (MountPath.Len() <= 1)
-			{
-				MountPath = ContentPath;
-				UE_LOG(LogPackageName, Warning, TEXT("Unable to convert mount point relative path: %s"), *ContentPath);
-			}
+			MountPath = ContentPath;
+			UE_LOG(LogPackageName, Warning, TEXT("Unable to convert mount point relative path: %s"), *ContentPath);
 		}
+	}
 
-		// Convert to a relative path using the FileManager
-		FString RelativeContentPath = IFileManager::Get().ConvertToRelativePath(*MountPath);
+	// Convert to a relative path using the FileManager
+	FString RelativeContentPath = IFileManager::Get().ConvertToRelativePath(*MountPath);
 
-		// Make sure the path ends in a trailing path separator for consistency.
-		if (!RelativeContentPath.EndsWith(TEXT("/"), ESearchCase::CaseSensitive))
-		{
-			RelativeContentPath += TEXT("/");
-		}
+	// Make sure the path ends in a trailing path separator for consistency.
+	if (!RelativeContentPath.EndsWith(TEXT("/"), ESearchCase::CaseSensitive))
+	{
+		RelativeContentPath += TEXT("/");
+	}
 
-		TStringBuilder<256> AbsolutePathBuilder;
-		FPathViews::ToAbsolutePath(RelativeContentPath, AbsolutePathBuilder);
+	TStringBuilder<256> AbsolutePathBuilder;
+	FPathViews::ToAbsolutePath(RelativeContentPath, AbsolutePathBuilder);
 
-		return FMountPoint(FString(RootPath), MoveTemp(RelativeContentPath), FString(AbsolutePathBuilder), Flags);
-	}
+	return FMountPoint(FString(RootPath), MoveTemp(RelativeContentPath), FString(AbsolutePathBuilder), Flags);
+}
 
-	void InsertMountPoint(const FString& RootPath, const FString& ContentPath)
-	{	
-		TUniquePtr<FMountPoint> MountPointOwner(MakeUnique<FMountPoint>(
-			ConstructMountPoint(RootPath, ContentPath, FMountPoint::EMountFlags::Removable)));
-		FMountPoint* MountPoint = MountPointOwner.Get();
-		{
-			UE::TWriteScopeLock ScopeLock(MountLock);
-			MountPoints.Insert(MoveTemp(MountPointOwner), 0);
+void FLongPackagePathsSingleton::InsertMountPoint(const FString& RootPath, const FString& ContentPath)
+{	
+	TUniquePtr<FMountPoint> MountPointOwner(MakeUnique<FMountPoint>(
+		ConstructMountPoint(RootPath, ContentPath, FMountPoint::EMountFlags::Removable)));
+	FMountPoint* MountPoint = MountPointOwner.Get();
+	{
+		UE::TWriteScopeLock ScopeLock(MountLock);
+		MountPoints.Insert(MoveTemp(MountPointOwner), 0);
 
-			RootPathTree.FindOrAdd(MountPoint->RootPath) = MountPoint;
-			ContentPathTree.FindOrAdd(MountPoint->ContentPathRelative) = MountPoint;
-			ContentPathTree.FindOrAdd(MountPoint->ContentPathAbsolute) = MountPoint;
-		}		
+		RootPathTree.FindOrAdd(MountPoint->RootPath) = MountPoint;
+		ContentPathTree.FindOrAdd(MountPoint->ContentPathRelative) = MountPoint;
+		ContentPathTree.FindOrAdd(MountPoint->ContentPathAbsolute) = MountPoint;
+	}		
 
-		// Downgrade the log messages on startup to Verbose to reduce startup spam
+	// Downgrade the log messages on startup to Verbose to reduce startup spam
 #if !NO_LOGGING
-		FString Message = FString::Printf(TEXT("FPackageName: Mount point added: '%s' mounted to '%s'"),
-			*MountPoint->ContentPathRelative, *MountPoint->RootPath);
-		if (IsEngineStartupModuleLoadingComplete())
-		{
-			UE_LOG(LogPackageName, Log, TEXT("%s"), *Message);
-		}
-		else
-		{
-			UE_LOG(LogPackageName, Verbose, TEXT("%s"), *Message);
-		}
-#endif
-
-		// Let subscribers know that a new content path was mounted
-		FPackageName::OnContentPathMounted().Broadcast( RootPath, MountPoint->ContentPathRelative);
+	FString Message = FString::Printf(TEXT("FPackageName: Mount point added: '%s' mounted to '%s'"),
+		*MountPoint->ContentPathRelative, *MountPoint->RootPath);
+	if (IsEngineStartupModuleLoadingComplete())
+	{
+		UE_LOG(LogPackageName, Log, TEXT("%s"), *Message);
 	}
-
-	// This will remove a previously inserted mount point
-	void RemoveMountPoint(const FString& RootPath, const FString& ContentPath)
+	else
 	{
-		RemoveMountPoint(ConstructMountPoint(RootPath, ContentPath, FMountPoint::EMountFlags::Removable));
+		UE_LOG(LogPackageName, Verbose, TEXT("%s"), *Message);
 	}
+#endif
+
+	// Let subscribers know that a new content path was mounted
+	FPackageName::OnContentPathMounted().Broadcast( RootPath, MountPoint->ContentPathRelative);
+}
+
+void FLongPackagePathsSingleton::RemoveMountPoint(const FString& RootPath, const FString& ContentPath)
+{
+	RemoveMountPoint(ConstructMountPoint(RootPath, ContentPath, FMountPoint::EMountFlags::Removable));
+}
 
-	void RemoveMountPoint(const FMountPoint& MountPoint)
+void FLongPackagePathsSingleton::RemoveMountPoint(const FMountPoint& MountPoint)
+{
+	bool bRemoved = false;
 	{
-		bool bRemoved = false;
+		UE::TWriteScopeLock ScopeLock(MountLock);
+		for (int32 MountPointIndex = 0; MountPointIndex < MountPoints.Num(); /* Conditionally incremented in loop */)
 		{
-			UE::TWriteScopeLock ScopeLock(MountLock);
-			for (int32 MountPointIndex = 0; MountPointIndex < MountPoints.Num(); /* Conditionally incremented in loop */)
+			FMountPoint& ExistingMount = *MountPoints[MountPointIndex];
+			if (ExistingMount.RootPath == MountPoint.RootPath &&
+				ExistingMount.ContentPathRelative == MountPoint.ContentPathRelative &&
+				ExistingMount.bRemovable)
 			{
-				FMountPoint& ExistingMount = *MountPoints[MountPointIndex];
-				if (ExistingMount.RootPath == MountPoint.RootPath &&
-					ExistingMount.ContentPathRelative == MountPoint.ContentPathRelative &&
-					ExistingMount.bRemovable)
+				bool bRootPathExisted;
+				RootPathTree.Remove(ExistingMount.RootPath, &bRootPathExisted);
+				if (bRootPathExisted)
 				{
-					bool bRootPathExisted;
-					RootPathTree.Remove(ExistingMount.RootPath, &bRootPathExisted);
-					if (bRootPathExisted)
+					// See if there is a lower priority path that replaces it
+					for (TUniquePtr<FMountPoint>& OtherMount : TArrayView<TUniquePtr<FMountPoint>>(MountPoints).RightChop(MountPointIndex + 1))
 					{
-						// See if there is a lower priority path that replaces it
-						for (TUniquePtr<FMountPoint>& OtherMount : TArrayView<TUniquePtr<FMountPoint>>(MountPoints).RightChop(MountPointIndex + 1))
+						if (OtherMount->RootPath == MountPoint.RootPath)
 						{
-							if (OtherMount->RootPath == MountPoint.RootPath)
-							{
-								RootPathTree.FindOrAdd(OtherMount->RootPath) = OtherMount.Get();
-								break;
-							}
+							RootPathTree.FindOrAdd(OtherMount->RootPath) = OtherMount.Get();
+							break;
 						}
 					}
-					bool bContentPathRelativeExisted;
-					ContentPathTree.Remove(ExistingMount.ContentPathRelative, &bContentPathRelativeExisted);
-					if (bContentPathRelativeExisted)
+				}
+				bool bContentPathRelativeExisted;
+				ContentPathTree.Remove(ExistingMount.ContentPathRelative, &bContentPathRelativeExisted);
+				if (bContentPathRelativeExisted)
+				{
+					// See if there is a lower priority path that replaces it
+					for (TUniquePtr<FMountPoint>& OtherMount :
+						TArrayView<TUniquePtr<FMountPoint>>(MountPoints).RightChop(MountPointIndex+1))
 					{
-						// See if there is a lower priority path that replaces it
-						for (TUniquePtr<FMountPoint>& OtherMount :
-							TArrayView<TUniquePtr<FMountPoint>>(MountPoints).RightChop(MountPointIndex+1))
+						if (OtherMount->ContentPathRelative == MountPoint.ContentPathRelative)
 						{
-							if (OtherMount->ContentPathRelative == MountPoint.ContentPathRelative)
-							{
-								ContentPathTree.FindOrAdd(OtherMount->ContentPathRelative) = OtherMount.Get();
-								break;
-							}
+							ContentPathTree.FindOrAdd(OtherMount->ContentPathRelative) = OtherMount.Get();
+							break;
 						}
 					}
-					bool bContentPathAbsoluteExisted;
-					ContentPathTree.Remove(ExistingMount.ContentPathAbsolute, &bContentPathAbsoluteExisted);
-					if (bContentPathAbsoluteExisted)
+				}
+				bool bContentPathAbsoluteExisted;
+				ContentPathTree.Remove(ExistingMount.ContentPathAbsolute, &bContentPathAbsoluteExisted);
+				if (bContentPathAbsoluteExisted)
+				{
+					// See if there is a lower priority path that replaces it
+					for (TUniquePtr<FMountPoint>& OtherMount :
+						TArrayView<TUniquePtr<FMountPoint>>(MountPoints).RightChop(MountPointIndex + 1))
 					{
-						// See if there is a lower priority path that replaces it
-						for (TUniquePtr<FMountPoint>& OtherMount :
-							TArrayView<TUniquePtr<FMountPoint>>(MountPoints).RightChop(MountPointIndex + 1))
+						if (OtherMount->ContentPathAbsolute == MountPoint.ContentPathAbsolute)
 						{
-							if (OtherMount->ContentPathAbsolute == MountPoint.ContentPathAbsolute)
-							{
-								ContentPathTree.FindOrAdd(OtherMount->ContentPathAbsolute) = OtherMount.Get();
-								break;
-							}
+							ContentPathTree.FindOrAdd(OtherMount->ContentPathAbsolute) = OtherMount.Get();
+							break;
 						}
 					}
-
-					MountPoints.RemoveAt(MountPointIndex);
-					bRemoved = true;
 				}
-				else
-				{
-					++MountPointIndex;
-				}
-			}
-			if (bRemoved)
-			{
-				UE_LOG(LogPackageName, Display, TEXT("FPackageName: Mount point removed: '%s' unmounted from '%s'"),
-					*MountPoint.ContentPathRelative, *MountPoint.RootPath);
+
+				MountPoints.RemoveAt(MountPointIndex);
+				bRemoved = true;
 			}
 			else
 			{
-				UE_LOG(LogPackageName, Display, TEXT("FPackageName: Mount point remove failed: no MountPoint found mapping '%s' to '%s'"),
-					*MountPoint.ContentPathRelative, *MountPoint.RootPath);
+				++MountPointIndex;
 			}
 		}
-
-		// Let subscribers know that a new content path was unmounted
 		if (bRemoved)
 		{
-			FPackageName::OnContentPathDismounted().Broadcast(MountPoint.RootPath, MountPoint.ContentPathRelative);
+			UE_LOG(LogPackageName, Display, TEXT("FPackageName: Mount point removed: '%s' unmounted from '%s'"),
+				*MountPoint.ContentPathRelative, *MountPoint.RootPath);
+		}
+		else
+		{
+			UE_LOG(LogPackageName, Display, TEXT("FPackageName: Mount point remove failed: no MountPoint found mapping '%s' to '%s'"),
+				*MountPoint.ContentPathRelative, *MountPoint.RootPath);
 		}
 	}
 
-	// Checks whether the specific root path is a valid mount point.
-	bool MountPointExists(FStringView RootPath)
+	// Let subscribers know that a new content path was unmounted
+	if (bRemoved)
 	{
-		UE::TReadScopeLock ScopeLock(MountLock);
-		return RootPathTree.Contains(RootPath) || RootPath == MemoryRootPath;
+		FPackageName::OnContentPathDismounted().Broadcast(MountPoint.RootPath, MountPoint.ContentPathRelative);
 	}
+}
 
-private:
-#if !UE_BUILD_SHIPPING
-	const FAutoConsoleCommand DumpMountPointsCommand;
-	const FAutoConsoleCommand RegisterMountPointCommand;
-	const FAutoConsoleCommand UnregisterMountPointCommand;
-#endif
-	
-	FLongPackagePathsSingleton()
+bool FLongPackagePathsSingleton::MountPointExists(FStringView RootPath)
+{
+	UE::TReadScopeLock ScopeLock(MountLock);
+	return RootPathTree.Contains(RootPath) || RootPath == MemoryRootPath;
+}
+
+FLongPackagePathsSingleton::FLongPackagePathsSingleton()
 #if !UE_BUILD_SHIPPING
 	: DumpMountPointsCommand(
-	TEXT( "PackageName.DumpMountPoints" ),
-	*LOCTEXT("CommandText_DumpMountPoints", "Print registered LongPackagePath mount points").ToString(),
-		FConsoleCommandWithArgsDelegate::CreateRaw( this, &FLongPackagePathsSingleton::ExecDumpMountPoints ) )
+		TEXT("PackageName.DumpMountPoints"),
+		*LOCTEXT("CommandText_DumpMountPoints", "Print registered LongPackagePath mount points").ToString(),
+		FConsoleCommandWithArgsDelegate::CreateRaw(this, &FLongPackagePathsSingleton::ExecDumpMountPoints))
 	, RegisterMountPointCommand(
-	TEXT( "PackageName.RegisterMountPoint" ),
-	*LOCTEXT("CommandText_RegisterMountPoint", "<RootPath> <ContentPath> // Register a LongPackagePath mount point").ToString(),
-		FConsoleCommandWithArgsDelegate::CreateRaw( this, &FLongPackagePathsSingleton::ExecInsertMountPoint ) )
+		TEXT("PackageName.RegisterMountPoint"),
+		*LOCTEXT("CommandText_RegisterMountPoint", "<RootPath> <ContentPath> // Register a LongPackagePath mount point").ToString(),
+		FConsoleCommandWithArgsDelegate::CreateRaw(this, &FLongPackagePathsSingleton::ExecInsertMountPoint))
 	, UnregisterMountPointCommand(
-	TEXT( "PackageName.UnregisterMountPoint" ),
-	*LOCTEXT("CommandText_UnregisterMountPoint", "<RootPath> <ContentPath> // Remove a LongPackagePath mount point").ToString(),
-		FConsoleCommandWithArgsDelegate::CreateRaw( this, &FLongPackagePathsSingleton::ExecRemoveMountPoint ) )
+		TEXT("PackageName.UnregisterMountPoint"),
+		*LOCTEXT("CommandText_UnregisterMountPoint", "<RootPath> <ContentPath> // Remove a LongPackagePath mount point").ToString(),
+		FConsoleCommandWithArgsDelegate::CreateRaw(this, &FLongPackagePathsSingleton::ExecRemoveMountPoint))
 #endif
+{
+	SCOPED_BOOT_TIMING("FPackageName::FLongPackagePathsSingleton");
+
+	ConfigRootPath = TEXT("/Config/");
+	EngineRootPath = TEXT("/Engine/");
+	GameRootPath   = TEXT("/Game/");
+	ScriptRootPath = TEXT("/Script/");
+	MemoryRootPath = TEXT("/Memory/");
+	TempRootPath   = TEXT("/Temp/");
+
+	VerseSubPath = TEXT("/_Verse");
+
+	EngineContentPath      = FPaths::EngineContentDir();
+	ContentPathShort       = TEXT("../../Content/");
+	EngineShadersPath      = FPaths::EngineDir() / TEXT("Shaders/");
+	EngineShadersPathShort = TEXT("../../Shaders/");
+	GameContentPath        = FPaths::ProjectContentDir();
+	GameConfigPath         = FPaths::ProjectConfigDir();
+	GameScriptPath         = FPaths::ProjectDir() / TEXT("Script/");
+	GameSavedPath          = FPaths::ProjectSavedDir();
+
+	FString RebasedGameDir = FString::Printf(TEXT("../../../%s/"), FApp::GetProjectName());
+
+	GameContentPathRebased = RebasedGameDir / TEXT("Content/");
+	GameConfigPathRebased  = RebasedGameDir / TEXT("Config/");
+	GameScriptPathRebased  = RebasedGameDir / TEXT("Script/");
+	GameSavedPathRebased   = RebasedGameDir / TEXT("Saved/");
+	using EMountFlags = FMountPoint::EMountFlags;
+
+	auto LocalInsertMountPoint = [this](const FString& RootPath, const FString& ContentPathRelative, EMountFlags InFlags)
 	{
-		SCOPED_BOOT_TIMING("FPackageName::FLongPackagePathsSingleton");
-
-		ConfigRootPath = TEXT("/Config/");
-		EngineRootPath = TEXT("/Engine/");
-		GameRootPath   = TEXT("/Game/");
-		ScriptRootPath = TEXT("/Script/");
-		MemoryRootPath = TEXT("/Memory/");
-		TempRootPath   = TEXT("/Temp/");
-
-		VerseSubPath = TEXT("/_Verse");
-
-		EngineContentPath      = FPaths::EngineContentDir();
-		ContentPathShort       = TEXT("../../Content/");
-		EngineShadersPath      = FPaths::EngineDir() / TEXT("Shaders/");
-		EngineShadersPathShort = TEXT("../../Shaders/");
-		GameContentPath        = FPaths::ProjectContentDir();
-		GameConfigPath         = FPaths::ProjectConfigDir();
-		GameScriptPath         = FPaths::ProjectDir() / TEXT("Script/");
-		GameSavedPath          = FPaths::ProjectSavedDir();
-
-		FString RebasedGameDir = FString::Printf(TEXT("../../../%s/"), FApp::GetProjectName());
-
-		GameContentPathRebased = RebasedGameDir / TEXT("Content/");
-		GameConfigPathRebased  = RebasedGameDir / TEXT("Config/");
-		GameScriptPathRebased  = RebasedGameDir / TEXT("Script/");
-		GameSavedPathRebased   = RebasedGameDir / TEXT("Saved/");
-		using EMountFlags = FMountPoint::EMountFlags;
-
-		auto LocalInsertMountPoint = [this](const FString& RootPath, const FString& ContentPathRelative, EMountFlags InFlags)
-		{
-			// Don't call ConstructMountPoint, because that calls FPaths::ConvertRelativePathToFull followed by
-			// IFileManager::Get().ConvertToRelativePath, and that combination will convert ../../../Engine/Content
-			// into ../../Content, which we don't want here; we want to enter in the full ../../.. path for some entries.
-			TStringBuilder<256> AbsolutePathBuilder;
-			FPathViews::ToAbsolutePath(ContentPathRelative, AbsolutePathBuilder);
-			FMountPoint& MountPoint = *MountPoints.Add_GetRef(MakeUnique<FMountPoint>(
-				FString(RootPath), FString(ContentPathRelative), FString(AbsolutePathBuilder),
-				InFlags));
-
-			// We are appending to the MountPoint list, and will have only a few duplicates. If any duplicates do occur,
-			// keep the MountPoint that was appended earlier as the registered mountpoint.
-			FMountPoint*& ExistingRootPath = RootPathTree.FindOrAdd(MountPoint.RootPath);
-			if (!ExistingRootPath)
-			{
-				ExistingRootPath = &MountPoint;
-			}
-			FMountPoint*& ExistingContentPathRelative = ContentPathTree.FindOrAdd(MountPoint.ContentPathRelative);
-			if (!ExistingContentPathRelative)
-			{
-				ExistingContentPathRelative = &MountPoint;
-			}
-			FMountPoint*& ExistingContentPathAbsolute = ContentPathTree.FindOrAdd(MountPoint.ContentPathAbsolute);
-			if (!ExistingContentPathAbsolute)
-			{
-				ExistingContentPathAbsolute = &MountPoint;
-			}
-		};
-		
-		LocalInsertMountPoint(EngineRootPath, EngineContentPath, EMountFlags::NotReadOnly);
-		LocalInsertMountPoint(EngineRootPath, EngineShadersPath, EMountFlags::NotReadOnly | EMountFlags::Alias);
-		LocalInsertMountPoint(GameRootPath,   GameContentPath, EMountFlags::NotReadOnly);
-		LocalInsertMountPoint(ScriptRootPath, GameScriptPath, EMountFlags::ReadOnly);
-		LocalInsertMountPoint(TempRootPath,   GameSavedPath, EMountFlags::ReadOnly);
-		LocalInsertMountPoint(ConfigRootPath, GameConfigPath, EMountFlags::ReadOnly);
-
-		// Add other LocalPaths that have different sets of .. but are the same location on disk
-		LocalInsertMountPoint(EngineRootPath, EngineShadersPathShort, EMountFlags::NotReadOnly | EMountFlags::Alias);
-		if (FPaths::IsSamePath(GameContentPath, ContentPathShort))
-		{
-			// ../../Content points to the the Engine directory, the same as ../../../Engine/Content
-			// But if the /Game is pointing to ../../../Engine/Content as well, then map ../../Content to
-			// /Game instead of /Engine
-			LocalInsertMountPoint(GameRootPath, ContentPathShort, EMountFlags::NotReadOnly | EMountFlags::Alias);
+		// Don't call ConstructMountPoint, because that calls FPaths::ConvertRelativePathToFull followed by
+		// IFileManager::Get().ConvertToRelativePath, and that combination will convert ../../../Engine/Content
+		// into ../../Content, which we don't want here; we want to enter in the full ../../.. path for some entries.
+		TStringBuilder<256> AbsolutePathBuilder;
+		FPathViews::ToAbsolutePath(ContentPathRelative, AbsolutePathBuilder);
+		FMountPoint& MountPoint = *MountPoints.Add_GetRef(MakeUnique<FMountPoint>(
+			FString(RootPath), FString(ContentPathRelative), FString(AbsolutePathBuilder),
+			InFlags));
+
+		// We are appending to the MountPoint list, and will have only a few duplicates. If any duplicates do occur,
+		// keep the MountPoint that was appended earlier as the registered mountpoint.
+		FMountPoint*& ExistingRootPath = RootPathTree.FindOrAdd(MountPoint.RootPath);
+		if (!ExistingRootPath)
+		{
+			ExistingRootPath = &MountPoint;
 		}
-		else
+		FMountPoint*& ExistingContentPathRelative = ContentPathTree.FindOrAdd(MountPoint.ContentPathRelative);
+		if (!ExistingContentPathRelative)
+		{
+			ExistingContentPathRelative = &MountPoint;
+		}
+		FMountPoint*& ExistingContentPathAbsolute = ContentPathTree.FindOrAdd(MountPoint.ContentPathAbsolute);
+		if (!ExistingContentPathAbsolute)
 		{
-			LocalInsertMountPoint(EngineRootPath, ContentPathShort, EMountFlags::NotReadOnly | EMountFlags::Alias);
+			ExistingContentPathAbsolute = &MountPoint;
 		}
-		LocalInsertMountPoint(GameRootPath,   GameContentPathRebased, EMountFlags::NotReadOnly | EMountFlags::Alias);
-		LocalInsertMountPoint(ScriptRootPath, GameScriptPathRebased, EMountFlags::ReadOnly | EMountFlags::Alias);
-		LocalInsertMountPoint(TempRootPath,   GameSavedPathRebased, EMountFlags::ReadOnly | EMountFlags::Alias);
-		LocalInsertMountPoint(ConfigRootPath, GameConfigPathRebased, EMountFlags::ReadOnly | EMountFlags::Alias);
+	};
+		
+	LocalInsertMountPoint(EngineRootPath, EngineContentPath, EMountFlags::NotReadOnly);
+	LocalInsertMountPoint(EngineRootPath, EngineShadersPath, EMountFlags::NotReadOnly | EMountFlags::Alias);
+	LocalInsertMountPoint(GameRootPath,   GameContentPath, EMountFlags::NotReadOnly);
+	LocalInsertMountPoint(ScriptRootPath, GameScriptPath, EMountFlags::ReadOnly);
+	LocalInsertMountPoint(TempRootPath,   GameSavedPath, EMountFlags::ReadOnly);
+	LocalInsertMountPoint(ConfigRootPath, GameConfigPath, EMountFlags::ReadOnly);
+
+	// Add other LocalPaths that have different sets of .. but are the same location on disk
+	LocalInsertMountPoint(EngineRootPath, EngineShadersPathShort, EMountFlags::NotReadOnly | EMountFlags::Alias);
+	if (FPaths::IsSamePath(GameContentPath, ContentPathShort))
+	{
+		// ../../Content points to the the Engine directory, the same as ../../../Engine/Content
+		// But if the /Game is pointing to ../../../Engine/Content as well, then map ../../Content to
+		// /Game instead of /Engine
+		LocalInsertMountPoint(GameRootPath, ContentPathShort, EMountFlags::NotReadOnly | EMountFlags::Alias);
+	}
+	else
+	{
+		LocalInsertMountPoint(EngineRootPath, ContentPathShort, EMountFlags::NotReadOnly | EMountFlags::Alias);
 	}
+	LocalInsertMountPoint(GameRootPath,   GameContentPathRebased, EMountFlags::NotReadOnly | EMountFlags::Alias);
+	LocalInsertMountPoint(ScriptRootPath, GameScriptPathRebased, EMountFlags::ReadOnly | EMountFlags::Alias);
+	LocalInsertMountPoint(TempRootPath,   GameSavedPathRebased, EMountFlags::ReadOnly | EMountFlags::Alias);
+	LocalInsertMountPoint(ConfigRootPath, GameConfigPathRebased, EMountFlags::ReadOnly | EMountFlags::Alias);
+}
 
 #if !UE_BUILD_SHIPPING
-	void ExecDumpMountPoints(const TArray<FString>& Args)
+void FLongPackagePathsSingleton::ExecDumpMountPoints(const TArray<FString>& Args)
+{
+	UE_LOG(LogPackageName, Log, TEXT("Valid mount points:"));
+
+	UE::TReadScopeLock ScopeLock(MountLock);
+	for (const TUniquePtr<FMountPoint>& MountPoint : MountPoints)
 	{
-		UE_LOG(LogPackageName, Log, TEXT("Valid mount points:"));
+		UE_LOG(LogPackageName, Log, TEXT("	'%s' -> '%s'"), *MountPoint->RootPath, *MountPoint->ContentPathRelative);
+	}
 
-		UE::TReadScopeLock ScopeLock(MountLock);
-		for (const TUniquePtr<FMountPoint>& MountPoint : MountPoints)
+	UE_LOG(LogPackageName, Log, TEXT("Removable mount points:"));
+	for (const TUniquePtr<FMountPoint>& MountPoint : MountPoints)
+	{
+		if (MountPoint->bRemovable)
 		{
-			UE_LOG(LogPackageName, Log, TEXT("	'%s' -> '%s'"), *MountPoint->RootPath, *MountPoint->ContentPathRelative);
-		}
-
-		UE_LOG(LogPackageName, Log, TEXT("Removable mount points:"));
-		for (const TUniquePtr<FMountPoint>& MountPoint : MountPoints)
-		{
-			if (MountPoint->bRemovable)
-			{
-				UE_LOG(LogPackageName, Log, TEXT("	'%s'"), *MountPoint->RootPath);
-			}
+			UE_LOG(LogPackageName, Log, TEXT("	'%s'"), *MountPoint->RootPath);
 		}
 	}
+}
 
-	void ExecInsertMountPoint(const TArray<FString>& Args)
+void FLongPackagePathsSingleton::ExecInsertMountPoint(const TArray<FString>& Args)
+{
+	if ( Args.Num() < 2 )
 	{
-		if ( Args.Num() < 2 )
-		{
-			UE_LOG(LogPackageName, Log, TEXT("Usage: PackageName.RegisterMountPoint <RootPath> <ContentPath>"));
-			UE_LOG(LogPackageName, Log, TEXT("Example ContentPath: '../../../ProjectName/Content/'"));
-			UE_LOG(LogPackageName, Log, TEXT("Example RootPath: '/Game/'"));
-			return;
-		}
+		UE_LOG(LogPackageName, Log, TEXT("Usage: PackageName.RegisterMountPoint <RootPath> <ContentPath>"));
+		UE_LOG(LogPackageName, Log, TEXT("Example ContentPath: '../../../ProjectName/Content/'"));
+		UE_LOG(LogPackageName, Log, TEXT("Example RootPath: '/Game/'"));
+		return;
+	}
 
-		const FString& RootPath = Args[0];
-		const FString& ContentPath = Args[1];
+	const FString& RootPath = Args[0];
+	const FString& ContentPath = Args[1];
 
-		if (ContentPath[0] == TEXT('/'))
-		{
-			UE_LOG(LogPackageName, Error, TEXT("PackageName.RegisterMountPoint: Invalid ContentPath, should not start with '/'! Example: '../../../ProjectName/Content/'"));
-			return;
-		}
+	if (ContentPath[0] == TEXT('/'))
+	{
+		UE_LOG(LogPackageName, Error, TEXT("PackageName.RegisterMountPoint: Invalid ContentPath, should not start with '/'! Example: '../../../ProjectName/Content/'"));
+		return;
+	}
 
-		if (RootPath[0] != TEXT('/'))
-		{
-			UE_LOG(LogPackageName, Error, TEXT("PackageName.RegisterMountPoint: Invalid RootPath, should start with a '/'! Example: '/Game/'"));
-			return;
-		}
-		
-		Get().InsertMountPoint(RootPath, ContentPath);
+	if (RootPath[0] != TEXT('/'))
+	{
+		UE_LOG(LogPackageName, Error, TEXT("PackageName.RegisterMountPoint: Invalid RootPath, should start with a '/'! Example: '/Game/'"));
+		return;
 	}
+		
+	Get().InsertMountPoint(RootPath, ContentPath);
+}
 
-	void ExecRemoveMountPoint(const TArray<FString>& Args)
+void FLongPackagePathsSingleton::ExecRemoveMountPoint(const TArray<FString>& Args)
+{
+	if ( Args.Num() < 1 || Args.Num() > 2 )
 	{
-		if ( Args.Num() < 1 || Args.Num() > 2 )
-		{
-			UE_LOG(LogPackageName, Log, TEXT("Usage: PackageName.UnregisterMountPoint <Path>"));
-			UE_LOG(LogPackageName, Log, TEXT("Removes either a Root or Content path if the path is unambiguous"));
-			UE_LOG(LogPackageName, Log, TEXT("Usage: PackageName.UnregisterMountPoint <RootPath> <ContentPath>"));
-			UE_LOG(LogPackageName, Log, TEXT("Removes a specific Root path to Content path mapping"));
-			UE_LOG(LogPackageName, Log, TEXT("Example ContentPath: '../../../ProjectName/Content/'"));
-			UE_LOG(LogPackageName, Log, TEXT("Example RootPath: '/Game/'"));
-			return;
-		}
+		UE_LOG(LogPackageName, Log, TEXT("Usage: PackageName.UnregisterMountPoint <Path>"));
+		UE_LOG(LogPackageName, Log, TEXT("Removes either a Root or Content path if the path is unambiguous"));
+		UE_LOG(LogPackageName, Log, TEXT("Usage: PackageName.UnregisterMountPoint <RootPath> <ContentPath>"));
+		UE_LOG(LogPackageName, Log, TEXT("Removes a specific Root path to Content path mapping"));
+		UE_LOG(LogPackageName, Log, TEXT("Example ContentPath: '../../../ProjectName/Content/'"));
+		UE_LOG(LogPackageName, Log, TEXT("Example RootPath: '/Game/'"));
+		return;
+	}
 
-		if (Args.Num() == 1)
-		{
-			const FString& Path = Args[0];
-			const bool IsRootPath = Path[0] == TEXT('/');
+	if (Args.Num() == 1)
+	{
+		const FString& Path = Args[0];
+		const bool IsRootPath = Path[0] == TEXT('/');
 
-			FString RootPath;
-			FString ContentPath;
+		FString RootPath;
+		FString ContentPath;
 			
-			if (IsRootPath)
-			{
-				UE::TReadScopeLock ScopeLock(MountLock);
+		if (IsRootPath)
+		{
+			UE::TReadScopeLock ScopeLock(MountLock);
 				
-				TArray<FString> ContentPathRelatives;
-				for (TUniquePtr<FMountPoint>& ExistingMount : Get().MountPoints)
-				{
-					if (ExistingMount->RootPath == RootPath)
-					{
-						ContentPathRelatives.Add(ExistingMount->ContentPathRelative);
-					}
-				}
-				if (ContentPathRelatives.Num() == 0)
-				{
-					UE_LOG(LogPackageName, Error, TEXT("PackageName.UnregisterMountPoint: Root path '%s' is not mounted!"), *RootPath);
-					return;
-				}
-
-				if (ContentPathRelatives.Num() > 1)
+			TArray<FString> ContentPathRelatives;
+			for (TUniquePtr<FMountPoint>& ExistingMount : Get().MountPoints)
+			{
+				if (ExistingMount->RootPath == RootPath)
 				{
-					UE_LOG(LogPackageName, Error, TEXT("PackageName.UnregisterMountPoint: Root path '%s' is mounted to multiple content paths, specify content path to unmount explicitly!"), *RootPath);
-					for (const FString& ContentPathRelative : ContentPathRelatives)
-					{
-						UE_LOG(LogPackageName, Error, TEXT("- %s"), *ContentPathRelative);
-					}
-					return;
+					ContentPathRelatives.Add(ExistingMount->ContentPathRelative);
 				}
-
-				ContentPath = ContentPathRelatives[0];
 			}
-			else
+			if (ContentPathRelatives.Num() == 0)
 			{
-				UE::TReadScopeLock ScopeLock(MountLock);
-			
-				ContentPath = Path;
-				TArray<FString> RootPaths;
-				for (TUniquePtr<FMountPoint>& ExistingMount : Get().MountPoints)
-				{
-					if (ExistingMount->ContentPathRelative == ContentPath)
-					{
-						RootPaths.Add(ExistingMount->RootPath);
-					}
-				}
-
-				if (RootPaths.Num() == 0)
-				{
-					UE_LOG(LogPackageName, Error, TEXT("PackageName.UnregisterMountPoint: Content path '%s' is not mounted!"), *ContentPath);
-					return;
-				}
+				UE_LOG(LogPackageName, Error, TEXT("PackageName.UnregisterMountPoint: Root path '%s' is not mounted!"), *RootPath);
+				return;
+			}
 
-				if (RootPaths.Num() > 1)
+			if (ContentPathRelatives.Num() > 1)
+			{
+				UE_LOG(LogPackageName, Error, TEXT("PackageName.UnregisterMountPoint: Root path '%s' is mounted to multiple content paths, specify content path to unmount explicitly!"), *RootPath);
+				for (const FString& ContentPathRelative : ContentPathRelatives)
 				{
-					UE_LOG(LogPackageName, Error, TEXT("PackageName.UnregisterMountPoint: Content path '%s' is mounted to multiple root paths, specify root path to unmount explicitly!"), *ContentPath);
-					for (const FString& MountRootPath: RootPaths)
-					{
-						UE_LOG(LogPackageName, Error, TEXT("- %s"), *MountRootPath);
-					}
-					return;
+					UE_LOG(LogPackageName, Error, TEXT("- %s"), *ContentPathRelative);
 				}
-
-				RootPath = RootPaths[0];
+				return;
 			}
-			
-			Get().RemoveMountPoint(RootPath, ContentPath);
+
+			ContentPath = ContentPathRelatives[0];
 		}
 		else
 		{
-			const FString& RootPath = Args[0];
-			const FString& ContentPath = Args[1];
-			if (ContentPath[0] == TEXT('/'))
+			UE::TReadScopeLock ScopeLock(MountLock);
+			
+			ContentPath = Path;
+			TArray<FString> RootPaths;
+			for (TUniquePtr<FMountPoint>& ExistingMount : Get().MountPoints)
 			{
-				UE_LOG(LogPackageName, Error, TEXT("PackageName.UnregisterMountPoint: Invalid ContentPath, should not start with '/'! Example: '../../../ProjectName/Content/'"));
+				if (ExistingMount->ContentPathRelative == ContentPath)
+				{
+					RootPaths.Add(ExistingMount->RootPath);
+				}
+			}
+
+			if (RootPaths.Num() == 0)
+			{
+				UE_LOG(LogPackageName, Error, TEXT("PackageName.UnregisterMountPoint: Content path '%s' is not mounted!"), *ContentPath);
 				return;
 			}
 
-			if (RootPath[0] != TEXT('/'))
+			if (RootPaths.Num() > 1)
 			{
-				UE_LOG(LogPackageName, Error, TEXT("PackageName.UnregisterMountPoint: Invalid RootPath, should start with '/'! Example: '/Game/'"));
+				UE_LOG(LogPackageName, Error, TEXT("PackageName.UnregisterMountPoint: Content path '%s' is mounted to multiple root paths, specify root path to unmount explicitly!"), *ContentPath);
+				for (const FString& MountRootPath: RootPaths)
+				{
+					UE_LOG(LogPackageName, Error, TEXT("- %s"), *MountRootPath);
+				}
 				return;
 			}
 
-			Get().RemoveMountPoint(RootPath, ContentPath);
+			RootPath = RootPaths[0];
 		}
+			
+		Get().RemoveMountPoint(RootPath, ContentPath);
 	}
+	else
+	{
+		const FString& RootPath = Args[0];
+		const FString& ContentPath = Args[1];
+		if (ContentPath[0] == TEXT('/'))
+		{
+			UE_LOG(LogPackageName, Error, TEXT("PackageName.UnregisterMountPoint: Invalid ContentPath, should not start with '/'! Example: '../../../ProjectName/Content/'"));
+			return;
+		}
+
+		if (RootPath[0] != TEXT('/'))
+		{
+			UE_LOG(LogPackageName, Error, TEXT("PackageName.UnregisterMountPoint: Invalid RootPath, should start with '/'! Example: '/Game/'"));
+			return;
+		}
+
+		Get().RemoveMountPoint(RootPath, ContentPath);
+	}
+}
 #endif
-};
 
 void FPackageName::InternalFilenameToLongPackageName(FStringView InFilename, FStringBuilderBase& OutPackageName)
 {
diff --git a/Engine/Source/Runtime/CoreUObject/Private/Misc/PackageName/MountPoint.cpp b/Engine/Source/Runtime/CoreUObject/Private/Misc/PackageName/MountPoint.cpp
new file mode 100644
--- /dev/null
+++ b/Engine/Source/Runtime/CoreUObject/Private/Misc/PackageName/MountPoint.cpp
@@ -0,0 +1,14 @@
+// Copyright Epic Games, Inc. All Rights Reserved.
+
+#include "MountPointPrivate.h"
+
+FMountPoint::FMountPoint(FString&& InRootPath, FString&& InContentPathRelative,
+	FString&& InContentPathAbsolute, EMountFlags Flags)
+	: RootPath(MoveTemp(InRootPath))
+	, ContentPathRelative(MoveTemp(InContentPathRelative))
+	, ContentPathAbsolute(MoveTemp(InContentPathAbsolute))
+	, bRemovable(!!(Flags& EMountFlags::Removable))
+	, bReadOnly(!!(Flags& EMountFlags::ReadOnly))
+	, bAlias(!!(Flags& EMountFlags::Alias))
+{
+}
\ No newline at end of file
diff --git a/Engine/Source/Runtime/CoreUObject/Private/Misc/PackageName/MountPointPrivate.h b/Engine/Source/Runtime/CoreUObject/Private/Misc/PackageName/MountPointPrivate.h
new file mode 100644
--- /dev/null
+++ b/Engine/Source/Runtime/CoreUObject/Private/Misc/PackageName/MountPointPrivate.h
@@ -0,0 +1,41 @@
+// Copyright Epic Games, Inc. All Rights Reserved.
+
+#pragma once
+
+#include "Containers/UnrealString.h"
+#include "Misc/EnumClassFlags.h"
+
+struct FMountPoint
+{
+	enum class EMountFlags
+	{
+		Removable = 0x1,
+		NotRemovable = 0x0,
+		ReadOnly = 0x2,
+		NotReadOnly = 0x0,
+		Alias = 0x4,
+		NotAlias = 0x0,
+	};
+
+	/** The LongPackageName path(e.g., "/Engine/") */
+	const FString RootPath;
+
+	/** The LocalPath relative path(e.g., "../../../Engine/Content/") */
+	const FString ContentPathRelative;
+
+	/** The LocalPath absolute path */
+	const FString ContentPathAbsolute;
+
+	/** Whether the mountpoint was created by InsertMountPoint and can be removed by RemoveMountPoint. */
+	const bool bRemovable : 1;
+
+	/** Whether new Packages can be saved into the mountpoint. */
+	const bool bReadOnly : 1;
+
+	/** Whether this mountpoint is the same RootPath as another mountpoint but bound to a different ContentPath. */
+	const bool bAlias : 1;
+
+	FMountPoint(FString&& InRootPath, FString&& InContentPathRelative,
+		FString&& InContentPathAbsolute, EMountFlags Flags);
+};
+ENUM_CLASS_FLAGS(FMountPoint::EMountFlags);
diff --git a/Engine/Source/Runtime/CoreUObject/Private/Misc/PackageName/PackageNamePrivate.h b/Engine/Source/Runtime/CoreUObject/Private/Misc/PackageName/PackageNamePrivate.h
new file mode 100644
--- /dev/null
+++ b/Engine/Source/Runtime/CoreUObject/Private/Misc/PackageName/PackageNamePrivate.h
@@ -0,0 +1,81 @@
+// Copyright Epic Games, Inc. All Rights Reserved.
+
+#pragma once
+
+#include "Containers/Array.h"
+#include "Containers/DirectoryTree.h"
+#include "Containers/StringView.h"
+#include "Containers/UnrealString.h"
+#include "HAL/IConsoleManager.h"
+#include "Misc/PackageName.h"
+#include "Misc/PackageName/MountPointPrivate.h"
+#include "Misc/TransactionallySafeRWLock.h"
+#include "Templates/UniquePtr.h"
+
+struct FLongPackagePathsSingleton
+{
+	mutable FTransactionallySafeRWLock MountLock;
+
+	FString ConfigRootPath;
+	FString EngineRootPath;
+	FString GameRootPath;
+	FString ScriptRootPath;
+	FString MemoryRootPath;
+	FString TempRootPath;
+
+	FString VerseSubPath;
+
+	FString EngineContentPath;
+	FString ContentPathShort;
+	FString EngineShadersPath;
+	FString EngineShadersPathShort;
+	FString GameContentPath;
+	FString GameConfigPath;
+	FString GameScriptPath;
+	FString GameSavedPath;
+	FString GameContentPathRebased;
+	FString GameConfigPathRebased;
+	FString GameScriptPathRebased;
+	FString GameSavedPathRebased;
+
+	TDirectoryTree<FMountPoint*> RootPathTree;
+	TDirectoryTree<FMountPoint*> ContentPathTree;
+	TArray<TUniquePtr<FMountPoint>> MountPoints;
+
+	// singleton
+	static FLongPackagePathsSingleton& Get();
+
+	/** Initialization function to setup content paths that cannot be run until CoreUObject/PluginManager have been initialized. */
+	void OnCoreUObjectInitialized();
+
+	/**
+	 * Given a content path return the consistently-formatted RelativePath and AbsolutePath, suitable for lookup based only on text
+	 * The RelativePath is consistent with FileManager relative paths.
+	 */
+	static FMountPoint ConstructMountPoint(const FString& RootPath, const FString& ContentPath,
+		FMountPoint::EMountFlags Flags);
+
+	void InsertMountPoint(const FString& RootPath, const FString& ContentPath);
+
+	/** This will remove a previously inserted mount point. */
+	void RemoveMountPoint(const FString& RootPath, const FString& ContentPath);
+	void RemoveMountPoint(const FMountPoint& MountPoint);
+
+	/** Checks whether the specific root path is a valid mount point. */
+	bool MountPointExists(FStringView RootPath);
+
+private:
+#if !UE_BUILD_SHIPPING
+	const FAutoConsoleCommand DumpMountPointsCommand;
+	const FAutoConsoleCommand RegisterMountPointCommand;
+	const FAutoConsoleCommand UnregisterMountPointCommand;
+#endif
+	
+	FLongPackagePathsSingleton();
+
+#if !UE_BUILD_SHIPPING
+	void ExecDumpMountPoints(const TArray<FString>& Args);
+	void ExecInsertMountPoint(const TArray<FString>& Args);
+	void ExecRemoveMountPoint(const TArray<FString>& Args);
+#endif
+};
\ No newline at end of file
