PackageName: Fix edge-case crashes in exec commands, found by static analysis tools.
#rnx

#Synced-CL 50109798

diff --git a/Engine/Source/Runtime/CoreUObject/Private/Misc/PackageName.cpp b/Engine/Source/Runtime/CoreUObject/Private/Misc/PackageName.cpp
--- a/Engine/Source/Runtime/CoreUObject/Private/Misc/PackageName.cpp
+++ b/Engine/Source/Runtime/CoreUObject/Private/Misc/PackageName.cpp
@@ -479,6 +479,12 @@ void FLongPackagePathsSingleton::ExecInsertMountPoint(const TArray<FString>& Arg
 	const FString& RootPath = Args[0];
 	const FString& ContentPath = Args[1];
 
+	if (RootPath.IsEmpty() || ContentPath.IsEmpty())
+	{
+		UE_LOG(LogPackageName, Error, TEXT("PackageName.RegisterMountPoint: Invalid empty path, expected non-empty RootPath and ContentPath."));
+		return;
+	}
+
 	if (ContentPath[0] == TEXT('/'))
 	{
 		UE_LOG(LogPackageName, Error, TEXT("PackageName.RegisterMountPoint: Invalid ContentPath, should not start with '/'! Example: '../../../ProjectName/Content/'"));
@@ -507,6 +513,15 @@ void FLongPackagePathsSingleton::ExecRemoveMountPoint(const TArray<FString>& Arg
 		return;
 	}
 
+	for (const FString& Arg : Args)
+	{
+		if (Arg.IsEmpty())
+		{
+			UE_LOG(LogPackageName, Error, TEXT("PackageName.UnregisterMountPoint: Invalid empty path, expected non-empty root or content paths."));
+			return;
+		}
+	}
+
 	if (Args.Num() == 1)
 	{
 		const FString& Path = Args[0];
@@ -519,6 +534,7 @@ void FLongPackagePathsSingleton::ExecRemoveMountPoint(const TArray<FString>& Arg
 		{
 			UE::TReadScopeLock ScopeLock(MountLock);
 				
+			RootPath = Path;
 			TArray<FString> ContentPathRelatives;
 			for (TUniquePtr<FMountPoint>& ExistingMount : Get().MountPoints)
 			{
