PackageName: Convert FMountPoint from struct to class (variables guarded by accessors) and move it and its internal enum into namespace UE::PackageName. This is a mechanical change to support the upcoming public IMountPoint API for mountpoints.
#rnx

#Synced-CL 50109798

diff --git a/Engine/Source/Runtime/CoreUObject/Private/Misc/PackageName.cpp b/Engine/Source/Runtime/CoreUObject/Private/Misc/PackageName.cpp
--- a/Engine/Source/Runtime/CoreUObject/Private/Misc/PackageName.cpp
+++ b/Engine/Source/Runtime/CoreUObject/Private/Misc/PackageName.cpp
@@ -185,9 +185,11 @@ void FLongPackagePathsSingleton::OnCoreUObjectInitialized()
 	UE::CoreUObject::Private::PluginHandler::Install();
 }
 
-FMountPoint FLongPackagePathsSingleton::ConstructMountPoint(const FString& RootPath, const FString& ContentPath,
-	FMountPoint::EMountFlags Flags)
+UE::PackageName::FMountPoint FLongPackagePathsSingleton::ConstructMountPoint(const FString& RootPath, const FString& ContentPath,
+	UE::PackageName::EMountFlags Flags)
 {
+	using namespace UE::PackageName;
+
 	FString MountPath = ContentPath;
 
 	// If a relative path is passed, convert to an absolute path 
@@ -220,22 +222,24 @@ FMountPoint FLongPackagePathsSingleton::ConstructMountPoint(const FString& RootP
 
 void FLongPackagePathsSingleton::InsertMountPoint(const FString& RootPath, const FString& ContentPath)
 {	
+	using namespace UE::PackageName;
+
 	TUniquePtr<FMountPoint> MountPointOwner(MakeUnique<FMountPoint>(
-		ConstructMountPoint(RootPath, ContentPath, FMountPoint::EMountFlags::Removable)));
+		ConstructMountPoint(RootPath, ContentPath, EMountFlags::Removable)));
 	FMountPoint* MountPoint = MountPointOwner.Get();
 	{
 		UE::TWriteScopeLock ScopeLock(MountLock);
 		MountPoints.Insert(MoveTemp(MountPointOwner), 0);
 
-		RootPathTree.FindOrAdd(MountPoint->RootPath) = MountPoint;
-		ContentPathTree.FindOrAdd(MountPoint->ContentPathRelative) = MountPoint;
-		ContentPathTree.FindOrAdd(MountPoint->ContentPathAbsolute) = MountPoint;
+		RootPathTree.FindOrAdd(MountPoint->GetRootPath()) = MountPoint;
+		ContentPathTree.FindOrAdd(MountPoint->GetContentPathRelative()) = MountPoint;
+		ContentPathTree.FindOrAdd(MountPoint->GetContentPathAbsolute()) = MountPoint;
 	}		
 
 	// Downgrade the log messages on startup to Verbose to reduce startup spam
 #if !NO_LOGGING
 	FString Message = FString::Printf(TEXT("FPackageName: Mount point added: '%s' mounted to '%s'"),
-		*MountPoint->ContentPathRelative, *MountPoint->RootPath);
+		*MountPoint->GetContentPathRelative(), *MountPoint->GetRootPath());
 	if (IsEngineStartupModuleLoadingComplete())
 	{
 		UE_LOG(LogPackageName, Log, TEXT("%s"), *Message);
@@ -247,66 +251,70 @@ void FLongPackagePathsSingleton::InsertMountPoint(const FString& RootPath, const
 #endif
 
 	// Let subscribers know that a new content path was mounted
-	FPackageName::OnContentPathMounted().Broadcast( RootPath, MountPoint->ContentPathRelative);
+	FPackageName::OnContentPathMounted().Broadcast( RootPath, MountPoint->GetContentPathRelative());
 }
 
 void FLongPackagePathsSingleton::RemoveMountPoint(const FString& RootPath, const FString& ContentPath)
 {
-	RemoveMountPoint(ConstructMountPoint(RootPath, ContentPath, FMountPoint::EMountFlags::Removable));
+	using namespace UE::PackageName;
+
+	RemoveMountPoint(ConstructMountPoint(RootPath, ContentPath, EMountFlags::Removable));
 }
 
-void FLongPackagePathsSingleton::RemoveMountPoint(const FMountPoint& MountPoint)
+void FLongPackagePathsSingleton::RemoveMountPoint(const UE::PackageName::FMountPoint& MountPoint)
 {
+	using namespace UE::PackageName;
+
 	bool bRemoved = false;
 	{
 		UE::TWriteScopeLock ScopeLock(MountLock);
 		for (int32 MountPointIndex = 0; MountPointIndex < MountPoints.Num(); /* Conditionally incremented in loop */)
 		{
 			FMountPoint& ExistingMount = *MountPoints[MountPointIndex];
-			if (ExistingMount.RootPath == MountPoint.RootPath &&
-				ExistingMount.ContentPathRelative == MountPoint.ContentPathRelative &&
-				ExistingMount.bRemovable)
+			if (ExistingMount.GetRootPath() == MountPoint.GetRootPath() &&
+				ExistingMount.GetContentPathRelative() == MountPoint.GetContentPathRelative() &&
+				ExistingMount.IsRemovable())
 			{
 				bool bRootPathExisted;
-				RootPathTree.Remove(ExistingMount.RootPath, &bRootPathExisted);
+				RootPathTree.Remove(ExistingMount.GetRootPath(), &bRootPathExisted);
 				if (bRootPathExisted)
 				{
 					// See if there is a lower priority path that replaces it
 					for (TUniquePtr<FMountPoint>& OtherMount : TArrayView<TUniquePtr<FMountPoint>>(MountPoints).RightChop(MountPointIndex + 1))
 					{
-						if (OtherMount->RootPath == MountPoint.RootPath)
+						if (OtherMount->GetRootPath() == MountPoint.GetRootPath())
 						{
-							RootPathTree.FindOrAdd(OtherMount->RootPath) = OtherMount.Get();
+							RootPathTree.FindOrAdd(OtherMount->GetRootPath()) = OtherMount.Get();
 							break;
 						}
 					}
 				}
 				bool bContentPathRelativeExisted;
-				ContentPathTree.Remove(ExistingMount.ContentPathRelative, &bContentPathRelativeExisted);
+				ContentPathTree.Remove(ExistingMount.GetContentPathRelative(), &bContentPathRelativeExisted);
 				if (bContentPathRelativeExisted)
 				{
 					// See if there is a lower priority path that replaces it
 					for (TUniquePtr<FMountPoint>& OtherMount :
 						TArrayView<TUniquePtr<FMountPoint>>(MountPoints).RightChop(MountPointIndex+1))
 					{
-						if (OtherMount->ContentPathRelative == MountPoint.ContentPathRelative)
+						if (OtherMount->GetContentPathRelative() == MountPoint.GetContentPathRelative())
 						{
-							ContentPathTree.FindOrAdd(OtherMount->ContentPathRelative) = OtherMount.Get();
+							ContentPathTree.FindOrAdd(OtherMount->GetContentPathRelative()) = OtherMount.Get();
 							break;
 						}
 					}
 				}
 				bool bContentPathAbsoluteExisted;
-				ContentPathTree.Remove(ExistingMount.ContentPathAbsolute, &bContentPathAbsoluteExisted);
+				ContentPathTree.Remove(ExistingMount.GetContentPathAbsolute(), &bContentPathAbsoluteExisted);
 				if (bContentPathAbsoluteExisted)
 				{
 					// See if there is a lower priority path that replaces it
 					for (TUniquePtr<FMountPoint>& OtherMount :
 						TArrayView<TUniquePtr<FMountPoint>>(MountPoints).RightChop(MountPointIndex + 1))
 					{
-						if (OtherMount->ContentPathAbsolute == MountPoint.ContentPathAbsolute)
+						if (OtherMount->GetContentPathAbsolute() == MountPoint.GetContentPathAbsolute())
 						{
-							ContentPathTree.FindOrAdd(OtherMount->ContentPathAbsolute) = OtherMount.Get();
+							ContentPathTree.FindOrAdd(OtherMount->GetContentPathAbsolute()) = OtherMount.Get();
 							break;
 						}
 					}
@@ -323,19 +331,19 @@ void FLongPackagePathsSingleton::RemoveMountPoint(const FMountPoint& MountPoint)
 		if (bRemoved)
 		{
 			UE_LOG(LogPackageName, Display, TEXT("FPackageName: Mount point removed: '%s' unmounted from '%s'"),
-				*MountPoint.ContentPathRelative, *MountPoint.RootPath);
+				*MountPoint.GetContentPathRelative(), *MountPoint.GetRootPath());
 		}
 		else
 		{
 			UE_LOG(LogPackageName, Display, TEXT("FPackageName: Mount point remove failed: no MountPoint found mapping '%s' to '%s'"),
-				*MountPoint.ContentPathRelative, *MountPoint.RootPath);
+				*MountPoint.GetContentPathRelative(), *MountPoint.GetRootPath());
 		}
 	}
 
 	// Let subscribers know that a new content path was unmounted
 	if (bRemoved)
 	{
-		FPackageName::OnContentPathDismounted().Broadcast(MountPoint.RootPath, MountPoint.ContentPathRelative);
+		FPackageName::OnContentPathDismounted().Broadcast(MountPoint.GetRootPath(), MountPoint.GetContentPathRelative());
 	}
 }
 
@@ -361,6 +369,8 @@ FLongPackagePathsSingleton::FLongPackagePathsSingleton()
 		FConsoleCommandWithArgsDelegate::CreateRaw(this, &FLongPackagePathsSingleton::ExecRemoveMountPoint))
 #endif
 {
+	using namespace UE::PackageName;
+
 	SCOPED_BOOT_TIMING("FPackageName::FLongPackagePathsSingleton");
 
 	ConfigRootPath = TEXT("/Config/");
@@ -387,7 +397,6 @@ FLongPackagePathsSingleton::FLongPackagePathsSingleton()
 	GameConfigPathRebased  = RebasedGameDir / TEXT("Config/");
 	GameScriptPathRebased  = RebasedGameDir / TEXT("Script/");
 	GameSavedPathRebased   = RebasedGameDir / TEXT("Saved/");
-	using EMountFlags = FMountPoint::EMountFlags;
 
 	auto LocalInsertMountPoint = [this](const FString& RootPath, const FString& ContentPathRelative, EMountFlags InFlags)
 	{
@@ -402,17 +411,17 @@ FLongPackagePathsSingleton::FLongPackagePathsSingleton()
 
 		// We are appending to the MountPoint list, and will have only a few duplicates. If any duplicates do occur,
 		// keep the MountPoint that was appended earlier as the registered mountpoint.
-		FMountPoint*& ExistingRootPath = RootPathTree.FindOrAdd(MountPoint.RootPath);
+		FMountPoint*& ExistingRootPath = RootPathTree.FindOrAdd(MountPoint.GetRootPath());
 		if (!ExistingRootPath)
 		{
 			ExistingRootPath = &MountPoint;
 		}
-		FMountPoint*& ExistingContentPathRelative = ContentPathTree.FindOrAdd(MountPoint.ContentPathRelative);
+		FMountPoint*& ExistingContentPathRelative = ContentPathTree.FindOrAdd(MountPoint.GetContentPathRelative());
 		if (!ExistingContentPathRelative)
 		{
 			ExistingContentPathRelative = &MountPoint;
 		}
-		FMountPoint*& ExistingContentPathAbsolute = ContentPathTree.FindOrAdd(MountPoint.ContentPathAbsolute);
+		FMountPoint*& ExistingContentPathAbsolute = ContentPathTree.FindOrAdd(MountPoint.GetContentPathAbsolute());
 		if (!ExistingContentPathAbsolute)
 		{
 			ExistingContentPathAbsolute = &MountPoint;
@@ -448,20 +457,22 @@ FLongPackagePathsSingleton::FLongPackagePathsSingleton()
 #if !UE_BUILD_SHIPPING
 void FLongPackagePathsSingleton::ExecDumpMountPoints(const TArray<FString>& Args)
 {
+	using namespace UE::PackageName;
+
 	UE_LOG(LogPackageName, Log, TEXT("Valid mount points:"));
 
 	UE::TReadScopeLock ScopeLock(MountLock);
 	for (const TUniquePtr<FMountPoint>& MountPoint : MountPoints)
 	{
-		UE_LOG(LogPackageName, Log, TEXT("	'%s' -> '%s'"), *MountPoint->RootPath, *MountPoint->ContentPathRelative);
+		UE_LOG(LogPackageName, Log, TEXT("	'%s' -> '%s'"), *MountPoint->GetRootPath(), *MountPoint->GetContentPathRelative());
 	}
 
 	UE_LOG(LogPackageName, Log, TEXT("Removable mount points:"));
 	for (const TUniquePtr<FMountPoint>& MountPoint : MountPoints)
 	{
-		if (MountPoint->bRemovable)
+		if (MountPoint->IsRemovable())
 		{
-			UE_LOG(LogPackageName, Log, TEXT("	'%s'"), *MountPoint->RootPath);
+			UE_LOG(LogPackageName, Log, TEXT("	'%s'"), *MountPoint->GetRootPath());
 		}
 	}
 }
@@ -502,6 +513,8 @@ void FLongPackagePathsSingleton::ExecInsertMountPoint(const TArray<FString>& Arg
 
 void FLongPackagePathsSingleton::ExecRemoveMountPoint(const TArray<FString>& Args)
 {
+	using namespace UE::PackageName;
+
 	if ( Args.Num() < 1 || Args.Num() > 2 )
 	{
 		UE_LOG(LogPackageName, Log, TEXT("Usage: PackageName.UnregisterMountPoint <Path>"));
@@ -538,9 +551,9 @@ void FLongPackagePathsSingleton::ExecRemoveMountPoint(const TArray<FString>& Arg
 			TArray<FString> ContentPathRelatives;
 			for (TUniquePtr<FMountPoint>& ExistingMount : Get().MountPoints)
 			{
-				if (ExistingMount->RootPath == RootPath)
+				if (ExistingMount->GetRootPath() == RootPath)
 				{
-					ContentPathRelatives.Add(ExistingMount->ContentPathRelative);
+					ContentPathRelatives.Add(ExistingMount->GetContentPathRelative());
 				}
 			}
 			if (ContentPathRelatives.Num() == 0)
@@ -569,9 +582,9 @@ void FLongPackagePathsSingleton::ExecRemoveMountPoint(const TArray<FString>& Arg
 			TArray<FString> RootPaths;
 			for (TUniquePtr<FMountPoint>& ExistingMount : Get().MountPoints)
 			{
-				if (ExistingMount->ContentPathRelative == ContentPath)
+				if (ExistingMount->GetContentPathRelative() == ContentPath)
 				{
-					RootPaths.Add(ExistingMount->RootPath);
+					RootPaths.Add(ExistingMount->GetRootPath());
 				}
 			}
 
@@ -619,6 +632,8 @@ void FLongPackagePathsSingleton::ExecRemoveMountPoint(const TArray<FString>& Arg
 
 void FPackageName::InternalFilenameToLongPackageName(FStringView InFilename, FStringBuilderBase& OutPackageName)
 {
+	using namespace UE::PackageName;
+
 	const FLongPackagePathsSingleton& Paths = FLongPackagePathsSingleton::Get();
 	FString Filename(InFilename);
 	FPaths::NormalizeFilename(Filename);
@@ -678,9 +693,9 @@ void FPackageName::InternalFilenameToLongPackageName(FStringView InFilename, FSt
 		}
 		const FMountPoint* MountPoint = *MountPointPtr;
 		FStringView RelPath;
-		if (!FPathViews::TryMakeChildPathRelativeTo(SearchName, MountPoint->ContentPathRelative, RelPath))
+		if (!FPathViews::TryMakeChildPathRelativeTo(SearchName, MountPoint->GetContentPathRelative(), RelPath))
 		{
-			if (!FPathViews::TryMakeChildPathRelativeTo(SearchName, MountPoint->ContentPathAbsolute, RelPath))
+			if (!FPathViews::TryMakeChildPathRelativeTo(SearchName, MountPoint->GetContentPathAbsolute(), RelPath))
 			{
 				// DirectoryTree is less conservative than TryMakeChildPathRelativeTo; e.g. d:/Dir//Root will match
 				// d:/Dir/Root/Child in DirectoryTree but not in TryMakeChildPathRelativeTo. Treat this as PathNotMounted,
@@ -688,7 +703,7 @@ void FPackageName::InternalFilenameToLongPackageName(FStringView InFilename, FSt
 				return false;
 			}
 		}
-		OutPackageName << MountPoint->RootPath;
+		OutPackageName << MountPoint->GetRootPath();
 		FPathViews::AppendPath(OutPackageName, RelPath);
 		return true;
 	};
@@ -812,6 +827,8 @@ bool FPackageName::TryConvertFilenameToLongPackageName(FStringView InFilename, F
 
 FString FPackageName::FilenameToLongPackageName(const FString& InFilename)
 {
+	using namespace UE::PackageName;
+
 	FString FailureReason;
 	FString Result;
 	if (!TryConvertFilenameToLongPackageName(InFilename, Result, &FailureReason))
@@ -823,8 +840,8 @@ FString FPackageName::FilenameToLongPackageName(const FString& InFilename)
 			UE::TReadScopeLock ScopeLock(Paths.MountLock);
 			for (const TUniquePtr<FMountPoint>& MountPoint : Paths.MountPoints)
 			{
-				ContentRootsArrayRelative.Add(MountPoint->ContentPathRelative);
-				ContentRootsArrayAbsolute.Add(MountPoint->ContentPathAbsolute);
+				ContentRootsArrayRelative.Add(MountPoint->GetContentPathRelative());
+				ContentRootsArrayAbsolute.Add(MountPoint->GetContentPathAbsolute());
 			}
 		}
 
@@ -862,6 +879,8 @@ bool FPackageName::TryConvertLongPackageNameToFilename(const FString& InLongPack
 
 bool FPackageName::TryConvertLongPackageNameToFilename(FStringView InLongPackageName, FString& OutFilename, FStringView InExtension)
 {
+	using namespace UE::PackageName;
+
 	const auto& Paths = FLongPackagePathsSingleton::Get();
 	UE::TReadScopeLock ScopeLock(Paths.MountLock);
 	const FMountPoint* const* MountPointPtr = Paths.RootPathTree.FindClosestValue(InLongPackageName);
@@ -873,7 +892,7 @@ bool FPackageName::TryConvertLongPackageNameToFilename(FStringView InLongPackage
 
 	const FMountPoint* MountPoint = *MountPointPtr;
 	FStringView RelPath;
-	if (!FPathViews::TryMakeChildPathRelativeTo(InLongPackageName, MountPoint->RootPath, RelPath))
+	if (!FPathViews::TryMakeChildPathRelativeTo(InLongPackageName, MountPoint->GetRootPath(), RelPath))
 	{
 		// DirectoryTree is less conservative than TryMakeChildPathRelativeTo; e.g. d:/Dir//Root will match
 		// d:/Dir/Root/Child in DirectoryTree but not in TryMakeChildPathRelativeTo. Treat this as PathNotMounted,
@@ -881,7 +900,7 @@ bool FPackageName::TryConvertLongPackageNameToFilename(FStringView InLongPackage
 		return false;
 	}
 	TStringBuilder<256> Builder;
-	Builder << MountPoint->ContentPathRelative;
+	Builder << MountPoint->GetContentPathRelative();
 	FPathViews::AppendPath(Builder, RelPath);
 	Builder << InExtension;
 	OutFilename = Builder.ToView();
@@ -1169,6 +1188,8 @@ bool FPackageName::IsValidTextForLongPackageName(FStringView InLongPackageName,
 
 bool FPackageName::IsValidLongPackageName(FStringView InLongPackageName, bool bIncludeReadOnlyRoots, FText* OutReason)
 {
+	using namespace UE::PackageName;
+
 	EErrorCode Reason;
 	if (!IsValidLongPackageName(InLongPackageName, bIncludeReadOnlyRoots, &Reason))
 	{
@@ -1182,9 +1203,9 @@ bool FPackageName::IsValidLongPackageName(FStringView InLongPackageName, bool bI
 				TStringBuilder<2048> ValidRootsString;
 				for (const TUniquePtr<FMountPoint>& MountPoint : Paths.MountPoints)
 				{
-					if (bIncludeReadOnlyRoots || !MountPoint->bReadOnly)
+					if (bIncludeReadOnlyRoots || !MountPoint->IsReadOnly())
 					{
-						ValidRootsString << MountPoint->RootPath << TEXT(", ");
+						ValidRootsString << MountPoint->GetRootPath() << TEXT(", ");
 					}
 				}
 				if (bIncludeReadOnlyRoots)
@@ -1208,6 +1229,8 @@ bool FPackageName::IsValidLongPackageName(FStringView InLongPackageName, bool bI
 
 bool FPackageName::IsValidLongPackageName(FStringView InLongPackageName, bool bIncludeReadOnlyRoots /*= false*/, EErrorCode* OutReason /*= nullptr */)
 {
+	using namespace UE::PackageName;
+
 	if (!IsValidTextForLongPackageName(InLongPackageName, OutReason))
 	{
 		return false;
@@ -1220,12 +1243,12 @@ bool FPackageName::IsValidLongPackageName(FStringView InLongPackageName, bool bI
 	const FMountPoint* const* MountPointPtr = Paths.RootPathTree.FindClosestValue(InLongPackageName);
 	if (MountPointPtr)
 	{
-		if (bIncludeReadOnlyRoots || !(*MountPointPtr)->bReadOnly)
+		if (bIncludeReadOnlyRoots || !(*MountPointPtr)->IsReadOnly())
 		{
 			// DirectoryTree is less conservative than IsParentPathOf; e.g. d:/Dir//Root will match
 			// d:/Dir/Root/Child in DirectoryTree but not in IsParentPathOf. Treat this as PathNotMounted,
 			// since no mountdir will match it according to IsParentPathOf.
-			if (FPathViews::IsParentPathOf((*MountPointPtr)->RootPath, InLongPackageName))
+			if (FPathViews::IsParentPathOf((*MountPointPtr)->GetRootPath(), InLongPackageName))
 			{
 				if (OutReason) *OutReason = EErrorCode::PackageNameUnknown;
 				return true;
@@ -1349,6 +1372,8 @@ bool FPackageName::MountPointExists(const TCHAR* RootPath)
 
 FName FPackageName::GetPackageMountPoint(FStringView InPackagePath, bool InWithoutSlashes)
 {
+	using namespace UE::PackageName;
+
 	FLongPackagePathsSingleton& Paths = FLongPackagePathsSingleton::Get();
 	
 	UE::TReadScopeLock ScopeLock(Paths.MountLock);
@@ -1367,7 +1392,7 @@ FName FPackageName::GetPackageMountPoint(FStringView InPackagePath, bool InWitho
 		// DirectoryTree is less conservative than IsParentPathOf; e.g. d:/Dir//Root will match
 		// d:/Dir/Root/Child in DirectoryTree but not in IsParentPathOf. Treat this as PathNotMounted,
 		// since no mountdir will match it according to IsParentPathOf.
-		return TryRootPath((*MountPointPtr)->RootPath);
+		return TryRootPath((*MountPointPtr)->GetRootPath());
 	}
 	else
 	{
@@ -1394,6 +1419,8 @@ FName FPackageName::GetPackageMountPoint(const TCHAR* InPackagePath, bool InWith
 
 FString FPackageName::GetContentPathForPackageRoot(FStringView InMountPoint)
 {
+	using namespace UE::PackageName;
+
 	FLongPackagePathsSingleton& Paths = FLongPackagePathsSingleton::Get();
 
 	UE::TReadScopeLock ScopeLock(Paths.MountLock);
@@ -1403,7 +1430,7 @@ FString FPackageName::GetContentPathForPackageRoot(FStringView InMountPoint)
 		return FString();
 	}
 	const FMountPoint* MountPoint = *MountPointPtr;
-	return MountPoint->ContentPathAbsolute;
+	return MountPoint->GetContentPathAbsolute();
 }
 
 bool FPackageName::TryConvertToMountedPathComponents(FStringView InFilePathOrPackageName,
@@ -1541,6 +1568,8 @@ bool FPackageName::TryConvertToMountedPathComponents(FStringView InFilePathOrPac
 
 bool FPackageName::TryGetMountPointForPath(FStringView InFilePathOrPackageName, FStringBuilderBase& OutMountPointPackageName, FStringBuilderBase& OutMountPointFilePath, FStringBuilderBase& OutRelPath, EFlexNameType* OutFlexNameType, EErrorCode* OutFailureReason)
 {
+	using namespace UE::PackageName;
+
 	OutMountPointPackageName.Reset();
 	OutMountPointFilePath.Reset();
 	OutRelPath.Reset();
@@ -1567,7 +1596,7 @@ bool FPackageName::TryGetMountPointForPath(FStringView InFilePathOrPackageName,
 	if (MountPointPtr)
 	{
 		const FMountPoint* MountPoint = *MountPointPtr;
-		if (!FPathViews::TryMakeChildPathRelativeTo(InFilePathOrPackageName, MountPoint->RootPath, RelPath))
+		if (!FPathViews::TryMakeChildPathRelativeTo(InFilePathOrPackageName, MountPoint->GetRootPath(), RelPath))
 		{
 			// DirectoryTree is less conservative than TryMakeChildPathRelativeTo; e.g. d:/Dir//Root will match
 			// d:/Dir/Root/Child in DirectoryTree but not in TryMakeChildPathRelativeTo. Treat this as PathNotMounted,
@@ -1582,8 +1611,8 @@ bool FPackageName::TryGetMountPointForPath(FStringView InFilePathOrPackageName,
 			}
 			return false;
 		}
-		OutMountPointPackageName << MountPoint->RootPath;
-		OutMountPointFilePath << MountPoint->ContentPathRelative;
+		OutMountPointPackageName << MountPoint->GetRootPath();
+		OutMountPointFilePath << MountPoint->GetContentPathRelative();
 		OutRelPath << RelPath;
 		if (OutFlexNameType)
 		{
@@ -1607,9 +1636,9 @@ bool FPackageName::TryGetMountPointForPath(FStringView InFilePathOrPackageName,
 	if (MountPointPtr)
 	{
 		const FMountPoint* MountPoint = *MountPointPtr;
-		if (!FPathViews::TryMakeChildPathRelativeTo(PossibleAbsFilePath, MountPoint->ContentPathAbsolute, RelPath))
+		if (!FPathViews::TryMakeChildPathRelativeTo(PossibleAbsFilePath, MountPoint->GetContentPathAbsolute(), RelPath))
 		{
-			if (!FPathViews::TryMakeChildPathRelativeTo(PossibleAbsFilePath, MountPoint->ContentPathRelative, RelPath))
+			if (!FPathViews::TryMakeChildPathRelativeTo(PossibleAbsFilePath, MountPoint->GetContentPathRelative(), RelPath))
 			{
 				// DirectoryTree is less conservative than TryMakeChildPathRelativeTo; e.g. d:/Dir//Root will match
 				// d:/Dir/Root/Child in DirectoryTree but not in TryMakeChildPathRelativeTo. Treat this as PathNotMounted,
@@ -1625,8 +1654,8 @@ bool FPackageName::TryGetMountPointForPath(FStringView InFilePathOrPackageName,
 				return false;
 			}
 		}
-		OutMountPointPackageName << MountPoint->RootPath;
-		OutMountPointFilePath << MountPoint->ContentPathRelative;
+		OutMountPointPackageName << MountPoint->GetRootPath();
+		OutMountPointFilePath << MountPoint->GetContentPathRelative();
 		OutRelPath << RelPath;
 		if (OutFlexNameType)
 		{
@@ -2639,6 +2668,8 @@ void FPackageName::IteratePackagesInDirectory(const FString& RootDir, const FPac
 
 void FPackageName::QueryRootContentPaths(TArray<FString>& OutRootContentPaths, bool bIncludeReadOnlyRoots, bool bWithoutLeadingSlashes, bool bWithoutTrailingSlashes)
 {
+	using namespace UE::PackageName;
+
 	const FLongPackagePathsSingleton& Paths = FLongPackagePathsSingleton::Get();
 
 	{
@@ -2647,25 +2678,25 @@ void FPackageName::QueryRootContentPaths(TArray<FString>& OutRootContentPaths, b
 		// Maintain the legacy order: non-removable non-readonly, then removable, then non-removable readonly
 		for (const TUniquePtr<FMountPoint>& MountPoint : Paths.MountPoints)
 		{
-			if (!MountPoint->bRemovable && !MountPoint->bReadOnly && !MountPoint->bAlias)
+			if (!MountPoint->IsRemovable() && !MountPoint->IsReadOnly() && !MountPoint->IsAlias())
 			{
-				OutRootContentPaths.Add(MountPoint->RootPath);
+				OutRootContentPaths.Add(MountPoint->GetRootPath());
 			}
 		}
 		for (const TUniquePtr<FMountPoint>& MountPoint : Paths.MountPoints)
 		{
-			if (MountPoint->bRemovable)
+			if (MountPoint->IsRemovable())
 			{
-				OutRootContentPaths.Add(MountPoint->RootPath);
+				OutRootContentPaths.Add(MountPoint->GetRootPath());
 			}
 		}
 		if (bIncludeReadOnlyRoots)
 		{
 			for (const TUniquePtr<FMountPoint>& MountPoint : Paths.MountPoints)
 			{
-				if (!MountPoint->bRemovable && MountPoint->bReadOnly && !MountPoint->bAlias)
+				if (!MountPoint->IsRemovable() && MountPoint->IsReadOnly() && !MountPoint->IsAlias())
 				{
-					OutRootContentPaths.Add(MountPoint->RootPath);
+					OutRootContentPaths.Add(MountPoint->GetRootPath());
 				}
 			}
 			const FString& ExtraRoot = Paths.MemoryRootPath;
@@ -2692,6 +2723,8 @@ void FPackageName::QueryRootContentPaths(TArray<FString>& OutRootContentPaths, b
 
 TArray<FString> FPackageName::QueryMountPointLocalAbsPaths()
 {
+	using namespace UE::PackageName;
+
 	const FLongPackagePathsSingleton& Paths = FLongPackagePathsSingleton::Get();
 	TArray<FString> OutAbsPaths;
 
@@ -2700,9 +2733,9 @@ TArray<FString> FPackageName::QueryMountPointLocalAbsPaths()
 		OutAbsPaths.Reserve(Paths.MountPoints.Num() + 1);
 		for (const TUniquePtr<FMountPoint>& MountPoint : Paths.MountPoints)
 		{
-			if (!MountPoint->bAlias && !MountPoint->ContentPathAbsolute.IsEmpty())
+			if (!MountPoint->IsAlias() && !MountPoint->GetContentPathAbsolute().IsEmpty())
 			{
-				OutAbsPaths.Add(MountPoint->ContentPathAbsolute);
+				OutAbsPaths.Add(MountPoint->GetContentPathAbsolute());
 			}
 		}
 	}
diff --git a/Engine/Source/Runtime/CoreUObject/Private/Misc/PackageName/MountPoint.cpp b/Engine/Source/Runtime/CoreUObject/Private/Misc/PackageName/MountPoint.cpp
--- a/Engine/Source/Runtime/CoreUObject/Private/Misc/PackageName/MountPoint.cpp
+++ b/Engine/Source/Runtime/CoreUObject/Private/Misc/PackageName/MountPoint.cpp
@@ -2,13 +2,18 @@
 
 #include "MountPointPrivate.h"
 
+namespace UE::PackageName
+{
+
 FMountPoint::FMountPoint(FString&& InRootPath, FString&& InContentPathRelative,
 	FString&& InContentPathAbsolute, EMountFlags Flags)
 	: RootPath(MoveTemp(InRootPath))
 	, ContentPathRelative(MoveTemp(InContentPathRelative))
 	, ContentPathAbsolute(MoveTemp(InContentPathAbsolute))
-	, bRemovable(!!(Flags& EMountFlags::Removable))
-	, bReadOnly(!!(Flags& EMountFlags::ReadOnly))
-	, bAlias(!!(Flags& EMountFlags::Alias))
+	, bRemovable(!!(Flags & EMountFlags::Removable))
+	, bReadOnly(!!(Flags & EMountFlags::ReadOnly))
+	, bAlias(!!(Flags & EMountFlags::Alias))
 {
+}
+
 }
\ No newline at end of file
diff --git a/Engine/Source/Runtime/CoreUObject/Private/Misc/PackageName/MountPointPrivate.h b/Engine/Source/Runtime/CoreUObject/Private/Misc/PackageName/MountPointPrivate.h
--- a/Engine/Source/Runtime/CoreUObject/Private/Misc/PackageName/MountPointPrivate.h
+++ b/Engine/Source/Runtime/CoreUObject/Private/Misc/PackageName/MountPointPrivate.h
@@ -5,37 +5,83 @@
 #include "Containers/UnrealString.h"
 #include "Misc/EnumClassFlags.h"
 
-struct FMountPoint
-{
-	enum class EMountFlags
-	{
-		Removable = 0x1,
-		NotRemovable = 0x0,
-		ReadOnly = 0x2,
-		NotReadOnly = 0x0,
-		Alias = 0x4,
-		NotAlias = 0x0,
-	};
+namespace UE::PackageName
+{
 
-	/** The LongPackageName path(e.g., "/Engine/") */
-	const FString RootPath;
+enum class EMountFlags
+{
+	Removable = 0x1,
+	NotRemovable = 0x0,
+	ReadOnly = 0x2,
+	NotReadOnly = 0x0,
+	Alias = 0x4,
+	NotAlias = 0x0,
+};
+ENUM_CLASS_FLAGS(EMountFlags);
 
-	/** The LocalPath relative path(e.g., "../../../Engine/Content/") */
-	const FString ContentPathRelative;
+class FMountPoint
+{
+public:
+	FMountPoint(FString&& InRootPath, FString&& InContentPathRelative,
+		FString&& InContentPathAbsolute, EMountFlags Flags);
 
+	/** The LongPackageName path(e.g., "/Engine/") */
+	const FString& GetRootPath() const;
+	/** The LocalPath relative path(e.g., "../../../Engine/Content/") */
+	const FString& GetContentPathRelative() const;
 	/** The LocalPath absolute path */
-	const FString ContentPathAbsolute;
+	const FString& GetContentPathAbsolute() const;
 
 	/** Whether the mountpoint was created by InsertMountPoint and can be removed by RemoveMountPoint. */
-	const bool bRemovable : 1;
-
+	bool IsRemovable() const;
 	/** Whether new Packages can be saved into the mountpoint. */
-	const bool bReadOnly : 1;
-
+	bool IsReadOnly() const;
 	/** Whether this mountpoint is the same RootPath as another mountpoint but bound to a different ContentPath. */
-	const bool bAlias : 1;
+	bool IsAlias() const;
 
-	FMountPoint(FString&& InRootPath, FString&& InContentPathRelative,
-		FString&& InContentPathAbsolute, EMountFlags Flags);
+private:
+	const FString RootPath;
+	const FString ContentPathRelative;
+	const FString ContentPathAbsolute;
+	const bool bRemovable : 1;
+	const bool bReadOnly : 1;
+	const bool bAlias : 1;
 };
-ENUM_CLASS_FLAGS(FMountPoint::EMountFlags);
+
+
+///////////////////////////////////////////////////////
+// Inline implementations
+///////////////////////////////////////////////////////
+
+
+inline const FString& FMountPoint::GetRootPath() const
+{
+	return RootPath;
+}
+
+inline const FString& FMountPoint::GetContentPathRelative() const
+{
+	return ContentPathRelative;
+}
+
+inline const FString& FMountPoint::GetContentPathAbsolute() const
+{
+	return ContentPathAbsolute;
+}
+
+inline bool FMountPoint::IsRemovable() const
+{
+	return bRemovable;
+}
+
+inline bool FMountPoint::IsReadOnly() const
+{
+	return bReadOnly;
+}
+
+inline bool FMountPoint::IsAlias() const
+{
+	return bAlias;
+}
+
+} // namespace UE::PackageName
\ No newline at end of file
diff --git a/Engine/Source/Runtime/CoreUObject/Private/Misc/PackageName/PackageNamePrivate.h b/Engine/Source/Runtime/CoreUObject/Private/Misc/PackageName/PackageNamePrivate.h
--- a/Engine/Source/Runtime/CoreUObject/Private/Misc/PackageName/PackageNamePrivate.h
+++ b/Engine/Source/Runtime/CoreUObject/Private/Misc/PackageName/PackageNamePrivate.h
@@ -38,9 +38,9 @@ struct FLongPackagePathsSingleton
 	FString GameScriptPathRebased;
 	FString GameSavedPathRebased;
 
-	TDirectoryTree<FMountPoint*> RootPathTree;
-	TDirectoryTree<FMountPoint*> ContentPathTree;
-	TArray<TUniquePtr<FMountPoint>> MountPoints;
+	TDirectoryTree<UE::PackageName::FMountPoint*> RootPathTree;
+	TDirectoryTree<UE::PackageName::FMountPoint*> ContentPathTree;
+	TArray<TUniquePtr<UE::PackageName::FMountPoint>> MountPoints;
 
 	// singleton
 	static FLongPackagePathsSingleton& Get();
@@ -52,14 +52,14 @@ struct FLongPackagePathsSingleton
 	 * Given a content path return the consistently-formatted RelativePath and AbsolutePath, suitable for lookup based only on text
 	 * The RelativePath is consistent with FileManager relative paths.
 	 */
-	static FMountPoint ConstructMountPoint(const FString& RootPath, const FString& ContentPath,
-		FMountPoint::EMountFlags Flags);
+	static UE::PackageName::FMountPoint ConstructMountPoint(const FString& RootPath, const FString& ContentPath,
+		UE::PackageName::EMountFlags Flags);
 
 	void InsertMountPoint(const FString& RootPath, const FString& ContentPath);
 
 	/** This will remove a previously inserted mount point. */
 	void RemoveMountPoint(const FString& RootPath, const FString& ContentPath);
-	void RemoveMountPoint(const FMountPoint& MountPoint);
+	void RemoveMountPoint(const UE::PackageName::FMountPoint& MountPoint);
 
 	/** Checks whether the specific root path is a valid mount point. */
 	bool MountPointExists(FStringView RootPath);
