Fixed compilation errors in TCompiledInObjectPtr.

diff --git a/Engine/Source/Runtime/CoreUObject/Public/UObject/CompiledInObjectPtr.h b/Engine/Source/Runtime/CoreUObject/Public/UObject/CompiledInObjectPtr.h
--- a/Engine/Source/Runtime/CoreUObject/Public/UObject/CompiledInObjectPtr.h
+++ b/Engine/Source/Runtime/CoreUObject/Public/UObject/CompiledInObjectPtr.h
@@ -30,8 +30,11 @@
 
 #include "CoreTypes.h"
 #include "Concepts/ConvertibleTo.h"
+#include "Concepts/NotCVRefTo.h"
 #include "UObject/ObjectMacros.h"
 
+#include <type_traits>
+
 #if UE_WITH_CONSTINIT_UOBJECT
 
 namespace UE::CodeGen::ConstInit
@@ -131,7 +134,7 @@ struct TCompiledInObjectPtr
 	constexpr TCompiledInObjectPtr& operator=(const TCompiledInObjectPtr&) = default;
 
 	// Implicit compile-time construction from object pointer
-	template<UE::CConvertibleTo<T*> U> requires (!UE::CSameAs<T, U>)
+	template<UE::CNotCVRefTo<TCompiledInObjectPtr> U> requires (UE::CConvertibleTo<U, T*> && !UE::CSameAs<T, U>)
 	[[nodiscard]] constexpr TCompiledInObjectPtr(U&& InPtr)
 		: Pointer(InPtr)
 	{
@@ -215,9 +218,12 @@ class FStructBaseChain;
 namespace UE::Private
 {
 // Utility to allow both raw and encoded pointers to be wrapped in AsStructBaseChain to handle offsetting to private base class
+// The dummy template parameter will cause the body of the function to be deferred until it's called, where hopefully FStructBaseChain
+// is complete.
+template <typename Dummy = void>
 [[nodiscard]] consteval UE::CodeGen::ConstInit::TCompiledInObjectPtr<const FStructBaseChain> AsStructBaseChain(UE::CodeGen::ConstInit::FCompiledInObjectReference Ptr)
 {
-	return UE::CodeGen::ConstInit::TCompiledInObjectPtr<const FStructBaseChain>(Ptr);
+	return UE::CodeGen::ConstInit::TCompiledInObjectPtr<const std::conditional_t<std::is_void_v<Dummy>, FStructBaseChain, FStructBaseChain>>(Ptr);
 }
 }
 
