[StructUtils] Change PropertyBag to create FByteProperty for enums that are not a C++ class enum. Similar to what UHT and Blueprint produce with TEnumAsByte.

#Synced-CL 50158725

diff --git a/Engine/Plugins/Runtime/StateTree/Source/StateTreeEditorModule/Private/StateTreeCompiler.cpp b/Engine/Plugins/Runtime/StateTree/Source/StateTreeEditorModule/Private/StateTreeCompiler.cpp
--- a/Engine/Plugins/Runtime/StateTree/Source/StateTreeEditorModule/Private/StateTreeCompiler.cpp
+++ b/Engine/Plugins/Runtime/StateTree/Source/StateTreeEditorModule/Private/StateTreeCompiler.cpp
@@ -65,7 +65,7 @@ namespace UE::StateTree::Compiler
 	// Version used by AppendToClassSchema. Bump this version number when a recompile of all assets is needed.
 	//Only bump when only the compiler changes.
 	//You do not need to bump the version when the memory layout of UStateTree has changed.
-	static constexpr int32 ClassSchemaVersion = 2;
+	static constexpr int32 ClassSchemaVersion = 3;
 
 	FAutoConsoleVariable CVarLogEnableParameterDelegateDispatcherBinding(
 		TEXT("StateTree.Compiler.EnableParameterDelegateDispatcherBinding"),
diff --git a/Engine/Source/Runtime/CoreUObject/Private/StructUtils/PropertyBag.cpp b/Engine/Source/Runtime/CoreUObject/Private/StructUtils/PropertyBag.cpp
--- a/Engine/Source/Runtime/CoreUObject/Private/StructUtils/PropertyBag.cpp
+++ b/Engine/Source/Runtime/CoreUObject/Private/StructUtils/PropertyBag.cpp
@@ -512,11 +512,23 @@
 		case EPropertyBagPropertyType::Enum:
 			if (const UEnum* Enum = Cast<UEnum>(Desc.ValueTypeObject))
 			{
-				FEnumProperty* Prop = new FEnumProperty(PropertyScope, Desc.Name, RF_Public);
-				FNumericProperty* UnderlyingProp = new FByteProperty(Prop, "UnderlyingType", RF_Public); // HACK: Hardwire to byte property for now for BP compatibility
-				Prop->SetEnum(const_cast<UEnum*>(Enum));
-				Prop->AddCppProperty(UnderlyingProp);
-				return Prop;
+				// Only generate FEnumProperty for cpp enum class like UHT and Blueprint.
+				if (Enum->GetCppForm() == UEnum::ECppForm::EnumClass)
+				{
+					FEnumProperty* Prop = new FEnumProperty(PropertyScope, Desc.Name, RF_Public);
+					Prop->SetPropertyFlags(CPF_HasGetValueTypeHash);
+					FNumericProperty* UnderlyingProp = new FByteProperty(Prop, "UnderlyingType", RF_Public); // HACK: Hardwire to byte property for now for BP compatibility
+					Prop->SetEnum(const_cast<UEnum*>(Enum));
+					Prop->AddCppProperty(UnderlyingProp);
+					return Prop;
+				}
+				else
+				{
+					FByteProperty* Prop = new FByteProperty(PropertyScope, Desc.Name, RF_Public);
+					Prop->SetPropertyFlags(CPF_HasGetValueTypeHash);
+					Prop->Enum = const_cast<UEnum*>(Enum);
+					return Prop;
+				}
 			}
 			break;
 		case EPropertyBagPropertyType::Struct:
@@ -672,11 +684,21 @@
 			}
 		case EPropertyBagPropertyType::Enum:
 			{
-				const FEnumProperty* EnumProperty = CastFieldChecked<FEnumProperty>(Desc->CachedProperty);
-				const FNumericProperty* UnderlyingProperty = EnumProperty->GetUnderlyingProperty();
-				check(UnderlyingProperty);
-				OutValue = UnderlyingProperty->GetSignedIntPropertyValue(Address);
-				return EPropertyBagResult::Success;
+				if (const FEnumProperty* EnumProperty = CastField<FEnumProperty>(Desc->CachedProperty))
+				{
+					const FNumericProperty* UnderlyingProperty = EnumProperty->GetUnderlyingProperty();
+					check(UnderlyingProperty);
+					OutValue = UnderlyingProperty->GetSignedIntPropertyValue(Address);
+					return EPropertyBagResult::Success;
+				}
+				else if (const FByteProperty* ByteProperty = CastField<FByteProperty>(Desc->CachedProperty))
+				{
+					if (ByteProperty->Enum != nullptr)
+					{
+						OutValue = static_cast<int64>(ByteProperty->GetPropertyValue(Address));
+						return EPropertyBagResult::Success;
+					}
+				}
 			}
 		default:
 			return EPropertyBagResult::TypeMismatch;
@@ -750,11 +772,22 @@
 		}
 		case EPropertyBagPropertyType::Enum:
 		{
-			const FEnumProperty* EnumProperty = CastFieldChecked<FEnumProperty>(Desc->CachedProperty);
-			const FNumericProperty* UnderlyingProperty = EnumProperty->GetUnderlyingProperty();
-			check(UnderlyingProperty);
-			OutValue = UnderlyingProperty->GetUnsignedIntPropertyValue(Address);
-			return EPropertyBagResult::Success;
+			if (const FEnumProperty* EnumProperty = CastField<FEnumProperty>(Desc->CachedProperty))
+			{
+				const FNumericProperty* UnderlyingProperty = EnumProperty->GetUnderlyingProperty();
+				check(UnderlyingProperty);
+				OutValue = UnderlyingProperty->GetUnsignedIntPropertyValue(Address);
+				return EPropertyBagResult::Success;
+			}
+			else if (const FByteProperty* ByteProperty = CastField<FByteProperty>(Desc->CachedProperty))
+			{
+				if (ByteProperty->Enum != nullptr)
+				{
+					OutValue = ByteProperty->GetPropertyValue(Address);
+					return EPropertyBagResult::Success;
+				}
+			}
+
 		}
 		default:
 			return EPropertyBagResult::TypeMismatch;
@@ -828,11 +861,21 @@
 			}
 		case EPropertyBagPropertyType::Enum:
 			{
-				const FEnumProperty* EnumProperty = CastFieldChecked<FEnumProperty>(Desc->CachedProperty);
-				const FNumericProperty* UnderlyingProperty = EnumProperty->GetUnderlyingProperty();
-				check(UnderlyingProperty);
-				OutValue = static_cast<double>(UnderlyingProperty->GetSignedIntPropertyValue(Address));
-				return EPropertyBagResult::Success;
+				if (const FEnumProperty* EnumProperty = CastField<FEnumProperty>(Desc->CachedProperty))
+				{
+					const FNumericProperty* UnderlyingProperty = EnumProperty->GetUnderlyingProperty();
+					check(UnderlyingProperty);
+					OutValue = static_cast<double>(UnderlyingProperty->GetSignedIntPropertyValue(Address));
+					return EPropertyBagResult::Success;
+				}
+				else if (const FByteProperty* ByteProperty = CastField<FByteProperty>(Desc->CachedProperty))
+				{
+					if (ByteProperty->Enum != nullptr)
+					{
+						OutValue = ByteProperty->GetPropertyValue(Address);
+						return EPropertyBagResult::Success;
+					}
+				}
 			}
 		default:
 			return EPropertyBagResult::TypeMismatch;
@@ -886,16 +929,30 @@
 			return EPropertyBagResult::TypeMismatch;
 		}
 
-		const FEnumProperty* EnumProperty = CastFieldChecked<FEnumProperty>(Desc->CachedProperty);
-		const FNumericProperty* UnderlyingProperty = EnumProperty->GetUnderlyingProperty();
-		check(UnderlyingProperty);
+		if (const FEnumProperty* EnumProperty = CastField<FEnumProperty>(Desc->CachedProperty))
+		{
+			const FNumericProperty* UnderlyingProperty = EnumProperty->GetUnderlyingProperty();
+			check(UnderlyingProperty);
 	
-		if (RequestedEnum != EnumProperty->GetEnum())
+			if (RequestedEnum != EnumProperty->GetEnum())
+			{
+				return EPropertyBagResult::TypeMismatch;
+			}
+	
+			OutValue = static_cast<uint8>(UnderlyingProperty->GetUnsignedIntPropertyValue(Address));
+		}
+		else if (const FByteProperty* ByteProperty = CastField<FByteProperty>(Desc->CachedProperty))
+		{
+			if (RequestedEnum != ByteProperty->Enum)
+			{
+				return EPropertyBagResult::TypeMismatch;
+			}
+			OutValue = ByteProperty->GetPropertyValue(Address);
+		}
+		else
 		{
 			return EPropertyBagResult::TypeMismatch;
 		}
-	
-		OutValue = static_cast<uint8>(UnderlyingProperty->GetUnsignedIntPropertyValue(Address));
 
 		return EPropertyBagResult::Success;
 	}
@@ -1066,11 +1123,21 @@
 			}
 		case EPropertyBagPropertyType::Enum:
 			{
-				const FEnumProperty* EnumProperty = CastFieldChecked<FEnumProperty>(Desc->CachedProperty);
-				const FNumericProperty* UnderlyingProperty = EnumProperty->GetUnderlyingProperty();
-				check(UnderlyingProperty);
-				UnderlyingProperty->SetIntPropertyValue(Address, static_cast<uint64>(InValue));
-				return EPropertyBagResult::Success;
+				if (const FEnumProperty* EnumProperty = CastField<FEnumProperty>(Desc->CachedProperty))
+				{
+					const FNumericProperty* UnderlyingProperty = EnumProperty->GetUnderlyingProperty();
+					check(UnderlyingProperty);
+					UnderlyingProperty->SetIntPropertyValue(Address, static_cast<uint64>(InValue));
+					return EPropertyBagResult::Success;
+				}
+				else if (const FByteProperty* ByteProperty = CastField<FByteProperty>(Desc->CachedProperty))
+				{
+					if (ByteProperty->Enum)
+					{
+						ByteProperty->SetPropertyValue(Address, static_cast<uint8>(InValue));
+						return EPropertyBagResult::Success;
+					}
+				}
 			}
 		default:
 			return EPropertyBagResult::TypeMismatch;
@@ -1144,11 +1211,21 @@
 		}
 		case EPropertyBagPropertyType::Enum:
 		{
-			const FEnumProperty* EnumProperty = CastFieldChecked<FEnumProperty>(Desc->CachedProperty);
-			const FNumericProperty* UnderlyingProperty = EnumProperty->GetUnderlyingProperty();
-			check(UnderlyingProperty);
-			UnderlyingProperty->SetIntPropertyValue(Address, InValue);
-			return EPropertyBagResult::Success;
+			if (const FEnumProperty* EnumProperty = CastField<FEnumProperty>(Desc->CachedProperty))
+			{
+				const FNumericProperty* UnderlyingProperty = EnumProperty->GetUnderlyingProperty();
+				check(UnderlyingProperty);
+				UnderlyingProperty->SetIntPropertyValue(Address, InValue);
+				return EPropertyBagResult::Success;
+			}
+			else if (const FByteProperty* ByteProperty = CastField<FByteProperty>(Desc->CachedProperty))
+			{
+				if (ByteProperty->Enum)
+				{
+					ByteProperty->SetPropertyValue(Address, static_cast<uint8>(InValue));
+					return EPropertyBagResult::Success;
+				}
+			}
 		}
 		default:
 			return EPropertyBagResult::TypeMismatch;
@@ -1222,11 +1299,21 @@
 			}
 		case EPropertyBagPropertyType::Enum:
 			{
-				const FEnumProperty* EnumProperty = CastFieldChecked<FEnumProperty>(Desc->CachedProperty);
-				const FNumericProperty* UnderlyingProperty = EnumProperty->GetUnderlyingProperty();
-				check(UnderlyingProperty);
-				UnderlyingProperty->SetIntPropertyValue(Address, static_cast<uint64>(InValue));
-				return EPropertyBagResult::Success;
+				if (const FEnumProperty* EnumProperty = CastField<FEnumProperty>(Desc->CachedProperty))
+				{
+					const FNumericProperty* UnderlyingProperty = EnumProperty->GetUnderlyingProperty();
+					check(UnderlyingProperty);
+					UnderlyingProperty->SetIntPropertyValue(Address, static_cast<uint64>(InValue));
+					return EPropertyBagResult::Success;
+				}
+				else if (const FByteProperty* ByteProperty = CastField<FByteProperty>(Desc->CachedProperty))
+				{
+					if (ByteProperty->Enum)
+					{
+						ByteProperty->SetPropertyValue(Address, static_cast<uint8>(FMath::RoundToInt32(InValue)));
+						return EPropertyBagResult::Success;
+					}
+				}
 			}
 		default:
 			return EPropertyBagResult::TypeMismatch;
@@ -1280,18 +1367,31 @@
 			return EPropertyBagResult::TypeMismatch;
 		}
 
-		const FEnumProperty* EnumProperty = CastFieldChecked<FEnumProperty>(Desc->CachedProperty);
-		const FNumericProperty* UnderlyingProperty = EnumProperty->GetUnderlyingProperty();
-		check(UnderlyingProperty);
+		if (const FEnumProperty* EnumProperty = CastField<FEnumProperty>(Desc->CachedProperty))
+		{
+			const FNumericProperty* UnderlyingProperty = EnumProperty->GetUnderlyingProperty();
+			check(UnderlyingProperty);
 
-		if (Enum != EnumProperty->GetEnum())
+			if (Enum != EnumProperty->GetEnum())
+			{
+				return EPropertyBagResult::TypeMismatch;
+			}
+	
+			UnderlyingProperty->SetIntPropertyValue(Address, static_cast<uint64>(InValue));
+	
+			return EPropertyBagResult::Success;
+		}
+		else if (const FByteProperty* ByteProperty = CastField<FByteProperty>(Desc->CachedProperty))
 		{
-			return EPropertyBagResult::TypeMismatch;
+			if (Enum != ByteProperty->Enum)
+			{
+				return EPropertyBagResult::TypeMismatch;
+			}
+			ByteProperty->SetPropertyValue(Address, InValue);
+			return EPropertyBagResult::Success;
 		}
-	
-		UnderlyingProperty->SetIntPropertyValue(Address, static_cast<uint64>(InValue));
-	
-		return EPropertyBagResult::Success;
+
+		return EPropertyBagResult::TypeMismatch;
 	}
 
 	EPropertyBagResult SetPropertyValueAsStruct(const FPropertyBagPropertyDesc* Desc, void* Address, const FConstStructView InValue)
@@ -4137,16 +4237,29 @@
 		return EPropertyBagResult::TypeMismatch;
 	}
 
-	const FEnumProperty* EnumProperty = CastFieldChecked<FEnumProperty>(ValueDesc.CachedProperty);
-	const FNumericProperty* UnderlyingProperty = EnumProperty->GetUnderlyingProperty();
-	check(UnderlyingProperty);
+	if (const FEnumProperty* EnumProperty = CastField<FEnumProperty>(ValueDesc.CachedProperty))
+	{
+		const FNumericProperty* UnderlyingProperty = EnumProperty->GetUnderlyingProperty();
+		check(UnderlyingProperty);
 
-	if (Enum != EnumProperty->GetEnum())
+		if (Enum != EnumProperty->GetEnum())
+		{
+			return EPropertyBagResult::TypeMismatch;
+		}
+
+		return Add(InValue);
+	}
+	else if (const FByteProperty* ByteProperty = CastField<FByteProperty>(ValueDesc.CachedProperty))
 	{
-		return EPropertyBagResult::TypeMismatch;
+		if (Enum != ByteProperty->Enum)
+		{
+			return EPropertyBagResult::TypeMismatch;
+		}
+
+		return Add(InValue);
 	}
 
-	return Add(InValue);
+	return EPropertyBagResult::TypeMismatch;
 }
 
 EPropertyBagResult FPropertyBagSetRef::AddValueStruct(FConstStructView InValue)
