[CoreRedirects]

Use a stable sort when sorting matching redirects to ensure we don't inadvertently re-order redirects with the same match score but different match flags.

#rnx

#Synced-CL 50155796

diff --git a/Engine/Source/Runtime/CoreUObject/Private/UObject/CoreRedirects.cpp b/Engine/Source/Runtime/CoreUObject/Private/UObject/CoreRedirects.cpp
--- a/Engine/Source/Runtime/CoreUObject/Private/UObject/CoreRedirects.cpp
+++ b/Engine/Source/Runtime/CoreUObject/Private/UObject/CoreRedirects.cpp
@@ -919,6 +919,12 @@ bool FCoreRedirectObjectName::ExpandNames(const FStringView InStringView, FName&
 
 	if (FirstPeriodIndex == INDEX_NONE)
 	{
+		// We cannot make FNames from strings of this size
+		if (FullStringView.Len() >= NAME_SIZE)
+		{
+			return false;
+		}
+
 		// If start with /, fill in package name, otherwise name
 		if (SlashIndex != INDEX_NONE)
 		{
@@ -942,20 +948,46 @@ bool FCoreRedirectObjectName::ExpandNames(const FStringView InStringView, FName&
 
 	if (SlashIndex == INDEX_NONE)
 	{
+		// We cannot make FNames from strings of this size
+		if (LastPeriodIndex >= NAME_SIZE)
+		{
+			return false;
+		}
+
 		// No /, so start from beginning. There must be an outer if we got this far
 		OutOuter = FName(FullStringView.Mid(0, LastPeriodIndex));
 	}
 	else
 	{
+		// We cannot make FNames from strings of this size
+		if (FirstPeriodIndex >= NAME_SIZE)
+		{
+			return false;
+		}
+
 		OutPackage = FName(FullStringView.Left(FirstPeriodIndex));
 		if (FirstPeriodIndex != LastPeriodIndex)
 		{
+			const int32 StartPos = FirstPeriodIndex + 1;
+			const int32 Count = LastPeriodIndex - FirstPeriodIndex - 1;
+			// We cannot make FNames from strings of this size
+			if (Count >= NAME_SIZE)
+			{
+				return false;
+			}
+
 			// Extract Outer between periods
-			OutOuter = FName(FullStringView.Mid(FirstPeriodIndex + 1, LastPeriodIndex - FirstPeriodIndex - 1));
+			OutOuter = FName(FullStringView.Mid(StartPos, Count));
 		}
 	}
 
-	OutName = FName(FullStringView.Mid(LastPeriodIndex + 1));
+	const int32 StartPos = LastPeriodIndex + 1;
+	// We cannot make FNames from strings of this size
+	if ((FullStringView.Len() - StartPos) >= NAME_SIZE)
+	{
+		return false;
+	}
+	OutName = FName(FullStringView.Mid(StartPos));
 
 	return true;
 }
@@ -2158,6 +2190,12 @@ bool FCoreRedirects::RunTests()
 	NewRedirects.Emplace(ECoreRedirectFlags::Type_Package | ECoreRedirectFlags::Option_MatchPrefix | ECoreRedirectFlags::Option_Iterate, TEXT("/oldoldoldprefix"), TEXT("/oldoldprefix"));
 	NewRedirects.Emplace(ECoreRedirectFlags::Type_Package | ECoreRedirectFlags::Option_MatchPrefix, TEXT("/noiterate_oldprefix"), TEXT("/oldprefix"));
 
+	// Test ordering of package redirects and other redirects to ensure the package redirect is applied last so we can still redirect non-package names that might be using the old package name
+	NewRedirects.Emplace(ECoreRedirectFlags::Type_Package,	TEXT("/Script/PackageRedirectOrderingTest_PackageRedirectFirst"),			TEXT("/Script/PackageRedirectOrderingTest_PackageRedirectFirst_New"));
+	NewRedirects.Emplace(ECoreRedirectFlags::Type_Struct,	TEXT("/Script/PackageRedirectOrderingTest_PackageRedirectFirst.OldStruct1"),	TEXT("/Script/PackageRedirectOrderingTest_PackageRedirectFirst_New.NewStruct1"));
+	NewRedirects.Emplace(ECoreRedirectFlags::Type_Struct, TEXT("/Script/PackageRedirectOrderingTest_PackageRedirectLast.OldStruct1"), TEXT("/Script/PackageRedirectOrderingTest_PackageRedirectLast_New.NewStruct1"));
+	NewRedirects.Emplace(ECoreRedirectFlags::Type_Package, TEXT("/Script/PackageRedirectOrderingTest_PackageRedirectLast"), TEXT("/Script/PackageRedirectOrderingTest_PackageRedirectLast_New"));
+
 	FCoreRedirect EnumRedirectWithValues(ECoreRedirectFlags::Type_Enum, TEXT("MaskTarget_PhysMesh"), TEXT("/Script/ClothingSystemRuntimeCommon.EWeightMapTargetCommon"));
 	EnumRedirectWithValues.ValueChanges.Add(TEXT("MaskTarget_PhysMesh::MaxDistance"), TEXT("EWeightMapTargetCommon::MaxDistance"));
 	EnumRedirectWithValues.ValueChanges.Add(TEXT("MaskTarget_PhysMesh::BackstopDistance"), TEXT("EWeightMapTargetCommon::BackstopDistance"));
@@ -2226,6 +2264,12 @@ bool FCoreRedirects::RunTests()
 	Tests.Emplace(TEXT("/Game/Package/古い接尾辞"), TEXT("/Game/Package/新しい接尾辞"), ECoreRedirectFlags::Type_Package);
 	Tests.Emplace(TEXT("/Game/Package/混合接尾辞"), TEXT("/Game/Package/mixed_suffix"), ECoreRedirectFlags::Type_Package);
 
+	// Redirect ordering tests
+	Tests.Emplace(TEXT("/Script/PackageRedirectOrderingTest_PackageRedirectFirst"), TEXT("/Script/PackageRedirectOrderingTest_PackageRedirectFirst_New"), ECoreRedirectFlags::Type_Package);
+	Tests.Emplace(TEXT("/Script/PackageRedirectOrderingTest_PackageRedirectFirst.OldStruct1"), TEXT("/Script/PackageRedirectOrderingTest_PackageRedirectFirst_New.NewStruct1"), ECoreRedirectFlags::Type_Struct);
+	Tests.Emplace(TEXT("/Script/PackageRedirectOrderingTest_PackageRedirectLast.OldStruct1"), TEXT("/Script/PackageRedirectOrderingTest_PackageRedirectLast_New.NewStruct1"), ECoreRedirectFlags::Type_Struct);
+	Tests.Emplace(TEXT("/Script/PackageRedirectOrderingTest_PackageRedirectLast"), TEXT("/Script/PackageRedirectOrderingTest_PackageRedirectLast_New"), ECoreRedirectFlags::Type_Package);
+
 	// Enum value tests
 	FRedirectTest EnumValueTest(TEXT("/Script/ClothingSystemRuntimeCommon.MaskTarget_PhysMesh"), TEXT("/Script/ClothingSystemRuntimeCommon.EWeightMapTargetCommon"), ECoreRedirectFlags::Type_Enum);
 	EnumValueTest.Expected.ValueChanges.Add(TEXT("MaskTarget_PhysMesh::MaxDistance"), TEXT("EWeightMapTargetCommon::MaxDistance"));
