[ZenLoader]
- Fix PackageRef that could exceed 16-bit under some circumstance.
- Remove ExportCount that was unused at runtime to avoid bumping the struct size

#Synced-CL 50160399

diff --git a/Engine/Source/Runtime/CoreUObject/Private/Serialization/AsyncLoading2.cpp b/Engine/Source/Runtime/CoreUObject/Private/Serialization/AsyncLoading2.cpp
--- a/Engine/Source/Runtime/CoreUObject/Private/Serialization/AsyncLoading2.cpp
+++ b/Engine/Source/Runtime/CoreUObject/Private/Serialization/AsyncLoading2.cpp
@@ -1414,9 +1414,12 @@ private:
 	FPublicExportMap PublicExportMap;
 	FName OriginalPackageName;
 	int32 PackageObjectIndex = -1;
-	uint16 PackageRefCount = 0;         // Package references during loading activity, should be reasonably low.
-	std::atomic<uint16> RefCount = 1;   // Only taken by thread activity, the max count should be proportional to the number of worker threads.
+	uint32 PackageRefCount = 0;         // Package references during loading activity. We've seen 16-bit not being enough, 24-bit should probably be enough if we need to compact it in the future.
+#if WITH_EDITOR
 	int32 ExportCount = -1;
+#endif
+	std::atomic<uint16> RefCount = 1;   // Only taken by thread activity, the max count should be proportional to the number of worker threads.
+
 	// This mutex protects all the fields of this class via a FLockedPackageRef at the exception of RefCount, 
 	// which is a typical reference counting using an atomic for synchronization.
 	FAsyncLoadingMutex Mutex;
@@ -1436,7 +1439,9 @@ private:
 		OriginalPackageName = FName();
 		PackageObjectIndex = -1;
 		PackageRefCount = 0;
+#if WITH_EDITOR
 		ExportCount = -1;
+#endif
 		PackageHeaderExtension = EPackageExtension::Unspecified;
 		PackageLoader = EPackageLoader::Unknown;
 		bWasFullyLoaded = false;
