InstancedStruct: allow explicit cast from TInstancedStruct<> to FInstancedStruct
* also added unrelated defensive check against nullptr, flagged by validation tool
#rnx

#Synced-CL 49557561

diff --git a/Engine/Source/Runtime/CoreUObject/Public/StructUtils/InstancedStruct.h b/Engine/Source/Runtime/CoreUObject/Public/StructUtils/InstancedStruct.h
--- a/Engine/Source/Runtime/CoreUObject/Public/StructUtils/InstancedStruct.h
+++ b/Engine/Source/Runtime/CoreUObject/Public/StructUtils/InstancedStruct.h
@@ -364,6 +364,7 @@ public:
 	/** Initializes from a raw struct type and optional data. */
 	void InitializeAsScriptStruct(const UScriptStruct* InScriptStruct, const uint8* InStructMemory = nullptr)
 	{
+		checkf(InScriptStruct, TEXT("ScriptStruct cannot be nullptr"));
 		checkf(InScriptStruct->IsChildOf(TBaseStructure<BaseStructT>::Get()), TEXT("ScriptStruct must be a child of BaseStruct!"));
 		InstancedStruct.InitializeAs(InScriptStruct, InStructMemory);
 	}
@@ -402,6 +403,22 @@ public:
 		return This;
 	}
 
+	/**
+	 * Cast directly to FInstancedStruct.
+	 * Intended for compatibility with functions that require FInstancedStruct.
+	 * Be sure you know what you're doing, as this will discard some type-checking safeguards!
+	 */
+	explicit operator const FInstancedStruct&() const &
+	{
+		return InstancedStruct;
+	}
+
+	/**
+	 * Deleted overload to prevent accidental dangling refs to temps, such as
+	 * const FInstancedStruct& BadRef = static_cast<const FInstancedStruct&>(TInstancedStruct<T>::Make());
+	 */
+	explicit operator const FInstancedStruct&() && = delete;
+
 	/** Returns struct type. */
 	const UScriptStruct* GetScriptStruct() const
 	{
