Don't merge PendingMigrationWorks with PendingWorks list. Changing work priority to reflect the different work types (migration vs normal) and prioritizing migration works

#Synced-CL 50156122

diff --git a/Engine/Source/Runtime/CoreUObject/Private/UObject/RemoteExecutor.cpp b/Engine/Source/Runtime/CoreUObject/Private/UObject/RemoteExecutor.cpp
--- a/Engine/Source/Runtime/CoreUObject/Private/UObject/RemoteExecutor.cpp
+++ b/Engine/Source/Runtime/CoreUObject/Private/UObject/RemoteExecutor.cpp
@@ -169,7 +169,7 @@ public:
 
 	void EnqueueWork(FName InWorkName, FRemoteWorkPriority InWorkPriority, bool bIsTransactional, int32 InWorkListIndex, TUniqueFunction<void(void)>&& InWork)
 	{
-		FRemoteExecutorWork& NewWork = PendingWorks.Emplace_GetRef();
+		FRemoteExecutorWork& NewWork = InWorkPriority.GetWorkType() == ERemoteWorkType::Normal ? PendingWorks.Emplace_GetRef() : PendingMigrationWorks.Emplace_GetRef();
 		NewWork.Name = InWorkName;
 		NewWork.Priority = InWorkPriority;
 		NewWork.Work = MoveTemp(InWork);
@@ -183,19 +183,6 @@ public:
 		}
 	}
 
-	void EnqueueMigrationWork(FName InWorkName, FRemoteWorkPriority InWorkPriority, int32 InWorkListIndex, TUniqueFunction<void(void)>&& InWork)
-	{
-		FRemoteExecutorWork& NewWork = PendingMigrationWorks.Emplace_GetRef();
-		NewWork.Name = InWorkName;
-		NewWork.Priority = InWorkPriority;
-		NewWork.Work = MoveTemp(InWork);
-		NewWork.bIsTransactional = true;
-		NewWork.RequestId = GenerateNextTransactionId();
-		NewWork.WorkListIndex = InWorkListIndex;
-
-		UE::RemoteExecutor::OnTransactionQueuedDelegate.Broadcast(NewWork.RequestId, NewWork.Name);
-	}
-
 	void EnqueueWorkList(FName WorkName, FRemoteWorkPriority InWorkPriority, bool bIsTransactional, TArray<TUniqueFunction<void(void)>>&& InWorkList)
 	{
 		FRemoteExecutorWorkList& NewWorkList = PendingWorkLists.Emplace_GetRef();
@@ -206,6 +193,49 @@ public:
 		NewWorkList.PendingWork = MoveTemp(InWorkList);
 	}
 
+private:
+
+	void EnqueueNextWorkFromList(int32 WorkListIndex)
+	{
+		if (WorkListIndex != INDEX_NONE)
+		{
+			FRemoteExecutorWorkList& WorkList = PendingWorkLists[WorkListIndex];
+			check(WorkList.bWorking);
+
+			if (WorkList.PendingWork.Num() > 0)
+			{
+				EnqueueWork(WorkList.Name, WorkList.Priority, WorkList.bIsTransactional, WorkListIndex, MoveTemp(WorkList.PendingWork[0]));
+				WorkList.PendingWork.RemoveAt(0, EAllowShrinking::No);
+			}
+		}
+	};
+
+	void ConditionallyPumpNetwork(bool bForcePumpNetwork)
+	{
+		if (GRemoteExecutorNewNetPumpBehavior)
+		{
+			// Always pump if we're trying this transaction multiple times
+			bool bUnconditionallyPumpNetwork = bForcePumpNetwork || bShouldPumpNetwork;
+
+			// if GRemoteExecutorForcedNetPumpIterationCount is nonzero, we may need to unconditionally pump the network
+			if (GRemoteExecutorForcedNetPumpIterationCount && !bUnconditionallyPumpNetwork)
+			{
+				IterationsSinceNetworkPump++;
+				bUnconditionallyPumpNetwork = (IterationsSinceNetworkPump % GRemoteExecutorForcedNetPumpIterationCount) == 0;
+			}
+
+			if (bUnconditionallyPumpNetwork)
+			{
+				// pump the network, reset our state variables
+				IterationsSinceNetworkPump = 0;
+				bShouldPumpNetwork = false;
+				UE::RemoteExecutor::TickNetworkDelegate.ExecuteIfBound();
+			}
+		}
+	};
+
+public:
+
 	void ExecutePendingWork()
 	{
 		check(ExecutingWork == nullptr);
@@ -235,45 +265,6 @@ public:
 				}
 			};
 
-		auto EnqueueNextWorkFromList = [this](int32 WorkListIndex) -> void
-			{
-				if (WorkListIndex != INDEX_NONE)
-				{
-					FRemoteExecutorWorkList& WorkList = PendingWorkLists[WorkListIndex];
-					check(WorkList.bWorking);
-
-					if (WorkList.PendingWork.Num() > 0)
-					{
-						EnqueueWork(WorkList.Name, WorkList.Priority, WorkList.bIsTransactional, WorkListIndex, MoveTemp(WorkList.PendingWork[0]));
-						WorkList.PendingWork.RemoveAt(0, EAllowShrinking::No);
-					}
-				}
-			};
-
-		auto ConditionallyPumpNetwork = [this](bool bForcePumpNetwork)->void
-			{
-				if (GRemoteExecutorNewNetPumpBehavior)
-				{
-					// Always pump if we're trying this transaction multiple times
-					bool bUnconditionallyPumpNetwork = bForcePumpNetwork || bShouldPumpNetwork;
-			
-					// if GRemoteExecutorForcedNetPumpIterationCount is nonzero, we may need to unconditionally pump the network
-					if (GRemoteExecutorForcedNetPumpIterationCount && !bUnconditionallyPumpNetwork)
-					{
-						IterationsSinceNetworkPump++;
-						bUnconditionallyPumpNetwork = (IterationsSinceNetworkPump % GRemoteExecutorForcedNetPumpIterationCount) == 0;
-					}
-
-					if (bUnconditionallyPumpNetwork)
-					{
-						// pump the network, reset our state variables
-						IterationsSinceNetworkPump = 0;
-						bShouldPumpNetwork = false;
-						UE::RemoteExecutor::TickNetworkDelegate.ExecuteIfBound();
-					}
-				}
-			};
-
 		double LastStallPrintTime = FPlatformTime::Seconds();
 
 		for (int32 LocalIterationNumber = 0; ; LocalIterationNumber++)
@@ -335,13 +326,6 @@ public:
 				continue;
 			}
 
-			if (PendingMigrationWorks.Num())
-			{
-				// Prepend the Pending Migrations to the PendingWorks list
-				PendingMigrationWorks += MoveTemp(PendingWorks);
-				PendingWorks = MoveTemp(PendingMigrationWorks);
-			}
-
 			// enqueue the first work item from any new work lists
 			for (int32 WorkListIndex = 0; WorkListIndex < PendingWorkLists.Num(); WorkListIndex++)
 			{
@@ -358,7 +342,7 @@ public:
 			//
 			// if we get here and there is no work pending, break out and finish
 			//
-			if (PendingWorks.Num() == 0)
+			if (PendingWorks.Num() == 0 && PendingMigrationWorks.Num() == 0)
 			{
 				break;
 			}
@@ -368,160 +352,175 @@ public:
 				const bool bForcePumpNetwork = (LocalIterationNumber > 0);
 				ConditionallyPumpNetwork(bForcePumpNetwork);
 			}
-			
-			for (int32 PendingWorkIndex = 0; PendingWorkIndex < PendingWorks.Num(); PendingWorkIndex++)
-			{
-				FRemoteExecutorWork& PendingWork = PendingWorks[PendingWorkIndex];
-
-				UE_LOG(LogRemoteExec, VeryVerbose, TEXT("ExecutePendingWork[%d] : PendingWork[%d] '%s' RequestId %s %s"),
-					LocalIterationNumber,
-					PendingWorkIndex,
-					*PendingWork.Name.ToString(),
-					*PendingWork.RequestId.ToString(),
-					*PendingWork.Priority.ToString());
-			}
 
-			// round robin through executing all pending work
-			for(int32 PendingWorkIndex = 0; PendingWorkIndex < PendingWorks.Num(); )
-			{
-				check(ExecutingWork == nullptr);
-				TGuardValue<FRemoteExecutorWork*> ExecutingWorkSave(ExecutingWork, &PendingWorks[PendingWorkIndex]);
-				check(ExecutingWork);
+			ExecutePendingWork(PendingMigrationWorks, LocalIterationNumber, LastStallPrintTime);
+			ExecutePendingWork(PendingWorks, LocalIterationNumber, LastStallPrintTime);
+		}
 
-				UE_LOG(LogRemoteExec, VeryVerbose, TEXT("ExecutePendingWork[%d,%s] : Executing request %s %s --"),
-					LocalIterationNumber,
-					*ExecutingWork->RequestId.ToString(),
-					*ExecutingWork->Name.ToString(),
-					*ExecutingWork->Priority.ToString()); // -V522
+		PendingWorkLists.Empty();
+		AbortReason = UE::RemoteExecutor::ERemoteExecutorAbortReason::Unspecified;
+		AbortReasonDescription.Reset();
+	}
 
-				if (ExecutingWork->bIsTransactional)
-				{
-					if (!ExecutingWork->bRequestCreated)
-					{
-						// new request
-						ExecutingWork->bRequestCreated = true;
+private:
 
-						for (int32 SubsystemIndex = 0; SubsystemIndex < Subsystems.Num(); SubsystemIndex++)
-						{
-							FRemoteSubsystemBase* Subsystem = Subsystems[SubsystemIndex];
-							UE_LOG(LogRemoteExec, VeryVerbose, TEXT("ExecutePendingWork[%d,%s] : Creating request[%s]"),
-								LocalIterationNumber,
-								*ExecutingWork->RequestId.ToString(),
-								Subsystem->NameForDebug());
-							Subsystem->CreateRequest(ExecutingWork->RequestId, ExecutingWork->Priority);
+	void ExecutePendingWork(TArray<FRemoteExecutorWork>& WorksToExecute, int32 LocalIterationNumber, double& LastStallPrintTime)
+	{
+		for (int32 PendingWorkIndex = 0; PendingWorkIndex < WorksToExecute.Num(); PendingWorkIndex++)
+		{
+			FRemoteExecutorWork& PendingWork = WorksToExecute[PendingWorkIndex];
+
+			UE_LOG(LogRemoteExec, VeryVerbose, TEXT("ExecutePendingWork[%d] : PendingWork[%d] '%s' RequestId %s %s"),
+				LocalIterationNumber,
+				PendingWorkIndex,
+				*PendingWork.Name.ToString(),
+				*PendingWork.RequestId.ToString(),
+				*PendingWork.Priority.ToString());
+		}
 
-							Subsystem->SetActiveRequest(ExecutingWork->RequestId);
+		bool bPrintStall = false;
 
-							Subsystem->BeginRequest();
-						}
-					}
-					else
-					{
-						for (int32 SubsystemIndex = 0; SubsystemIndex < Subsystems.Num(); SubsystemIndex++)
-						{
-							FRemoteSubsystemBase* Subsystem = Subsystems[SubsystemIndex];
-							Subsystem->SetActiveRequest(ExecutingWork->RequestId);
-						}
-					}
+		// round robin through executing all pending work
+		for (int32 PendingWorkIndex = 0; PendingWorkIndex < WorksToExecute.Num(); )
+		{
+			check(ExecutingWork == nullptr);
+			TGuardValue<FRemoteExecutorWork*> ExecutingWorkSave(ExecutingWork, &WorksToExecute[PendingWorkIndex]);
+			check(ExecutingWork);
 
-					Now = FPlatformTime::Seconds();
-					if ((Now - LastStallPrintTime) > 1.0)
-					{
-						LastStallPrintTime = Now;
-						bPrintStall = true;
-					}
+			UE_LOG(LogRemoteExec, VeryVerbose, TEXT("ExecutePendingWork[%d,%s] : Executing request %s %s --"),
+				LocalIterationNumber,
+				*ExecutingWork->RequestId.ToString(),
+				*ExecutingWork->Name.ToString(),
+				*ExecutingWork->Priority.ToString()); // -V522
+
+			if (ExecutingWork->bIsTransactional)
+			{
+				if (!ExecutingWork->bRequestCreated)
+				{
+					// new request
+					ExecutingWork->bRequestCreated = true;
 
-					// tick the subsystems
 					for (int32 SubsystemIndex = 0; SubsystemIndex < Subsystems.Num(); SubsystemIndex++)
 					{
 						FRemoteSubsystemBase* Subsystem = Subsystems[SubsystemIndex];
-						Subsystem->TickRequest();
-					}
+						UE_LOG(LogRemoteExec, VeryVerbose, TEXT("ExecutePendingWork[%d,%s] : Creating request[%s]"),
+							LocalIterationNumber,
+							*ExecutingWork->RequestId.ToString(),
+							Subsystem->NameForDebug());
+						Subsystem->CreateRequest(ExecutingWork->RequestId, ExecutingWork->Priority);
 
-					int32 SubsystemReadyCount = 0;
+						Subsystem->SetActiveRequest(ExecutingWork->RequestId);
+
+						Subsystem->BeginRequest();
+					}
+				}
+				else
+				{
 					for (int32 SubsystemIndex = 0; SubsystemIndex < Subsystems.Num(); SubsystemIndex++)
 					{
 						FRemoteSubsystemBase* Subsystem = Subsystems[SubsystemIndex];
-						if (Subsystem->AreDependenciesSatisfied())
-						{
-							SubsystemReadyCount++;
-						}
-						else
-						{
-							if (bPrintStall)
-							{
-								UE_LOG(LogRemoteExec, VeryVerbose, TEXT("ExecutePendingWork[%d,%s] : subsystem %s not ready..."),
-									LocalIterationNumber,
-									*ExecutingWork->RequestId.ToString(),
-									Subsystem->NameForDebug());
-							}
-						}
+						Subsystem->SetActiveRequest(ExecutingWork->RequestId);
 					}
+				}
+
+				double Now = FPlatformTime::Seconds();
+				if ((Now - LastStallPrintTime) > 1.0)
+				{
+					LastStallPrintTime = Now;
+					bPrintStall = true;
+				}
 
-					if (SubsystemReadyCount != Subsystems.Num())
+				// tick the subsystems
+				for (int32 SubsystemIndex = 0; SubsystemIndex < Subsystems.Num(); SubsystemIndex++)
+				{
+					FRemoteSubsystemBase* Subsystem = Subsystems[SubsystemIndex];
+					Subsystem->TickRequest();
+				}
+
+				int32 SubsystemReadyCount = 0;
+				for (int32 SubsystemIndex = 0; SubsystemIndex < Subsystems.Num(); SubsystemIndex++)
+				{
+					FRemoteSubsystemBase* Subsystem = Subsystems[SubsystemIndex];
+					if (Subsystem->AreDependenciesSatisfied())
+					{
+						SubsystemReadyCount++;
+					}
+					else
 					{
 						if (bPrintStall)
 						{
-							UE_LOG(LogRemoteExec, VeryVerbose, TEXT("ExecutePendingWork[%d,%s] : %d of %d subsystems not ready..."),
+							UE_LOG(LogRemoteExec, VeryVerbose, TEXT("ExecutePendingWork[%d,%s] : subsystem %s not ready..."),
 								LocalIterationNumber,
 								*ExecutingWork->RequestId.ToString(),
-								Subsystems.Num() - SubsystemReadyCount, Subsystems.Num());
+								Subsystem->NameForDebug());
 						}
+					}
+				}
 
-						PendingWorkIndex++;
-						continue;
+				if (SubsystemReadyCount != Subsystems.Num())
+				{
+					if (bPrintStall)
+					{
+						UE_LOG(LogRemoteExec, VeryVerbose, TEXT("ExecutePendingWork[%d,%s] : %d of %d subsystems not ready..."),
+							LocalIterationNumber,
+							*ExecutingWork->RequestId.ToString(),
+							Subsystems.Num() - SubsystemReadyCount, Subsystems.Num());
 					}
 
-					// all of the subsystems are ready, try to perform the work
-					TRACE_CPUPROFILER_EVENT_SCOPE_TEXT(ExecutingWork->Name)
+					PendingWorkIndex++;
+					continue;
+				}
+
+				// all of the subsystems are ready, try to perform the work
+				TRACE_CPUPROFILER_EVENT_SCOPE_TEXT(ExecutingWork->Name)
 
 					FRemoteTransactionId RequestId = ExecutingWork->RequestId;
-					UE::RemoteExecutor::OnTransactionStartingDelegate.Broadcast(RequestId, ExecutingWork->Name);
-					AutoRTFM::ETransactionResult WorkTransactResult =
-						AutoRTFM::Transact(
-							[this, LocalIterationNumber]()
+				UE::RemoteExecutor::OnTransactionStartingDelegate.Broadcast(RequestId, ExecutingWork->Name);
+				AutoRTFM::ETransactionResult WorkTransactResult =
+					AutoRTFM::Transact(
+						[this, LocalIterationNumber]()
+						{
+							UE_AUTORTFM_OPEN
 							{
-								UE_AUTORTFM_OPEN
-								{
-									AbortReason = UE::RemoteExecutor::ERemoteExecutorAbortReason::Unspecified;
-									AbortReasonDescription.Reset();
-									ExecutingWork->NumAbortsRequested = 0;
+								AbortReason = UE::RemoteExecutor::ERemoteExecutorAbortReason::Unspecified;
+								AbortReasonDescription.Reset();
+								ExecutingWork->NumAbortsRequested = 0;
 
-									ExecutingWork->ExecutionAttempts++;
-								};
-
-								// You can set a breakpoint inside the abort handler below to see what caused the abort,
-								// or you can rely on the log output which only catches livelocks
-								const uint32 ExecutionAttempts = ExecutingWork->ExecutionAttempts;
-								UE_AUTORTFM_ONABORT(ExecutionAttempts, this)
-								{
-									UE_LOG(LogRemoteExec, VeryVerbose, TEXT("Transaction %s Aborted: %s: %s"), *ExecutingWork->Name.ToString(), EnumToString(AbortReason), *AbortReasonDescription);
+								ExecutingWork->ExecutionAttempts++;
+							};
 
-									constexpr uint32 MaxExecutionAttemptsBeforeLivelockWarning = 500;
-									if ((ExecutionAttempts >= MaxExecutionAttemptsBeforeLivelockWarning) && (ExecutionAttempts % MaxExecutionAttemptsBeforeLivelockWarning) == 0)
-									{
-										UE_LOG(LogRemoteExec, Display, TEXT("vvv Transaction Aborted %u Times. Dumping Callstack. vvv"), ExecutionAttempts);
-										FDebug::DumpStackTraceToLog(ELogVerbosity::Display);
-										UE_LOG(LogRemoteExec, Display, TEXT("^^^ Transaction Aborted %u Times ^^^"), ExecutionAttempts);
-									}
-								};
-
-								ExecutingWork->Work();
+							// You can set a breakpoint inside the abort handler below to see what caused the abort,
+							// or you can rely on the log output which only catches livelocks
+							const uint32 ExecutionAttempts = ExecutingWork->ExecutionAttempts;
+							UE_AUTORTFM_ONABORT(ExecutionAttempts, this)
+							{
+								UE_LOG(LogRemoteExec, VeryVerbose, TEXT("Transaction %s Aborted: %s: %s"), *ExecutingWork->Name.ToString(), EnumToString(AbortReason), *AbortReasonDescription);
 
-								// Did we do something that should late-abort this transaction?
-								if (ExecutingWork->NumAbortsRequested > 0)
+								constexpr uint32 MaxExecutionAttemptsBeforeLivelockWarning = 500;
+								if ((ExecutionAttempts >= MaxExecutionAttemptsBeforeLivelockWarning) && (ExecutionAttempts % MaxExecutionAttemptsBeforeLivelockWarning) == 0)
 								{
-									UE_LOG(LogRemoteExec, Warning, TEXT("We ended a transaction marked for abort. We expected to abort before this point (%d abort requests)."), ExecutingWork->NumAbortsRequested);
-									AutoRTFM::AbortTransaction();
-								}
-								else
-								{
-									ensureMsgf(AbortReason == UE::RemoteExecutor::ERemoteExecutorAbortReason::Unspecified && AbortReasonDescription.IsEmpty(), TEXT("RemoteExecutor: AbortReason %s specified during transaction without aborting"), *AbortReasonDescription);
+									UE_LOG(LogRemoteExec, Display, TEXT("vvv Transaction Aborted %u Times. Dumping Callstack. vvv"), ExecutionAttempts);
+									FDebug::DumpStackTraceToLog(ELogVerbosity::Display);
+									UE_LOG(LogRemoteExec, Display, TEXT("^^^ Transaction Aborted %u Times ^^^"), ExecutionAttempts);
 								}
+							};
 
-								if (ExecutingWork->bRequiresMultiServerCommit)
-								{
-									TRACE_CPUPROFILER_EVENT_SCOPE(MultiServerCommit)
+							ExecutingWork->Work();
+
+							// Did we do something that should late-abort this transaction?
+							if (ExecutingWork->NumAbortsRequested > 0)
+							{
+								UE_LOG(LogRemoteExec, Warning, TEXT("We ended a transaction marked for abort. We expected to abort before this point (%d abort requests)."), ExecutingWork->NumAbortsRequested);
+								AutoRTFM::AbortTransaction();
+							}
+							else
+							{
+								ensureMsgf(AbortReason == UE::RemoteExecutor::ERemoteExecutorAbortReason::Unspecified && AbortReasonDescription.IsEmpty(), TEXT("RemoteExecutor: AbortReason %s specified during transaction without aborting"), *AbortReasonDescription);
+							}
+
+							if (ExecutingWork->bRequiresMultiServerCommit)
+							{
+								TRACE_CPUPROFILER_EVENT_SCOPE(MultiServerCommit)
 
 									// we're done with the work and about to commit - 
 									// first we need to send borrowed objects back
@@ -532,296 +531,291 @@ public:
 										*ExecutingWork->RequestId.ToString(),
 										*ExecutingWork->RequiresMultiServerCommitReason);
 
-									check (!MultiServerCommitRequestId.IsValid());
-									check (MultiServerCommitReadyServers.Num() == 0);
-									check (bMultiServerCommitRequiresAbort == false);
+								check(!MultiServerCommitRequestId.IsValid());
+								check(MultiServerCommitReadyServers.Num() == 0);
+								check(bMultiServerCommitRequiresAbort == false);
 
-									UE_AUTORTFM_OPEN
-									{
-										MultiServerCommitRequestId = ExecutingWork->RequestId;
-										bMultiServerCommitRequiresAbort = false;
-									};
+								UE_AUTORTFM_OPEN
+								{
+									MultiServerCommitRequestId = ExecutingWork->RequestId;
+									bMultiServerCommitRequiresAbort = false;
+								};
 
-									// first collect the list of servers that need to be involved
-									TArray<FRemoteServerId> MultiServerCommitRemoteServers;
-									for (int32 SubsystemIndex = 0; SubsystemIndex < Subsystems.Num(); SubsystemIndex++)
-									{
-										UE_LOG(LogRemoteExec, VeryVerbose, TEXT("ExecutePendingWork[%d,%s] : BeginMultiServerCommit subsystem %d..."),
-											LocalIterationNumber,
-											*ExecutingWork->RequestId.ToString(),
-											SubsystemIndex);
+								// first collect the list of servers that need to be involved
+								TArray<FRemoteServerId> MultiServerCommitRemoteServers;
+								for (int32 SubsystemIndex = 0; SubsystemIndex < Subsystems.Num(); SubsystemIndex++)
+								{
+									UE_LOG(LogRemoteExec, VeryVerbose, TEXT("ExecutePendingWork[%d,%s] : BeginMultiServerCommit subsystem %d..."),
+										LocalIterationNumber,
+										*ExecutingWork->RequestId.ToString(),
+										SubsystemIndex);
 
-										FRemoteSubsystemBase* Subsystem = Subsystems[SubsystemIndex];
-										Subsystem->BeginMultiServerCommit(MultiServerCommitRemoteServers);
-									}
+									FRemoteSubsystemBase* Subsystem = Subsystems[SubsystemIndex];
+									Subsystem->BeginMultiServerCommit(MultiServerCommitRemoteServers);
+								}
+
+								UE_LOG(LogRemoteExec, VeryVerbose, TEXT("ExecutePendingWork[%d,%s] : BeginMultiServerCommit DONE (%d servers)"),
+									LocalIterationNumber,
+									*ExecutingWork->RequestId.ToString(),
+									MultiServerCommitRemoteServers.Num());
 
-									UE_LOG(LogRemoteExec, VeryVerbose, TEXT("ExecutePendingWork[%d,%s] : BeginMultiServerCommit DONE (%d servers)"),
+								for (FRemoteServerId ServerId : MultiServerCommitRemoteServers)
+								{
+									UE_LOG(LogRemoteExec, VeryVerbose, TEXT("ExecutePendingWork[%d,%s] : BeginMultiServerCommit server %s"),
 										LocalIterationNumber,
 										*ExecutingWork->RequestId.ToString(),
-										MultiServerCommitRemoteServers.Num());
+										*ServerId.ToString());
+								}
 
-									for (FRemoteServerId ServerId : MultiServerCommitRemoteServers)
-									{
-										UE_LOG(LogRemoteExec, VeryVerbose, TEXT("ExecutePendingWork[%d,%s] : BeginMultiServerCommit server %s"),
-											LocalIterationNumber,
-											*ExecutingWork->RequestId.ToString(),
-											*ServerId.ToString());
-									}
+								// signal to the relevant servers that they need to help us do a multi-server commit
+								UE_AUTORTFM_OPEN
+								{
+									UE::RemoteExecutor::BeginMultiServerCommitDelegate.ExecuteIfBound(ExecutingWork->RequestId, ExecutingWork->Priority, MultiServerCommitRemoteServers);
+								};
 
-									// signal to the relevant servers that they need to help us do a multi-server commit
-									UE_AUTORTFM_OPEN
-									{
-										UE::RemoteExecutor::BeginMultiServerCommitDelegate.ExecuteIfBound(ExecutingWork->RequestId, ExecutingWork->Priority, MultiServerCommitRemoteServers);
-									};
+								// ask each subsystem to send any necessary data as part of this commit
+								for (int32 SubsystemIndex = 0; SubsystemIndex < Subsystems.Num(); SubsystemIndex++)
+								{
+									UE_LOG(LogRemoteExec, VeryVerbose, TEXT("ExecutePendingWork[%d,%s] : ExecuteMultiServerCommit subsystem %d..."),
+										LocalIterationNumber,
+										*ExecutingWork->RequestId.ToString(),
+										SubsystemIndex);
 
-									// ask each subsystem to send any necessary data as part of this commit
-									for (int32 SubsystemIndex = 0; SubsystemIndex < Subsystems.Num(); SubsystemIndex++)
+									FRemoteSubsystemBase* Subsystem = Subsystems[SubsystemIndex];
+									Subsystem->ExecuteMultiServerCommit();
+								}
+
+								UE_AUTORTFM_OPEN
+								{
+									// tell each server that we're done sending commit data and we are waiting for them
+									// to respond with whether they are ready or we need to abort and retry
+									UE::RemoteExecutor::ReadyMultiServerCommitDelegate.ExecuteIfBound(ExecutingWork->RequestId, MultiServerCommitRemoteServers);
+								};
+
+								// tick the network until everything is ready
+								double LastPrintTime = FPlatformTime::Seconds();
+
+								for (;;)
+								{
+									double Now = FPlatformTime::Seconds();
+									if ((Now - LastPrintTime) > 1.0)
 									{
-										UE_LOG(LogRemoteExec, VeryVerbose, TEXT("ExecutePendingWork[%d,%s] : ExecuteMultiServerCommit subsystem %d..."),
+										LastPrintTime = Now;
+										UE_LOG(LogRemoteExec, VeryVerbose, TEXT("ExecutePendingWork[%d,%s] : WAITING for multi-server commit..."),
 											LocalIterationNumber,
-											*ExecutingWork->RequestId.ToString(),
-											SubsystemIndex);
-
-										FRemoteSubsystemBase* Subsystem = Subsystems[SubsystemIndex];
-										Subsystem->ExecuteMultiServerCommit();
+											*ExecutingWork->RequestId.ToString());
 									}
 
-									UE_AUTORTFM_OPEN
+									if (MultiServerCommitRemoteServers.Num() == MultiServerCommitReadyServers.Num())
 									{
-										// tell each server that we're done sending commit data and we are waiting for them
-										// to respond with whether they are ready or we need to abort and retry
-										UE::RemoteExecutor::ReadyMultiServerCommitDelegate.ExecuteIfBound(ExecutingWork->RequestId, MultiServerCommitRemoteServers);
-									};
+										// we got a response from each server, do a sanity check to ensure the list of ready
+										// servers is the list we expected to have										
+										UE_AUTORTFM_OPEN
+										{
+											MultiServerCommitRemoteServers.Sort();
+											MultiServerCommitReadyServers.Sort();
+										};
 
-									// tick the network until everything is ready
-									double LastPrintTime = FPlatformTime::Seconds();
+										for (int32 ServerIndex = 0; ServerIndex < MultiServerCommitRemoteServers.Num(); ServerIndex++)
+										{
+											FRemoteServerId ServerId = MultiServerCommitRemoteServers[ServerIndex];
+											FRemoteServerId ReadyServerId = MultiServerCommitReadyServers[ServerIndex];
 
-									for (;;)
-									{
-										double Now = FPlatformTime::Seconds();
-										if ((Now - LastPrintTime) > 1.0)
+											if (ServerId != ReadyServerId)
+											{
+												UE_LOG(LogRemoteExec, VeryVerbose, TEXT("ExecutePendingWork[%d,%s] : Multi-server commit expected server %s but got %s"),
+													LocalIterationNumber,
+													*ExecutingWork->RequestId.ToString(),
+													*ServerId.ToString(),
+													*ReadyServerId.ToString());
+											}
+
+											check(ServerId == ReadyServerId);
+										}
+
+										if (bMultiServerCommitRequiresAbort)
 										{
-											LastPrintTime = Now;
-											UE_LOG(LogRemoteExec, VeryVerbose, TEXT("ExecutePendingWork[%d,%s] : WAITING for multi-server commit..."),
+											UE_LOG(LogRemoteExec, Verbose, TEXT("ExecutePendingWork[%d,%s] : Multi-server commit ALL servers READY, but we are flagged with ABORT"),
 												LocalIterationNumber,
 												*ExecutingWork->RequestId.ToString());
-										}
-										
-										if (MultiServerCommitRemoteServers.Num() == MultiServerCommitReadyServers.Num())
-										{
-											// we got a response from each server, do a sanity check to ensure the list of ready
-											// servers is the list we expected to have										
-											UE_AUTORTFM_OPEN
-											{
-												MultiServerCommitRemoteServers.Sort();
-												MultiServerCommitReadyServers.Sort();
-											};
 
-											for (int32 ServerIndex = 0; ServerIndex < MultiServerCommitRemoteServers.Num(); ServerIndex++)
+											UE_AUTORTFM_OPEN
 											{
-												FRemoteServerId ServerId = MultiServerCommitRemoteServers[ServerIndex];
-												FRemoteServerId ReadyServerId = MultiServerCommitReadyServers[ServerIndex];
-
-												if (ServerId != ReadyServerId)
+												for (int32 SubsystemIndex = 0; SubsystemIndex < Subsystems.Num(); SubsystemIndex++)
 												{
-													UE_LOG(LogRemoteExec, VeryVerbose, TEXT("ExecutePendingWork[%d,%s] : Multi-server commit expected server %s but got %s"),
-														LocalIterationNumber,
-														*ExecutingWork->RequestId.ToString(),
-														*ServerId.ToString(),
-														*ReadyServerId.ToString());
+													FRemoteSubsystemBase* Subsystem = Subsystems[SubsystemIndex];
+													Subsystem->AbortMultiServerCommit();
 												}
 
-												check(ServerId == ReadyServerId);
-											}
+											// aborted, tell all servers to abandon this commit
+											UE::RemoteExecutor::AbandonMultiServerCommitDelegate.ExecuteIfBound(ExecutingWork->RequestId, MultiServerCommitRemoteServers);
 
-											if (bMultiServerCommitRequiresAbort)
-											{
-												UE_LOG(LogRemoteExec, Verbose, TEXT("ExecutePendingWork[%d,%s] : Multi-server commit ALL servers READY, but we are flagged with ABORT"),
-													LocalIterationNumber,
-													*ExecutingWork->RequestId.ToString());
+											MultiServerCommitRequestId = FRemoteTransactionId::Invalid();
+											MultiServerCommitReadyServers.Reset();
+											bMultiServerCommitRequiresAbort = false;
+											};
 
-												UE_AUTORTFM_OPEN
-												{
-													for (int32 SubsystemIndex = 0; SubsystemIndex < Subsystems.Num(); SubsystemIndex++)
-													{
-														FRemoteSubsystemBase* Subsystem = Subsystems[SubsystemIndex];
-														Subsystem->AbortMultiServerCommit();
-													}
-
-													// aborted, tell all servers to abandon this commit
-													UE::RemoteExecutor::AbandonMultiServerCommitDelegate.ExecuteIfBound(ExecutingWork->RequestId, MultiServerCommitRemoteServers);
-
-													MultiServerCommitRequestId = FRemoteTransactionId::Invalid();
-													MultiServerCommitReadyServers.Reset();
-													bMultiServerCommitRequiresAbort = false;
-												};
-												
-												using namespace UE::RemoteExecutor;
-												AbortCurrentTransaction(EAbortTransactionBehavior::Immediately, ERemoteExecutorAbortReason::RequiresDependencies, TEXT("bMultiServerCommitRequiresAbort"));
-
-												return;
-											}
-											else
-											{
-												// the commit was accepted by every server, notify each subsystem that we are committing
-												UE_LOG(LogRemoteExec, VeryVerbose, TEXT("ExecutePendingWork[%d,%s] : Multi-server commit ALL servers READY, COMMITTING"),
-													LocalIterationNumber,
-													*ExecutingWork->RequestId.ToString());
+											using namespace UE::RemoteExecutor;
+											AbortCurrentTransaction(EAbortTransactionBehavior::Immediately, ERemoteExecutorAbortReason::RequiresDependencies, TEXT("bMultiServerCommitRequiresAbort"));
 
-												for (int32 SubsystemIndex = 0; SubsystemIndex < Subsystems.Num(); SubsystemIndex++)
-												{
-													FRemoteSubsystemBase* Subsystem = Subsystems[SubsystemIndex];
-													Subsystem->CommitMultiServerCommit();
-												}
+											return;
+										}
+										else
+										{
+											// the commit was accepted by every server, notify each subsystem that we are committing
+											UE_LOG(LogRemoteExec, VeryVerbose, TEXT("ExecutePendingWork[%d,%s] : Multi-server commit ALL servers READY, COMMITTING"),
+												LocalIterationNumber,
+												*ExecutingWork->RequestId.ToString());
 
-												break;
+											for (int32 SubsystemIndex = 0; SubsystemIndex < Subsystems.Num(); SubsystemIndex++)
+											{
+												FRemoteSubsystemBase* Subsystem = Subsystems[SubsystemIndex];
+												Subsystem->CommitMultiServerCommit();
 											}
-										}
 
-										UE_AUTORTFM_OPEN
-										{
-											UE::RemoteExecutor::TickNetworkDelegate.ExecuteIfBound();
-										};
+											break;
+										}
 									}
 
-									UE_LOG(LogRemoteExec, VeryVerbose, TEXT("ExecutePendingWork[%d,%s] : DONE with multi-server commit"),
-										LocalIterationNumber,
-										*ExecutingWork->RequestId.ToString());
-
 									UE_AUTORTFM_OPEN
 									{
-										MultiServerCommitRequestId = FRemoteTransactionId::Invalid();
-										MultiServerCommitReadyServers.Reset();
-										bMultiServerCommitRequiresAbort = false;
+										UE::RemoteExecutor::TickNetworkDelegate.ExecuteIfBound();
 									};
+								}
 
-									UE_AUTORTFM_OPEN
-									{
-										// tell each server that we are committed and they should commit
-										UE::RemoteExecutor::EndMultiServerCommitDelegate.ExecuteIfBound(ExecutingWork->RequestId, MultiServerCommitRemoteServers);
-									};
+								UE_LOG(LogRemoteExec, VeryVerbose, TEXT("ExecutePendingWork[%d,%s] : DONE with multi-server commit"),
+									LocalIterationNumber,
+									*ExecutingWork->RequestId.ToString());
+
+								UE_AUTORTFM_OPEN
+								{
+									MultiServerCommitRequestId = FRemoteTransactionId::Invalid();
+									MultiServerCommitReadyServers.Reset();
+									bMultiServerCommitRequiresAbort = false;
 								};
-							});
 
-					bool bWorkComplete = false;
-					bool bWorkTransactionAborted = false;
+								UE_AUTORTFM_OPEN
+								{
+									// tell each server that we are committed and they should commit
+									UE::RemoteExecutor::EndMultiServerCommitDelegate.ExecuteIfBound(ExecutingWork->RequestId, MultiServerCommitRemoteServers);
+								};
+							};
+						});
 
-					if (WorkTransactResult == AutoRTFM::ETransactionResult::AbortedByRequest)
-					{
-						using namespace UE::RemoteExecutor;
+				bool bWorkComplete = false;
+				bool bWorkTransactionAborted = false;
 
-						bWorkTransactionAborted = true;
+				if (WorkTransactResult == AutoRTFM::ETransactionResult::AbortedByRequest)
+				{
+					using namespace UE::RemoteExecutor;
 
-						OnTransactionAbortedDelegate.Broadcast(RequestId, ExecutingWork->ExecutionAttempts, AbortReasonDescription);
+					bWorkTransactionAborted = true;
 
-						if (AbortReason == ERemoteExecutorAbortReason::RequiresDependencies) // -V547
-						{
-							UE_LOG(LogRemoteExec, Verbose, TEXT("ExecutePendingWork[%d,%s] : Work ABORTED for required dependency '%s'"),
-								LocalIterationNumber,
-								*ExecutingWork->RequestId.ToString(),
-								*AbortReasonDescription);
-						}
-						else if (AbortReason == ERemoteExecutorAbortReason::AbandonWork) // -V547
-						{
-							UE_LOG(LogRemoteExec, Verbose, TEXT("ExecutePendingWork[%d,%s] : Work ABORTED for ABANDONMENT '%s'"),
-								LocalIterationNumber,
-								*ExecutingWork->RequestId.ToString(),
-								*AbortReasonDescription);
-							bWorkComplete = true;
-						}
-						else
-						{
-							UE_LOG(LogRemoteExec, Verbose, TEXT("ExecutePendingWork[%d,%s] : Work ABORTED for unknown reason '%s'"),
-								LocalIterationNumber,
-								*ExecutingWork->RequestId.ToString(),
-								*AbortReasonDescription);
-						}
+					OnTransactionAbortedDelegate.Broadcast(RequestId, ExecutingWork->ExecutionAttempts, AbortReasonDescription);
 
-						if (!bWorkComplete)
-						{
-							for (int32 SubsystemIndex = 0; SubsystemIndex < Subsystems.Num(); SubsystemIndex++)
-							{
-								FRemoteSubsystemBase* Subsystem = Subsystems[SubsystemIndex];
-								Subsystem->TickAbortedRequest();
-							}
-						}
+					if (AbortReason == ERemoteExecutorAbortReason::RequiresDependencies) // -V547
+					{
+						UE_LOG(LogRemoteExec, Verbose, TEXT("ExecutePendingWork[%d,%s] : Work ABORTED for required dependency '%s'"),
+							LocalIterationNumber,
+							*ExecutingWork->RequestId.ToString(),
+							*AbortReasonDescription);
 					}
-					else
+					else if (AbortReason == ERemoteExecutorAbortReason::AbandonWork) // -V547
 					{
-						check(AbortReason == UE::RemoteExecutor::ERemoteExecutorAbortReason::Unspecified);
-
-						UE_LOG(LogRemoteExec, VeryVerbose, TEXT("ExecutePendingWork[%d,%s] : Work COMPLETED after %u attempts"),
+						UE_LOG(LogRemoteExec, Verbose, TEXT("ExecutePendingWork[%d,%s] : Work ABORTED for ABANDONMENT '%s'"),
 							LocalIterationNumber,
 							*ExecutingWork->RequestId.ToString(),
-							ExecutingWork->ExecutionAttempts);
-						UE::RemoteExecutor::OnTransactionCompletedDelegate.Broadcast(RequestId, ExecutingWork->ExecutionAttempts);
-
+							*AbortReasonDescription);
 						bWorkComplete = true;
 					}
+					else
+					{
+						UE_LOG(LogRemoteExec, Verbose, TEXT("ExecutePendingWork[%d,%s] : Work ABORTED for unknown reason '%s'"),
+							LocalIterationNumber,
+							*ExecutingWork->RequestId.ToString(),
+							*AbortReasonDescription);
+					}
 
-					if (bWorkComplete)
+					if (!bWorkComplete)
 					{
 						for (int32 SubsystemIndex = 0; SubsystemIndex < Subsystems.Num(); SubsystemIndex++)
 						{
 							FRemoteSubsystemBase* Subsystem = Subsystems[SubsystemIndex];
-							UE_LOG(LogRemoteExec, VeryVerbose, TEXT("ExecutePendingWork[%d,%s] : End request[%s]"),
-								LocalIterationNumber,
-								*ExecutingWork->RequestId.ToString(),
-								Subsystem->NameForDebug());
-							Subsystem->EndRequest(!bWorkTransactionAborted);
-							Subsystem->ClearActiveRequest();
+							Subsystem->TickAbortedRequest();
 						}
+					}
+				}
+				else
+				{
+					check(AbortReason == UE::RemoteExecutor::ERemoteExecutorAbortReason::Unspecified);
 
-						for (int32 SubsystemIndex = 0; SubsystemIndex < Subsystems.Num(); SubsystemIndex++)
-						{
-							FRemoteSubsystemBase* Subsystem = Subsystems[SubsystemIndex];
-							Subsystem->DestroyRequest(ExecutingWork->RequestId);
-						}
+					UE_LOG(LogRemoteExec, VeryVerbose, TEXT("ExecutePendingWork[%d,%s] : Work COMPLETED after %u attempts"),
+						LocalIterationNumber,
+						*ExecutingWork->RequestId.ToString(),
+						ExecutingWork->ExecutionAttempts);
+					UE::RemoteExecutor::OnTransactionCompletedDelegate.Broadcast(RequestId, ExecutingWork->ExecutionAttempts);
 
-						UE::RemoteExecutor::OnTransactionReleasedDelegate.Broadcast(ExecutingWork->RequestId);
+					bWorkComplete = true;
+				}
 
-						// remove this work from the list and do not increment PendingWorkIndex
-						EnqueueNextWorkFromList(ExecutingWork->WorkListIndex);
-						PendingWorks.RemoveAt(PendingWorkIndex, EAllowShrinking::No);
+				if (bWorkComplete)
+				{
+					for (int32 SubsystemIndex = 0; SubsystemIndex < Subsystems.Num(); SubsystemIndex++)
+					{
+						FRemoteSubsystemBase* Subsystem = Subsystems[SubsystemIndex];
+						UE_LOG(LogRemoteExec, VeryVerbose, TEXT("ExecutePendingWork[%d,%s] : End request[%s]"),
+							LocalIterationNumber,
+							*ExecutingWork->RequestId.ToString(),
+							Subsystem->NameForDebug());
+						Subsystem->EndRequest(!bWorkTransactionAborted);
+						Subsystem->ClearActiveRequest();
 					}
-					else
+
+					for (int32 SubsystemIndex = 0; SubsystemIndex < Subsystems.Num(); SubsystemIndex++)
 					{
-						// we will revisit this later, move on to the next work
-						PendingWorkIndex++;
+						FRemoteSubsystemBase* Subsystem = Subsystems[SubsystemIndex];
+						Subsystem->DestroyRequest(ExecutingWork->RequestId);
 					}
-				}
-				else
-				{
-					UE::RemoteExecutor::OnRegionBeginDelegate.ExecuteIfBound(ExecutingWork->Name);
 
-					// execute non-transactional work
-					ExecutingWork->Work();
+					UE::RemoteExecutor::OnTransactionReleasedDelegate.Broadcast(ExecutingWork->RequestId);
 
 					// remove this work from the list and do not increment PendingWorkIndex
 					EnqueueNextWorkFromList(ExecutingWork->WorkListIndex);
-					PendingWorks.RemoveAt(PendingWorkIndex, EAllowShrinking::No);
-
-					UE::RemoteExecutor::OnRegionEndDelegate.ExecuteIfBound(TEXT(""));
+					WorksToExecute.RemoveAt(PendingWorkIndex, EAllowShrinking::No);
 				}
-
-				// it's possible this work ended with trying to perform a multi-server
-				// commit that got aborted because a remote multi-server commit of higher
-				// priority came in. if so, we need to break out of processing further work
-				// so we can service the remote multi-server commit
-				if (ActiveRemoteMultiServerCommitRequestId.IsValid())
+				else
 				{
-					break;
+					// we will revisit this later, move on to the next work
+					PendingWorkIndex++;
 				}
+			}
+			else
+			{
+				UE::RemoteExecutor::OnRegionBeginDelegate.ExecuteIfBound(ExecutingWork->Name);
+
+				// execute non-transactional work
+				ExecutingWork->Work();
 
-				// Finally, see if we have any requests waiting that need servicing
-				ConditionallyPumpNetwork(false);
+				// remove this work from the list and do not increment PendingWorkIndex
+				EnqueueNextWorkFromList(ExecutingWork->WorkListIndex);
+				WorksToExecute.RemoveAt(PendingWorkIndex, EAllowShrinking::No);
+
+				UE::RemoteExecutor::OnRegionEndDelegate.ExecuteIfBound(TEXT(""));
 			}
-		}
 
-		PendingWorkLists.Empty();
-		AbortReason = UE::RemoteExecutor::ERemoteExecutorAbortReason::Unspecified;
-		AbortReasonDescription.Reset();
-	}
+			// it's possible this work ended with trying to perform a multi-server
+			// commit that got aborted because a remote multi-server commit of higher
+			// priority came in. if so, we need to break out of processing further work
+			// so we can service the remote multi-server commit
+			if (ActiveRemoteMultiServerCommitRequestId.IsValid())
+			{
+				break;
+			}
+
+			// Finally, see if we have any requests waiting that need servicing
+			ConditionallyPumpNetwork(false);
+		}
+	}	
 };
 
 FRemoteExecutor GRemoteExecutor;
@@ -1189,12 +1183,8 @@ void EnqueueWork(FName WorkName, bool bIsTransactional, TUniqueFunction<void(voi
 
 void EnqueueMigrationWork(FName WorkName, TUniqueFunction<void(void)>&& InWork)
 {
-#if UE_WITH_REMOTE_OBJECT_HANDLE && UE_AUTORTFM
-	FRemoteWorkPriority RootWorkPriority = FRemoteWorkPriority::CreateRootWorkPriority(FRemoteServerId::GetLocalServerId(), GRemoteExecutor.GenerateNextTransactionId());
-	GRemoteExecutor.EnqueueMigrationWork(WorkName, RootWorkPriority, INDEX_NONE, MoveTemp(InWork));
-#else
-	InWork();
-#endif
+	FRemoteWorkPriority RootWorkPriority = FRemoteWorkPriority::CreateRootWorkPriority(FRemoteServerId::GetLocalServerId(), GRemoteExecutor.GenerateNextTransactionId(), ERemoteWorkType::Migration);
+	EnqueueWorkWithExplicitPriority(WorkName, RootWorkPriority, /*bIsTransactional=*/ true, MoveTemp(InWork));
 }
 
 
@@ -1219,26 +1209,50 @@ void EnqueueWorkList(FName WorkName, bool bIsTransactional, TArray<TUniqueFuncti
 
 } // namespace UE::RemoteExecutor
 
+const TCHAR* EnumToString(ERemoteWorkType WorkType)
+{
+	switch (WorkType)
+	{
+	case ERemoteWorkType::Migration:
+		return TEXT("Migration");
+	case ERemoteWorkType::Normal:
+		return TEXT("Normal");
+	default:
+		checkf(false, TEXT("Unknown ERemoteWorkType %d"), (int32)WorkType);
+		break;
+	}
+	return TEXT("");
+}
+
 FString FRemoteWorkPriority::ToString() const
 {
-	return FString::Format(TEXT("[pri: rsi {0} depth {1} id {2}]"),
+	return FString::Format(TEXT("[pri: rsi {0} depth {1} id {2} type {3}]"),
 		{ *GetRootServerId().ToString(),
 		GetWorkDepth(),
-		*GetRootWorkTransactionId().ToString() });
+		*GetRootWorkTransactionId().ToString(),
+		EnumToString(GetWorkType()) });
 }
 
-FRemoteWorkPriority FRemoteWorkPriority::CreateRootWorkPriority(FRemoteServerId ServerId, FRemoteTransactionId TransactionId)
+FRemoteWorkPriority FRemoteWorkPriority::CreateRootWorkPriority(FRemoteServerId ServerId, FRemoteTransactionId TransactionId, ERemoteWorkType WorkType /*= ERemoteWorkType::Normal*/)
 {
-	FRemoteWorkPriority Result{};
+	FRemoteWorkPriority Result(WorkType);
 
 	uint32 RawServerId = ServerId.GetIdNumber();
 
 	uint32 RawTransactionId = 0;
 	FMemory::Memcpy(&RawTransactionId, &TransactionId, sizeof(FRemoteTransactionId));
 
-	Result.PackedData |= (static_cast<uint64>(RawServerId) << 32);
-	Result.PackedData |= (0xFFull << 24);
+	constexpr uint64 RootWorkDepth = 0xFFull;
+
+	Result.PackedData |= (static_cast<uint64>(RawServerId & 0x7FFFFFFFu) << 32);
+	Result.PackedData |= (RootWorkDepth << 24);
 	Result.PackedData |= (RawTransactionId & 0xFFFFFFu);
+
+	checkf(WorkType == Result.GetWorkType(), TEXT("Work type (%u) was not properly encoded (%d)"), (uint32)WorkType, (uint32)Result.GetWorkType());
+	checkf(ServerId == Result.GetRootServerId(), TEXT("Server id (%s) was not properly encoded (%s)"), *ServerId.ToString(), *Result.GetRootServerId().ToString());
+	checkf(RootWorkDepth == Result.GetWorkDepth(), TEXT("Work depth (%llu) was not properly encoded (%u)"), RootWorkDepth, Result.GetWorkDepth());
+	checkf(TransactionId == Result.GetRootWorkTransactionId(), TEXT("Transaction id (%u) was not properly encoded (%u)"), TransactionId.GetIdNumber(), Result.GetRootWorkTransactionId().GetIdNumber());
+
 	return Result;
 }
 
@@ -1251,6 +1265,9 @@ FRemoteWorkPriority FRemoteWorkPriority::CreateDependentWorkPriority() const
 	check(GetWorkDepth() > 0);
 
 	Result.PackedData |= static_cast<uint64>(GetWorkDepth() - 1) << 24;
+
+	checkf(GetWorkType() == Result.GetWorkType(), TEXT("Work type (%u) changed when creating dependent work priority (%d)"), (uint32)GetWorkType(), (uint32)Result.GetWorkType());
+
 	return Result;
 }
 
diff --git a/Engine/Source/Runtime/CoreUObject/Public/UObject/RemoteExecutor.h b/Engine/Source/Runtime/CoreUObject/Public/UObject/RemoteExecutor.h
--- a/Engine/Source/Runtime/CoreUObject/Public/UObject/RemoteExecutor.h
+++ b/Engine/Source/Runtime/CoreUObject/Public/UObject/RemoteExecutor.h
@@ -49,10 +49,19 @@ private:
 	uint32 Id = 0;
 };
 
+enum class ERemoteWorkType
+{
+	Migration = 0,
+	Normal = 1
+};
+
+COREUOBJECT_API const TCHAR* EnumToString(ERemoteWorkType WorkType);
+
 struct FRemoteWorkPriority
 {
 	// the data is packed in a uint64 to allow passing by value in a register
-	// [63:32] : FRemoteServerId RootServerId
+	// [63]    : 0 if this priority represents transactional migration work, 1 otherwise
+	// [62:32] : FRemoteServerId RootServerId
 	// [31:24] : uint8 WorkDepth
 	// [23:0]  : FRemoteTransactionId RootWorkTransactionId
 	//
@@ -66,13 +75,23 @@ struct FRemoteWorkPriority
 	static_assert(sizeof(FRemoteServerId) == 4);
 	static_assert(sizeof(FRemoteTransactionId) == 4);
 
-	FRemoteServerId GetRootServerId() const { return FRemoteServerId::FromIdNumber(static_cast<uint32>(PackedData >> 32)); }
+	FRemoteServerId GetRootServerId() const { return FRemoteServerId::FromIdNumber(static_cast<uint32>((PackedData >> 32) & 0x7FFFFFFFu)); }
 	uint32 GetWorkDepth() const { return (PackedData >> 24) & 0xFFu; }
 	FRemoteTransactionId GetRootWorkTransactionId() const { return FRemoteTransactionId(PackedData & 0xFFFFFFu); }
-
+	ERemoteWorkType GetWorkType() const { return !(PackedData >> 63) ? ERemoteWorkType::Migration : ERemoteWorkType::Normal; }
 	bool IsValid() const { return GetRootServerId().IsValid(); }
 
-	COREUOBJECT_API static FRemoteWorkPriority CreateRootWorkPriority(FRemoteServerId ServerId, FRemoteTransactionId TransactionId);
+	FRemoteWorkPriority()
+		: FRemoteWorkPriority(ERemoteWorkType::Normal)
+	{
+	}
+	
+	explicit FRemoteWorkPriority(ERemoteWorkType WorkType)
+		: PackedData(static_cast<uint64>(WorkType) << 63)
+	{
+	}
+
+	COREUOBJECT_API static FRemoteWorkPriority CreateRootWorkPriority(FRemoteServerId ServerId, FRemoteTransactionId TransactionId, ERemoteWorkType WorkType = ERemoteWorkType::Normal);
 
 	// creates a copy of this priority with the WorkDepth decreased by 1
 	COREUOBJECT_API FRemoteWorkPriority CreateDependentWorkPriority() const;
