[StructUtils] Change PropertyBag to create FByteProperty for enums that are not a C++ class enum. Similar to what UHT and Blueprint produce with TEnumAsByte.

#Synced-CL 50158725

#Synced-CL 50158725

#was-backout 50176042

#Synced-CL 50200679

diff --git a/Engine/Plugins/Experimental/UAF/UAF/Source/UAF/Private/Param/ParamType.cpp b/Engine/Plugins/Experimental/UAF/UAF/Source/UAF/Private/Param/ParamType.cpp
--- a/Engine/Plugins/Experimental/UAF/UAF/Source/UAF/Private/Param/ParamType.cpp
+++ b/Engine/Plugins/Experimental/UAF/UAF/Source/UAF/Private/Param/ParamType.cpp
@@ -178,9 +178,17 @@ FAnimNextParamType FAnimNextParamType::FromProperty(const FProperty* InProperty)
 	{
 		Type.ValueType = EPropertyBagPropertyType::Bool;
 	}
-	else if (InProperty->IsA<FByteProperty>())
+	else if (const FByteProperty* ByteProperty = CastField<const FByteProperty>(InProperty))
 	{
-		Type.ValueType = EPropertyBagPropertyType::Byte;
+		if (ByteProperty->Enum != nullptr)
+		{
+			Type.ValueType = EPropertyBagPropertyType::Enum;
+			Type.ValueTypeObject = ByteProperty->Enum;
+		}
+		else
+		{
+			Type.ValueType = EPropertyBagPropertyType::Byte;
+		}
 	}
 	else if (InProperty->IsA<FIntProperty>())
 	{
@@ -647,7 +655,12 @@ FAnimNextParamType FAnimNextParamType::FromString(const FString& InString)
 			ObjectInnerString = InTypeString.RightChop(1).TrimStartAndEnd();
 			ObjectValueType = EValueType::Enum;
 		}
-		
+		else if(InTypeString.StartsWith(TEXT("TEnumAsByte<U"), ESearchCase::CaseSensitive))
+		{
+			ObjectInnerString = InTypeString.RightChop(13).LeftChop(1).TrimStartAndEnd();
+			ObjectValueType = EValueType::Enum;
+		}
+
 		if(UObject* ObjectType = FindFirstObject<UObject>(*ObjectInnerString, EFindFirstObjectOptions::NativeFirst))
 		{
 			OutType.ValueType = ObjectValueType;
diff --git a/Engine/Plugins/Experimental/UAF/UAF/Source/UAF/Private/Variables/UAFInstanceVariableContainer.cpp b/Engine/Plugins/Experimental/UAF/UAF/Source/UAF/Private/Variables/UAFInstanceVariableContainer.cpp
--- a/Engine/Plugins/Experimental/UAF/UAF/Source/UAF/Private/Variables/UAFInstanceVariableContainer.cpp
+++ b/Engine/Plugins/Experimental/UAF/UAF/Source/UAF/Private/Variables/UAFInstanceVariableContainer.cpp
@@ -76,11 +76,22 @@ namespace UE::UAF::Private
 		}
 		case EPropertyBagPropertyType::Enum:
 		{
-			const FEnumProperty* EnumProperty = CastFieldChecked<FEnumProperty>(InProperty);
-			const FNumericProperty* UnderlyingProperty = EnumProperty->GetUnderlyingProperty();
-			check(UnderlyingProperty);
-			OutValue = UnderlyingProperty->GetSignedIntPropertyValue(Address);
-			return EPropertyBagResult::Success;
+			if (const FEnumProperty* EnumProperty = CastField<FEnumProperty>(InProperty))
+			{
+				const FNumericProperty* UnderlyingProperty = EnumProperty->GetUnderlyingProperty();
+				check(UnderlyingProperty);
+				OutValue = UnderlyingProperty->GetSignedIntPropertyValue(Address);
+				return EPropertyBagResult::Success;
+			}
+			else if (const FByteProperty* ByteProperty = CastField<FByteProperty>(InProperty))
+			{
+				if (ByteProperty->Enum != nullptr)
+				{
+					OutValue = ByteProperty->GetPropertyValue(Address);
+					return EPropertyBagResult::Success;
+				}
+			}
+			return EPropertyBagResult::TypeMismatch;
 		}
 		default:
 			return EPropertyBagResult::TypeMismatch;
@@ -144,11 +155,22 @@ namespace UE::UAF::Private
 		}
 		case EPropertyBagPropertyType::Enum:
 		{
-			const FEnumProperty* EnumProperty = CastFieldChecked<FEnumProperty>(InProperty);
-			const FNumericProperty* UnderlyingProperty = EnumProperty->GetUnderlyingProperty();
-			check(UnderlyingProperty);
-			OutValue = UnderlyingProperty->GetUnsignedIntPropertyValue(Address);
-			return EPropertyBagResult::Success;
+			if (const FEnumProperty* EnumProperty = CastField<FEnumProperty>(InProperty))
+			{
+				const FNumericProperty* UnderlyingProperty = EnumProperty->GetUnderlyingProperty();
+				check(UnderlyingProperty);
+				OutValue = UnderlyingProperty->GetUnsignedIntPropertyValue(Address);
+				return EPropertyBagResult::Success;
+			}
+			else if (const FByteProperty* ByteProperty = CastField<FByteProperty>(InProperty))
+			{
+				if (ByteProperty->Enum != nullptr)
+				{
+					OutValue = ByteProperty->GetPropertyValue(Address);
+					return EPropertyBagResult::Success;
+				}
+			}
+			return EPropertyBagResult::TypeMismatch;
 		}
 		default:
 			return EPropertyBagResult::TypeMismatch;
@@ -212,11 +234,22 @@ namespace UE::UAF::Private
 		}
 		case EPropertyBagPropertyType::Enum:
 		{
-			const FEnumProperty* EnumProperty = CastFieldChecked<FEnumProperty>(InProperty);
-			const FNumericProperty* UnderlyingProperty = EnumProperty->GetUnderlyingProperty();
-			check(UnderlyingProperty);
-			OutValue = static_cast<double>(UnderlyingProperty->GetSignedIntPropertyValue(Address));
-			return EPropertyBagResult::Success;
+			if (const FEnumProperty* EnumProperty = CastField<FEnumProperty>(InProperty))
+			{
+				const FNumericProperty* UnderlyingProperty = EnumProperty->GetUnderlyingProperty();
+				check(UnderlyingProperty);
+				OutValue = static_cast<double>(UnderlyingProperty->GetSignedIntPropertyValue(Address));
+				return EPropertyBagResult::Success;
+			}
+			else if (const FByteProperty* ByteProperty = CastField<FByteProperty>(InProperty))
+			{
+				if (ByteProperty->Enum != nullptr)
+				{
+					OutValue = ByteProperty->GetPropertyValue(Address);
+					return EPropertyBagResult::Success;
+				}
+			}
+			return EPropertyBagResult::TypeMismatch;
 		}
 		default:
 			return EPropertyBagResult::TypeMismatch;
@@ -251,18 +284,32 @@ namespace UE::UAF::Private
 			return EPropertyBagResult::TypeMismatch;
 		}
 
-		const FEnumProperty* EnumProperty = CastFieldChecked<FEnumProperty>(InProperty);
-		const FNumericProperty* UnderlyingProperty = EnumProperty->GetUnderlyingProperty();
-		check(UnderlyingProperty);
+		if (const FEnumProperty* EnumProperty = CastField<FEnumProperty>(InProperty))
+		{
+			const FNumericProperty* UnderlyingProperty = EnumProperty->GetUnderlyingProperty();
+			check(UnderlyingProperty);
 
-		if (RequestedEnum != EnumProperty->GetEnum())
+			if (RequestedEnum != EnumProperty->GetEnum())
+			{
+				return EPropertyBagResult::TypeMismatch;
+			}
+
+			OutValue = static_cast<uint8>(UnderlyingProperty->GetUnsignedIntPropertyValue(Address));
+			return EPropertyBagResult::Success;
+		}
+		else if (const FByteProperty* ByteProperty = CastField<FByteProperty>(InProperty))
+		{
+			if (RequestedEnum != ByteProperty->Enum)
+			{
+				return EPropertyBagResult::TypeMismatch;
+			}
+			OutValue = ByteProperty->GetPropertyValue(Address);
+			return EPropertyBagResult::Success;
+		}
+		else
 		{
 			return EPropertyBagResult::TypeMismatch;
 		}
-
-		OutValue = static_cast<uint8>(UnderlyingProperty->GetUnsignedIntPropertyValue(Address));
-
-		return EPropertyBagResult::Success;
 	}
 
 	EPropertyBagResult GetPropertyValueAsStruct(const FProperty* InProperty, EPropertyBagPropertyType InValueType, const void* Address, const UScriptStruct* RequestedStruct, uint8* OutValue)
diff --git a/Engine/Plugins/Experimental/UAF/UAF/Source/UAFUncookedOnly/Private/UncookedOnlyUtils.cpp b/Engine/Plugins/Experimental/UAF/UAF/Source/UAFUncookedOnly/Private/UncookedOnlyUtils.cpp
--- a/Engine/Plugins/Experimental/UAF/UAF/Source/UAFUncookedOnly/Private/UncookedOnlyUtils.cpp
+++ b/Engine/Plugins/Experimental/UAF/UAF/Source/UAFUncookedOnly/Private/UncookedOnlyUtils.cpp
@@ -804,7 +804,6 @@ FEdGraphPinType FUtils::GetPinTypeFromParamType(const FAnimNextParamType& InPara
 		PinType.PinCategory = UEdGraphSchema_K2::PC_Text;
 		break;
 	case EPropertyBagPropertyType::Enum:
-		// @todo: some pin coloring is not correct due to this (byte-as-enum vs enum). 
 		PinType.PinCategory = UEdGraphSchema_K2::PC_Enum;
 		PinType.PinSubCategoryObject = const_cast<UObject*>(InParamType.ValueTypeObject.Get());
 		break;
diff --git a/Engine/Plugins/Runtime/StateTree/Source/StateTreeEditorModule/Private/StateTreeCompiler.cpp b/Engine/Plugins/Runtime/StateTree/Source/StateTreeEditorModule/Private/StateTreeCompiler.cpp
--- a/Engine/Plugins/Runtime/StateTree/Source/StateTreeEditorModule/Private/StateTreeCompiler.cpp
+++ b/Engine/Plugins/Runtime/StateTree/Source/StateTreeEditorModule/Private/StateTreeCompiler.cpp
@@ -65,7 +65,7 @@ namespace UE::StateTree::Compiler
 	// Version used by AppendToClassSchema. Bump this version number when a recompile of all assets is needed.
 	//Only bump when only the compiler changes.
 	//You do not need to bump the version when the memory layout of UStateTree has changed.
-	static constexpr int32 ClassSchemaVersion = 4;
+	static constexpr int32 ClassSchemaVersion = 5;
 
 	FAutoConsoleVariable CVarLogEnableParameterDelegateDispatcherBinding(
 		TEXT("StateTree.Compiler.EnableParameterDelegateDispatcherBinding"),
diff --git a/Engine/Source/Runtime/CoreUObject/Private/StructUtils/PropertyBag.cpp b/Engine/Source/Runtime/CoreUObject/Private/StructUtils/PropertyBag.cpp
--- a/Engine/Source/Runtime/CoreUObject/Private/StructUtils/PropertyBag.cpp
+++ b/Engine/Source/Runtime/CoreUObject/Private/StructUtils/PropertyBag.cpp
@@ -512,11 +512,23 @@
 		case EPropertyBagPropertyType::Enum:
 			if (const UEnum* Enum = Cast<UEnum>(Desc.ValueTypeObject))
 			{
-				FEnumProperty* Prop = new FEnumProperty(PropertyScope, Desc.Name, RF_Public);
-				FNumericProperty* UnderlyingProp = new FByteProperty(Prop, "UnderlyingType", RF_Public); // HACK: Hardwire to byte property for now for BP compatibility
-				Prop->SetEnum(const_cast<UEnum*>(Enum));
-				Prop->AddCppProperty(UnderlyingProp);
-				return Prop;
+				// Only generate FEnumProperty for cpp enum class like UHT and Blueprint.
+				if (Enum->GetCppForm() == UEnum::ECppForm::EnumClass)
+				{
+					FEnumProperty* Prop = new FEnumProperty(PropertyScope, Desc.Name, RF_Public);
+					Prop->SetPropertyFlags(CPF_HasGetValueTypeHash);
+					FNumericProperty* UnderlyingProp = new FByteProperty(Prop, "UnderlyingType", RF_Public); // HACK: Hardwire to byte property for now for BP compatibility
+					Prop->SetEnum(const_cast<UEnum*>(Enum));
+					Prop->AddCppProperty(UnderlyingProp);
+					return Prop;
+				}
+				else
+				{
+					FByteProperty* Prop = new FByteProperty(PropertyScope, Desc.Name, RF_Public);
+					Prop->SetPropertyFlags(CPF_HasGetValueTypeHash);
+					Prop->Enum = const_cast<UEnum*>(Enum);
+					return Prop;
+				}
 			}
 			break;
 		case EPropertyBagPropertyType::Struct:
@@ -672,11 +684,22 @@
 			}
 		case EPropertyBagPropertyType::Enum:
 			{
-				const FEnumProperty* EnumProperty = CastFieldChecked<FEnumProperty>(Desc->CachedProperty);
-				const FNumericProperty* UnderlyingProperty = EnumProperty->GetUnderlyingProperty();
-				check(UnderlyingProperty);
-				OutValue = UnderlyingProperty->GetSignedIntPropertyValue(Address);
-				return EPropertyBagResult::Success;
+				if (const FEnumProperty* EnumProperty = CastField<FEnumProperty>(Desc->CachedProperty))
+				{
+					const FNumericProperty* UnderlyingProperty = EnumProperty->GetUnderlyingProperty();
+					check(UnderlyingProperty);
+					OutValue = UnderlyingProperty->GetSignedIntPropertyValue(Address);
+					return EPropertyBagResult::Success;
+				}
+				else if (const FByteProperty* ByteProperty = CastField<FByteProperty>(Desc->CachedProperty))
+				{
+					if (ByteProperty->Enum != nullptr)
+					{
+						OutValue = ByteProperty->GetPropertyValue(Address);
+						return EPropertyBagResult::Success;
+					}
+				}
+				return EPropertyBagResult::TypeMismatch;
 			}
 		default:
 			return EPropertyBagResult::TypeMismatch;
@@ -750,11 +773,22 @@
 		}
 		case EPropertyBagPropertyType::Enum:
 		{
-			const FEnumProperty* EnumProperty = CastFieldChecked<FEnumProperty>(Desc->CachedProperty);
-			const FNumericProperty* UnderlyingProperty = EnumProperty->GetUnderlyingProperty();
-			check(UnderlyingProperty);
-			OutValue = UnderlyingProperty->GetUnsignedIntPropertyValue(Address);
-			return EPropertyBagResult::Success;
+			if (const FEnumProperty* EnumProperty = CastField<FEnumProperty>(Desc->CachedProperty))
+			{
+				const FNumericProperty* UnderlyingProperty = EnumProperty->GetUnderlyingProperty();
+				check(UnderlyingProperty);
+				OutValue = UnderlyingProperty->GetUnsignedIntPropertyValue(Address);
+				return EPropertyBagResult::Success;
+			}
+			else if (const FByteProperty* ByteProperty = CastField<FByteProperty>(Desc->CachedProperty))
+			{
+				if (ByteProperty->Enum != nullptr)
+				{
+					OutValue = ByteProperty->GetPropertyValue(Address);
+					return EPropertyBagResult::Success;
+				}
+			}
+			return EPropertyBagResult::TypeMismatch;
 		}
 		default:
 			return EPropertyBagResult::TypeMismatch;
@@ -828,11 +862,22 @@
 			}
 		case EPropertyBagPropertyType::Enum:
 			{
-				const FEnumProperty* EnumProperty = CastFieldChecked<FEnumProperty>(Desc->CachedProperty);
-				const FNumericProperty* UnderlyingProperty = EnumProperty->GetUnderlyingProperty();
-				check(UnderlyingProperty);
-				OutValue = static_cast<double>(UnderlyingProperty->GetSignedIntPropertyValue(Address));
-				return EPropertyBagResult::Success;
+				if (const FEnumProperty* EnumProperty = CastField<FEnumProperty>(Desc->CachedProperty))
+				{
+					const FNumericProperty* UnderlyingProperty = EnumProperty->GetUnderlyingProperty();
+					check(UnderlyingProperty);
+					OutValue = static_cast<double>(UnderlyingProperty->GetSignedIntPropertyValue(Address));
+					return EPropertyBagResult::Success;
+				}
+				else if (const FByteProperty* ByteProperty = CastField<FByteProperty>(Desc->CachedProperty))
+				{
+					if (ByteProperty->Enum != nullptr)
+					{
+						OutValue = ByteProperty->GetPropertyValue(Address);
+						return EPropertyBagResult::Success;
+					}
+				}
+				return EPropertyBagResult::TypeMismatch;
 			}
 		default:
 			return EPropertyBagResult::TypeMismatch;
@@ -886,16 +931,30 @@
 			return EPropertyBagResult::TypeMismatch;
 		}
 
-		const FEnumProperty* EnumProperty = CastFieldChecked<FEnumProperty>(Desc->CachedProperty);
-		const FNumericProperty* UnderlyingProperty = EnumProperty->GetUnderlyingProperty();
-		check(UnderlyingProperty);
+		if (const FEnumProperty* EnumProperty = CastField<FEnumProperty>(Desc->CachedProperty))
+		{
+			const FNumericProperty* UnderlyingProperty = EnumProperty->GetUnderlyingProperty();
+			check(UnderlyingProperty);
 	
-		if (RequestedEnum != EnumProperty->GetEnum())
+			if (RequestedEnum != EnumProperty->GetEnum())
+			{
+				return EPropertyBagResult::TypeMismatch;
+			}
+	
+			OutValue = static_cast<uint8>(UnderlyingProperty->GetUnsignedIntPropertyValue(Address));
+		}
+		else if (const FByteProperty* ByteProperty = CastField<FByteProperty>(Desc->CachedProperty))
+		{
+			if (RequestedEnum != ByteProperty->Enum)
+			{
+				return EPropertyBagResult::TypeMismatch;
+			}
+			OutValue = ByteProperty->GetPropertyValue(Address);
+		}
+		else
 		{
 			return EPropertyBagResult::TypeMismatch;
 		}
-	
-		OutValue = static_cast<uint8>(UnderlyingProperty->GetUnsignedIntPropertyValue(Address));
 
 		return EPropertyBagResult::Success;
 	}
@@ -1066,11 +1125,22 @@
 			}
 		case EPropertyBagPropertyType::Enum:
 			{
-				const FEnumProperty* EnumProperty = CastFieldChecked<FEnumProperty>(Desc->CachedProperty);
-				const FNumericProperty* UnderlyingProperty = EnumProperty->GetUnderlyingProperty();
-				check(UnderlyingProperty);
-				UnderlyingProperty->SetIntPropertyValue(Address, static_cast<uint64>(InValue));
-				return EPropertyBagResult::Success;
+				if (const FEnumProperty* EnumProperty = CastField<FEnumProperty>(Desc->CachedProperty))
+				{
+					const FNumericProperty* UnderlyingProperty = EnumProperty->GetUnderlyingProperty();
+					check(UnderlyingProperty);
+					UnderlyingProperty->SetIntPropertyValue(Address, static_cast<uint64>(InValue));
+					return EPropertyBagResult::Success;
+				}
+				else if (const FByteProperty* ByteProperty = CastField<FByteProperty>(Desc->CachedProperty))
+				{
+					if (ByteProperty->Enum)
+					{
+						ByteProperty->SetPropertyValue(Address, static_cast<uint8>(InValue));
+						return EPropertyBagResult::Success;
+					}
+				}
+				return EPropertyBagResult::TypeMismatch;
 			}
 		default:
 			return EPropertyBagResult::TypeMismatch;
@@ -1144,11 +1214,22 @@
 		}
 		case EPropertyBagPropertyType::Enum:
 		{
-			const FEnumProperty* EnumProperty = CastFieldChecked<FEnumProperty>(Desc->CachedProperty);
-			const FNumericProperty* UnderlyingProperty = EnumProperty->GetUnderlyingProperty();
-			check(UnderlyingProperty);
-			UnderlyingProperty->SetIntPropertyValue(Address, InValue);
-			return EPropertyBagResult::Success;
+			if (const FEnumProperty* EnumProperty = CastField<FEnumProperty>(Desc->CachedProperty))
+			{
+				const FNumericProperty* UnderlyingProperty = EnumProperty->GetUnderlyingProperty();
+				check(UnderlyingProperty);
+				UnderlyingProperty->SetIntPropertyValue(Address, InValue);
+				return EPropertyBagResult::Success;
+			}
+			else if (const FByteProperty* ByteProperty = CastField<FByteProperty>(Desc->CachedProperty))
+			{
+				if (ByteProperty->Enum)
+				{
+					ByteProperty->SetPropertyValue(Address, static_cast<uint8>(InValue));
+					return EPropertyBagResult::Success;
+				}
+			}
+			return EPropertyBagResult::TypeMismatch;
 		}
 		default:
 			return EPropertyBagResult::TypeMismatch;
@@ -1222,11 +1303,22 @@
 			}
 		case EPropertyBagPropertyType::Enum:
 			{
-				const FEnumProperty* EnumProperty = CastFieldChecked<FEnumProperty>(Desc->CachedProperty);
-				const FNumericProperty* UnderlyingProperty = EnumProperty->GetUnderlyingProperty();
-				check(UnderlyingProperty);
-				UnderlyingProperty->SetIntPropertyValue(Address, static_cast<uint64>(InValue));
-				return EPropertyBagResult::Success;
+				if (const FEnumProperty* EnumProperty = CastField<FEnumProperty>(Desc->CachedProperty))
+				{
+					const FNumericProperty* UnderlyingProperty = EnumProperty->GetUnderlyingProperty();
+					check(UnderlyingProperty);
+					UnderlyingProperty->SetIntPropertyValue(Address, static_cast<uint64>(InValue));
+					return EPropertyBagResult::Success;
+				}
+				else if (const FByteProperty* ByteProperty = CastField<FByteProperty>(Desc->CachedProperty))
+				{
+					if (ByteProperty->Enum)
+					{
+						ByteProperty->SetPropertyValue(Address, static_cast<uint8>(FMath::RoundToInt32(InValue)));
+						return EPropertyBagResult::Success;
+					}
+				}
+				return EPropertyBagResult::TypeMismatch;
 			}
 		default:
 			return EPropertyBagResult::TypeMismatch;
@@ -1280,18 +1372,33 @@
 			return EPropertyBagResult::TypeMismatch;
 		}
 
-		const FEnumProperty* EnumProperty = CastFieldChecked<FEnumProperty>(Desc->CachedProperty);
-		const FNumericProperty* UnderlyingProperty = EnumProperty->GetUnderlyingProperty();
-		check(UnderlyingProperty);
+		if (const FEnumProperty* EnumProperty = CastField<FEnumProperty>(Desc->CachedProperty))
+		{
+			const FNumericProperty* UnderlyingProperty = EnumProperty->GetUnderlyingProperty();
+			check(UnderlyingProperty);
 
-		if (Enum != EnumProperty->GetEnum())
+			if (Enum != EnumProperty->GetEnum())
+			{
+				return EPropertyBagResult::TypeMismatch;
+			}
+	
+			UnderlyingProperty->SetIntPropertyValue(Address, static_cast<uint64>(InValue));
+	
+			return EPropertyBagResult::Success;
+		}
+		else if (const FByteProperty* ByteProperty = CastField<FByteProperty>(Desc->CachedProperty))
+		{
+			if (Enum != ByteProperty->Enum)
+			{
+				return EPropertyBagResult::TypeMismatch;
+			}
+			ByteProperty->SetPropertyValue(Address, InValue);
+			return EPropertyBagResult::Success;
+		}
+		else
 		{
 			return EPropertyBagResult::TypeMismatch;
 		}
-	
-		UnderlyingProperty->SetIntPropertyValue(Address, static_cast<uint64>(InValue));
-	
-		return EPropertyBagResult::Success;
 	}
 
 	EPropertyBagResult SetPropertyValueAsStruct(const FPropertyBagPropertyDesc* Desc, void* Address, const FConstStructView InValue)
@@ -4137,16 +4244,31 @@
 		return EPropertyBagResult::TypeMismatch;
 	}
 
-	const FEnumProperty* EnumProperty = CastFieldChecked<FEnumProperty>(ValueDesc.CachedProperty);
-	const FNumericProperty* UnderlyingProperty = EnumProperty->GetUnderlyingProperty();
-	check(UnderlyingProperty);
+	if (const FEnumProperty* EnumProperty = CastField<FEnumProperty>(ValueDesc.CachedProperty))
+	{
+		const FNumericProperty* UnderlyingProperty = EnumProperty->GetUnderlyingProperty();
+		check(UnderlyingProperty);
 
-	if (Enum != EnumProperty->GetEnum())
+		if (Enum != EnumProperty->GetEnum())
+		{
+			return EPropertyBagResult::TypeMismatch;
+		}
+
+		return Add(InValue);
+	}
+	else if (const FByteProperty* ByteProperty = CastField<FByteProperty>(ValueDesc.CachedProperty))
+	{
+		if (Enum != ByteProperty->Enum)
+		{
+			return EPropertyBagResult::TypeMismatch;
+		}
+
+		return Add(InValue);
+	}
+	else
 	{
 		return EPropertyBagResult::TypeMismatch;
 	}
-
-	return Add(InValue);
 }
 
 EPropertyBagResult FPropertyBagSetRef::AddValueStruct(FConstStructView InValue)
