[UE] Don't reconstruct a subobject after instancing its missing outer if it already constructed the subobject from the source archetype (optimization).

Notes:
- In cases where instancing a subobject recursively instances a missing owner, non-native owner types could end up redundantly reconstructing the subobject using the same source template.
- This is because for the non-native (Blueprint) case, we don't use StaticFindObject() to look for the instance after constructing the outer. This is intentional because in the typical case, we usually need StaticAllocateObject() to reconstruct a natively-constructed default subobject instance using the non-native outer class default object's data as the source rather than the native class default initialization that's handled by the native ctor.
- However, in this case that reconstruction is redundant, as we've already constructed and mapped a new instance within the object graph using the source archetype for initialization as part of instancing the missing outer object.

#Synced-CL 50036996

diff --git a/Engine/Source/Runtime/CoreUObject/Private/UObject/CoreNative.cpp b/Engine/Source/Runtime/CoreUObject/Private/UObject/CoreNative.cpp
--- a/Engine/Source/Runtime/CoreUObject/Private/UObject/CoreNative.cpp
+++ b/Engine/Source/Runtime/CoreUObject/Private/UObject/CoreNative.cpp
@@ -269,6 +269,13 @@ UObject* FObjectInstancingGraph::GetInstancedSubobject( UObject* SourceSubobject
 								SubobjectOuter = GetInstancedSubobject(SourceSubobject->GetOuter(), SourceSubobject->GetOuter(), CurrentObject, Flags);
 
 								checkf(SubobjectOuter && SubobjectOuter != INVALID_OBJECT, TEXT("No corresponding destination object found for '%s' while attempting to instance subobject '%s'"), *SourceSubobject->GetOuter()->GetFullName(), *SourceSubobject->GetFullName());
+
+								// Constructing the outer may have also instanced this subobject - avoid reconstructing it again with the same archetype.
+								InstancedSubobject = GetDestinationObject(SourceSubobject);
+								if (InstancedSubobject != nullptr)
+								{
+									return InstancedSubobject;
+								}
 							}
 
 							FName SubobjectName = SourceSubobject->GetFName();
