Fix an issue where creating a new level from a template could result in a crash.

#rnx

### Problem
- Since editorbulkdata was first written we have tried to find the FLinkerLoad during serialization from the owning UObject's UPackage mirroring the older bulkdata which used FLinkerLoad::FindExistingLinkerForPackage.
- Recent optimizations have meant that it does not seem to be safe to rely on the owners package containing a valid FLinkerLoad and instead we should rely on finding the FLinkerLoad from the archive that the object is being serialized from.

### Fix
- To fix this in the short term if we cannot find the FLinkerLoad via the package we should make another attempt to do so via the archive being serialized from and only give up if the linker cannot be found by either code path.
- A cvar 'Serialization.AllowLinkerFromArchive' has been added which is true by default. Setting it to false will disable the new attempt to find the linker via the archive and restore the old behavior in case issues are found.
- Longer term we should consider switching to only use the archive.
- Kept the original version of ::GetTrailerFromOwner as it is used when creating editor bulkdata from old style bulkdata which is a special case pass.
- Created a new version of :: GetTrailerFromOwner called ::FindPackageTrailer which takes both the owning object and the archive being serialized from.
- Renamed ::GetPackagePathFromOwner to ::FindPackagePath which now takes both the owning object and the archive being serialized from.

diff --git a/Engine/Source/Runtime/CoreUObject/Private/Serialization/EditorBulkData.cpp b/Engine/Source/Runtime/CoreUObject/Private/Serialization/EditorBulkData.cpp
--- a/Engine/Source/Runtime/CoreUObject/Private/Serialization/EditorBulkData.cpp
+++ b/Engine/Source/Runtime/CoreUObject/Private/Serialization/EditorBulkData.cpp
@@ -190,10 +190,38 @@ const FLinkerLoad* GetLinkerLoadFromOwner(UObject* Owner)
 	}
 }
 
-/** Utility for finding the FPackageTrailer associated with a given UObject */
-static const FPackageTrailer* GetTrailerFromOwner(UObject* Owner)
+/** Utility for finding the package path associated with a given UObject */
+static FPackagePath GetPackagePathFromOwner(UObject* Owner)
 {
 	const FLinkerLoad* Linker = GetLinkerLoadFromOwner(Owner);
+
+	if (Linker != nullptr)
+	{
+		return Linker->GetPackagePath();
+	}
+	else
+	{
+		return FPackagePath();
+	}
+}
+
+bool GCVarAllowLinkerFromArchive = true;
+static FAutoConsoleVariableRef CVar_AllowLinkerFromArchive(
+	TEXT("Serialization.AllowLinkerFromArchive"),
+	GCVarAllowLinkerFromArchive,
+	TEXT("When enabled we can attempt to find a FLinkerLoad from the archive as well as the UObject. Set to false for the old behavior")
+);
+
+/** Utility for finding the FPackageTrailer associated from either a UObject or FArchive */
+static const FPackageTrailer* FindPackageTrailer(UObject* Owner, FArchive& Ar)
+{
+	const FLinkerLoad* Linker = GetLinkerLoadFromOwner(Owner);
+
+	if (Linker == nullptr && GCVarAllowLinkerFromArchive)
+	{
+		Linker = Cast<FLinkerLoad>(Ar.GetLinker());
+	}
+
 	if (Linker != nullptr)
 	{
 		return Linker->GetPackageTrailer();
@@ -204,11 +232,16 @@ static const FPackageTrailer* GetTrailerFromOwner(UObject* Owner)
 	}
 }
 
-/** Utility for finding the package path associated with a given UObject */
-static FPackagePath GetPackagePathFromOwner(UObject* Owner)
+/** Utility for finding the package path associated from either a UObject or FArchive */
+static FPackagePath FindPackagePath(UObject* Owner, FArchive& Ar)
 {
 	const FLinkerLoad* Linker = GetLinkerLoadFromOwner(Owner);
 
+	if (Linker == nullptr && GCVarAllowLinkerFromArchive)
+	{
+		Linker = Cast<FLinkerLoad>(Ar.GetLinker());
+	}
+
 	if (Linker != nullptr)
 	{
 		return Linker->GetPackagePath();
@@ -1171,8 +1204,8 @@ void FEditorBulkData::Serialize(FArchive& Ar, UObject* Owner, bool bAllowRegiste
 			OffsetInFile = INDEX_NONE;
 			PackagePath.Empty();
 
-			const FPackageTrailer* Trailer = GetTrailerFromOwner(Owner);
-				
+			const FPackageTrailer* Trailer = FindPackageTrailer(Owner, Ar);
+
 			if (IsStoredInPackageTrailer())
 			{
 				checkf(Trailer != nullptr, TEXT("Payload was stored in a package trailer, but there no trailer loaded. [%s]"), *GetDebugNameFromOwner(Owner));
@@ -1219,7 +1252,7 @@ void FEditorBulkData::Serialize(FArchive& Ar, UObject* Owner, bool bAllowRegiste
 				FArchive* CacheableArchive = Ar.GetCacheableArchive();
 				if (Ar.IsAllowingLazyLoading() && CacheableArchive != nullptr)
 				{
-					PackagePath = GetPackagePathFromOwner(Owner);
+					PackagePath = FindPackagePath(Owner, Ar);
 				}
 				else
 				{
