[UE] Create a transient placeholder type to allow for serialization on load with IDOs enabled when a type import exists but is no longer a UClass derivative.

#Synced-CL 50036996

diff --git a/Engine/Source/Runtime/CoreUObject/Private/UObject/LinkerLoad.cpp b/Engine/Source/Runtime/CoreUObject/Private/UObject/LinkerLoad.cpp
--- a/Engine/Source/Runtime/CoreUObject/Private/UObject/LinkerLoad.cpp
+++ b/Engine/Source/Runtime/CoreUObject/Private/UObject/LinkerLoad.cpp
@@ -4626,10 +4626,22 @@ UClass* FLinkerLoad::TryCreatePlaceholderClassImport(int32 ImportIndex)
 
 				if (ClassObjectPackage)
 				{
+					FName ClassObjectName = Import.ObjectName;
+					EObjectFlags ClassObjectFlags = RF_Public;
+
+					// If it's now a non-native struct type, use a transient placeholder class type so exports can still be serialized.
+					const UScriptStruct* ExistingType = FindObjectFast<UScriptStruct>(ClassObjectPackage, Import.ObjectName, EFindObjectFlags::ExactClass);
+					if (ExistingType)
+					{
+						ClassObjectName = NAME_None;
+						ClassObjectFlags |= RF_Transient;
+						ClassObjectPackage = GetTransientPackage();
+					}
+
 					UStruct* SuperStruct = FindKnownSuperStruct(ImportTypeHierarchy);
 
 					// Create an opaque, non-native public type object that has no reflected properties.
-					UStruct* StructObject = UE::FPropertyBagRepository::CreatePropertyBagPlaceholderType(ClassObjectPackage, ImportClass, Import.ObjectName, RF_Public, SuperStruct);
+					UStruct* StructObject = UE::FPropertyBagRepository::CreatePropertyBagPlaceholderType(ClassObjectPackage, ImportClass, ClassObjectName, ClassObjectFlags, SuperStruct);
 					ClassObject = CastChecked<UClass>(StructObject);
 
 					// Patch it into the import table so that we resolve to this class for any future exports of this type.
