[CrashFix] Unkown map properties with a struct key type that is also unknown causes a crash when adding the map property to the UnknownProperty tree because the key type doesn't get the CPF_HasGetValueTypeHash flag

because this check occurs when serializing a property for the sake of adding to the unknown property tree, we don't actually care about the value. I'm introducing a new flag to skip adding to the map and rehashing it when we're serializing to the unknown property tree

#Synced-CL 50033915

diff --git a/Engine/Source/Runtime/CoreUObject/Private/UObject/Class.cpp b/Engine/Source/Runtime/CoreUObject/Private/UObject/Class.cpp
--- a/Engine/Source/Runtime/CoreUObject/Private/UObject/Class.cpp
+++ b/Engine/Source/Runtime/CoreUObject/Private/UObject/Class.cpp
@@ -1939,6 +1939,7 @@ void UStruct::SerializeVersionedTaggedProperties(FStructuredArchive::FSlot Slot,
 								// Serialize the value if it contains a struct or enum because only those may contain an unknown property.
 								if (bSerializeValue)
 								{
+									TGuardValue<bool> ScopedLoadingUnknownPropertyTree(SerializeContext->bIsLoadingToUnknownPropertyTree, true);
 									UnderlyingArchive.Seek(StartOfProperty);
 									FStructuredArchive::FSlot ValueSlotCopy = PropertyRecord.EnterField(TEXT("Value"));
 									void* TempData = TempProperty->AllocateAndInitializeValue();
diff --git a/Engine/Source/Runtime/CoreUObject/Private/UObject/PropertyMap.cpp b/Engine/Source/Runtime/CoreUObject/Private/UObject/PropertyMap.cpp
--- a/Engine/Source/Runtime/CoreUObject/Private/UObject/PropertyMap.cpp
+++ b/Engine/Source/Runtime/CoreUObject/Private/UObject/PropertyMap.cpp
@@ -679,22 +679,52 @@ void FMapProperty::SerializeItem(FStructuredArchive::FSlot Slot, TNotNull<void*>
 			CopyValuesInternal(Value, Defaults, 1);
 		}
 
+		const bool bIsLoadingToUnknownPropertyTree =
+#if WITH_EDITOR
+			FUObjectThreadContext::Get().GetSerializeContext()->bIsLoadingToUnknownPropertyTree;
+#else
+			false;
+#endif
+		uint8* TempKeyValueStorage = nullptr;
+		uint8* TempElementValueStorage = nullptr;
+		ON_SCOPE_EXIT
+		{
+			if (TempKeyValueStorage)
+			{
+				if (TempElementValueStorage)
+				{
+					ValueProp->DestroyValue(TempElementValueStorage);
+				}
+				KeyProp->DestroyValue(TempKeyValueStorage);
+				FMemory::Free(TempKeyValueStorage);
+			}
+		};
+		if (bIsLoadingToUnknownPropertyTree)
+		{
+			// When loading into the unknown property tree the key type may be an unknown struct and therefore unhashable. In that case we can't add it to the map. 
+			// Luckily we don't actually care about truly loading the data since we're just harvesting type info. We'll skip adding it and instead serialize into a temp location
+			TempKeyValueStorage = (uint8*)FMemory::Malloc(MapLayout.SetLayout.Size);
+			KeyProp->InitializeValue(TempKeyValueStorage);
+			TempElementValueStorage = TempKeyValueStorage + MapLayout.ValueOffset;
+			ValueProp->InitializeValue(TempElementValueStorage);
+		}
+
 		if (!Defaults || MapHelper.Num() == 0 || bReplaceMap) // Faster loading path when loading into an empty map or replacing the entire map
 		{
 			if (NumKeysToRemove && !bReplaceMap)
 			{
 				// Load and discard keys to remove, map is empty
-				void* TempKeyValueStorage = FMemory::Malloc(MapLayout.SetLayout.Size);
-				KeyProp->InitializeValue(TempKeyValueStorage);
+				if (!TempKeyValueStorage)
+				{
+					TempKeyValueStorage = (uint8*)FMemory::Malloc(MapLayout.SetLayout.Size);
+					KeyProp->InitializeValue(TempKeyValueStorage);
+				}
 
 				FSerializedPropertyScope SerializedProperty(UnderlyingArchive, KeyProp, this);
 				for (; NumKeysToRemove; --NumKeysToRemove)
 				{
 					KeyProp->SerializeItem(KeysToRemoveArray.EnterElement(), TempKeyValueStorage);
 				}
-
-				KeyProp->DestroyValue(TempKeyValueStorage);
-				FMemory::Free(TempKeyValueStorage);
 			}
 
 			int32 NumEntries = 0;
@@ -705,16 +735,31 @@ void FMapProperty::SerializeItem(FStructuredArchive::FSlot Slot, TNotNull<void*>
 			for (; NumEntries; --NumEntries)
 			{
 				FStructuredArchive::FRecord EntryRecord = EntriesArray.EnterElement().EnterRecord();
-				int32 Index = MapHelper.AddDefaultValue_Invalid_NeedsRehash();
+
+
+				void* KeyPtr;
+				void* ValuePtr;
+				if (bIsLoadingToUnknownPropertyTree)
+				{
+					KeyPtr = TempKeyValueStorage;
+					ValuePtr = TempElementValueStorage;
+				}
+				else
+				{
+					int32 Index = MapHelper.AddDefaultValue_Invalid_NeedsRehash();
+					KeyPtr = MapHelper.GetKeyPtr(Index);
+					ValuePtr = MapHelper.GetValuePtr(Index);
+				}
+
 				{
 					UE::FSerializedPropertyPathScope SerializedPropertyPath(Context, {UE::NAME_Key});
 					FSerializedPropertyScope SerializedProperty(UnderlyingArchive, KeyProp, this);
-					KeyProp->SerializeItem(EntryRecord.EnterField(TEXT("Key")), MapHelper.GetKeyPtr(Index));
+					KeyProp->SerializeItem(EntryRecord.EnterField(TEXT("Key")), KeyPtr);
 				}
 				{
 					UE::FSerializedPropertyPathScope SerializedPropertyPath(Context, {UE::NAME_Value});
 					FSerializedPropertyScope SerializedProperty(UnderlyingArchive, ValueProp, this);
-					ValueProp->SerializeItem(EntryRecord.EnterField(TEXT("Value")), MapHelper.GetValuePtr(Index));
+					ValueProp->SerializeItem(EntryRecord.EnterField(TEXT("Value")), ValuePtr);
 				}
 			}
 
@@ -724,26 +769,21 @@ void FMapProperty::SerializeItem(FStructuredArchive::FSlot Slot, TNotNull<void*>
 			// with even better placeholder support
 			if(KeyProp->HasAnyPropertyFlags(CPF_HasGetValueTypeHash)) 
 #endif
+			if (!bIsLoadingToUnknownPropertyTree)
 			{
 				MapHelper.Rehash();
 			}
 		}
 		else // Slower loading path that mutates non-empty map
 		{
-			uint8* TempKeyValueStorage = nullptr;
-			ON_SCOPE_EXIT
-			{
-				if (TempKeyValueStorage)
-				{
-					KeyProp->DestroyValue(TempKeyValueStorage);
-					FMemory::Free(TempKeyValueStorage);
-				}
-			};
 
 			if (NumKeysToRemove)
 			{
-				TempKeyValueStorage = (uint8*)FMemory::Malloc(MapLayout.SetLayout.Size);
-				KeyProp->InitializeValue(TempKeyValueStorage);
+				if (!TempKeyValueStorage)
+				{
+					TempKeyValueStorage = (uint8*)FMemory::Malloc(MapLayout.SetLayout.Size);
+					KeyProp->InitializeValue(TempKeyValueStorage);
+				}
 
 				UE::FSerializedPropertyPathScope SerializedPropertyPath(Context, {UE::NAME_Key});
 				FSerializedPropertyScope SerializedProperty(UnderlyingArchive, KeyProp, this);
@@ -782,7 +822,7 @@ void FMapProperty::SerializeItem(FStructuredArchive::FSlot Slot, TNotNull<void*>
 					KeyProp->SerializeItem(EntryRecord.EnterField(TEXT("Key")), TempKeyValueStorage);
 				}
 
-				void* ValuePtr = MapHelper.FindOrAdd(TempKeyValueStorage);
+				void* ValuePtr = bIsLoadingToUnknownPropertyTree ? TempElementValueStorage : MapHelper.FindOrAdd(TempKeyValueStorage);
 
 				// Deserialize value into hash map-owned memory
 				{
diff --git a/Engine/Source/Runtime/CoreUObject/Private/UObject/UObjectThreadContext.cpp b/Engine/Source/Runtime/CoreUObject/Private/UObject/UObjectThreadContext.cpp
--- a/Engine/Source/Runtime/CoreUObject/Private/UObject/UObjectThreadContext.cpp
+++ b/Engine/Source/Runtime/CoreUObject/Private/UObject/UObjectThreadContext.cpp
@@ -67,6 +67,7 @@ FUObjectSerializeContext::FUObjectSerializeContext()
 	, bTrackUnknownEnumNames(false)
 	, bImpersonateProperties(false)
 	, bIsLoadingToPropertyBagObject(false)
+	, bIsLoadingToUnknownPropertyTree(false)
 #endif
 {
 }
diff --git a/Engine/Source/Runtime/CoreUObject/Public/UObject/UObjectThreadContext.h b/Engine/Source/Runtime/CoreUObject/Public/UObject/UObjectThreadContext.h
--- a/Engine/Source/Runtime/CoreUObject/Public/UObject/UObjectThreadContext.h
+++ b/Engine/Source/Runtime/CoreUObject/Public/UObject/UObjectThreadContext.h
@@ -197,6 +197,9 @@ public:
 	UE_INTERNAL bool bImpersonateProperties;
 	/** True when serializing to a property bag placeholder object while loading. */
 	UE_INTERNAL bool bIsLoadingToPropertyBagObject;
+	/** True when serializing to an unknown property tree for determining property type info. The loaded value doesn't need to be correct so we can skip
+		loading values in certain cases for the sake of bypassing issues with incomplete info or for the sake of optimization */
+	UE_INTERNAL bool bIsLoadingToUnknownPropertyTree;
 #endif
 
 	/** Adds a new loaded object */
