PerPlatformProperties Serialization for Preview controlled by setting bClientUsePreviewPPXData in CookedPreviewData in WindowsEngine.ini

#Synced-CL 50069373

#Synced-CL 50069373

diff --git a/Engine/Source/Runtime/CoreUObject/Private/UObject/PerPlatformProperties.cpp b/Engine/Source/Runtime/CoreUObject/Private/UObject/PerPlatformProperties.cpp
--- a/Engine/Source/Runtime/CoreUObject/Private/UObject/PerPlatformProperties.cpp
+++ b/Engine/Source/Runtime/CoreUObject/Private/UObject/PerPlatformProperties.cpp
@@ -8,18 +8,25 @@
 
 #include "UObject/PerPlatformPropertiesImpl.inl"
 
-IMPLEMENT_TYPE_LAYOUT(FFreezablePerPlatformFloat);
-IMPLEMENT_TYPE_LAYOUT(FFreezablePerPlatformInt);
+#if WITH_PREVIEW_PPX_DATA
 
+FName PPXPreviewOverrides::PerPlatformOverride = NAME_None;
+FName PPXPreviewOverrides::PerPlatformGroupOverride = NAME_None;
+
+void PPXPreviewOverrides::SetPPXPreviewOverrides(FName& InPerPlatformOverride, FName& InPerPlatformGroupOverride)
+{
+	PerPlatformOverride = InPerPlatformOverride;
+	PerPlatformGroupOverride = InPerPlatformGroupOverride;
+}
+
+#endif
 
 template COREUOBJECT_API FArchive& operator<<(FArchive&, TPerPlatformProperty<FPerPlatformInt, int32, NAME_IntProperty>&);
 template COREUOBJECT_API FArchive& operator<<(FArchive&, TPerPlatformProperty<FPerPlatformFloat, float, NAME_FloatProperty>&);
 template COREUOBJECT_API FArchive& operator<<(FArchive&, TPerPlatformProperty<FPerPlatformBool, bool, NAME_BoolProperty>&);
-template COREUOBJECT_API FArchive& operator<<(FArchive&, TPerPlatformProperty<FFreezablePerPlatformFloat, float, NAME_FloatProperty>&);
 template COREUOBJECT_API void operator<<(FStructuredArchive::FSlot Slot, TPerPlatformProperty<FPerPlatformInt, int32, NAME_IntProperty>&);
 template COREUOBJECT_API void operator<<(FStructuredArchive::FSlot Slot, TPerPlatformProperty<FPerPlatformFloat, float, NAME_FloatProperty>&);
 template COREUOBJECT_API void operator<<(FStructuredArchive::FSlot Slot, TPerPlatformProperty<FPerPlatformBool, bool, NAME_BoolProperty>&);
-template COREUOBJECT_API void operator<<(FStructuredArchive::FSlot Slot, TPerPlatformProperty<FFreezablePerPlatformFloat, float, NAME_FloatProperty>&);
 
 template COREUOBJECT_API FArchive& operator<<(FArchive&, TPerPlatformProperty<FPerPlatformFrameRate, FFrameRate, NAME_FrameRate>&);
 template COREUOBJECT_API void operator<<(FStructuredArchive::FSlot Slot, TPerPlatformProperty<FPerPlatformFrameRate, FFrameRate, NAME_FrameRate>&);
@@ -41,9 +48,3 @@ FString FPerPlatformInt::ToString() const
 
 	return Result;
 }
-
-FString FFreezablePerPlatformInt::ToString() const
-{
-	return FPerPlatformInt(*this).ToString();
-}
-
diff --git a/Engine/Source/Runtime/CoreUObject/Public/UObject/PerPlatformProperties.h b/Engine/Source/Runtime/CoreUObject/Public/UObject/PerPlatformProperties.h
--- a/Engine/Source/Runtime/CoreUObject/Public/UObject/PerPlatformProperties.h
+++ b/Engine/Source/Runtime/CoreUObject/Public/UObject/PerPlatformProperties.h
@@ -16,9 +16,15 @@ PerPlatformProperties.h: Property types that can be overridden on a per-platform
 #include "Misc/FrameRate.h"
 #include "UObject/ObjectMacros.h"
 #include "UObject/ScriptMacros.h"
-
+#include "Misc/Parse.h"
+#include "Misc/CommandLine.h"
 #include "PerPlatformProperties.generated.h"
 
+// This is an unsupported feature.
+#ifndef WITH_PREVIEW_PPX_DATA
+#define WITH_PREVIEW_PPX_DATA 0
+#endif
+
 namespace PerPlatformProperty::Private
 {
 	template<typename MapType>
@@ -48,65 +54,19 @@ namespace PerPlatformProperty::Private
 		}
 	};
 
-	template<>
-	struct FNameFuncs<FMemoryImageName>
-	{
-		static FMemoryImageName NameToKey(FName Name) { return FMemoryImageName(Name); }
-
-		template<typename ValueType>
-		static void SerializePerPlatformMap(FArchive& Ar, TMemoryImageMap<FMemoryImageName, ValueType>& Map)
-		{
-			if( Ar.IsLoading())
-			{
-				TMemoryImageMap<FName, ValueType> TempMap;
-				Ar << TempMap;
-				Map.Reset();
-				for( TPair<FName, ValueType>& Pair : TempMap)
-				{
-					Map.Add(FMemoryImageName(Pair.Key), Pair.Value);
-				}
-			}
-			else
-			{
-				TMemoryImageMap<FName, ValueType> TempMap;
-				for (TPair<FMemoryImageName, ValueType>& Pair : Map)
-				{
-					TempMap.Add(FName(Pair.Key), Pair.Value);
-				}
-				Ar << TempMap;
-			}
-		}
-
-		template<typename ValueType>
-		static void SerializePerPlatformMap(FArchive& UnderlyingArchive, FStructuredArchive::FRecord& Record, TMemoryImageMap<FMemoryImageName, ValueType>& Map)
-		{
-			if (UnderlyingArchive.IsLoading())
-			{
-				TMemoryImageMap<FName, ValueType> TempMap;
-
-				Record << SA_VALUE(TEXT("PerPlatform"), TempMap);
-				Map.Reset();
-				for (TPair<FName, ValueType>& Pair : TempMap)
-				{
-					Map.Add(FMemoryImageName(Pair.Key), Pair.Value);
-				}
-			}
-			else
-			{
-				TMemoryImageMap<FName, ValueType> TempMap;
-				for (TPair<FMemoryImageName, ValueType>& Pair : Map)
-				{
-					TempMap.Add(FName(Pair.Key), Pair.Value);
-				}
-				Record << SA_VALUE(TEXT("PerPlatform"), TempMap);
-			}
-		}
-	};
-
 	template<typename MapType>
 	using KeyFuncs = FNameFuncs<typename TGetKeyType<MapType>::Type>;
 }
 
+#if WITH_PREVIEW_PPX_DATA
+struct PPXPreviewOverrides
+{
+	static FName PerPlatformOverride;
+	static FName PerPlatformGroupOverride;
+	static void SetPPXPreviewOverrides(FName& InPerPlatformOverride, FName& InPerPlatformGroupOverride);
+};
+#endif
+
 /** TPerPlatformProperty - template parent class for per-platform properties 
  *  Implements Serialize function to replace value at cook time, and 
  *  backwards-compatible loading code for properties converted from simple types.
@@ -117,24 +77,31 @@ struct TPerPlatformProperty
 	typedef _ValueType ValueType;
 	typedef _StructType StructType;
 
-#if WITH_EDITOR
+#if WITH_EDITOR || WITH_PREVIEW_PPX_DATA
 	/** Get the value for the given platform (using standard "ini" name, so Windows, not Win64 or WindowsClient), which can be used to lookup the group */
 	_ValueType GetValueForPlatform(FName PlatformName) const
 	{
 		const _StructType* This = StaticCast<const _StructType*>(this);
-
 		using MapType = decltype(This->PerPlatform);
 		using KeyFuncs = typename PerPlatformProperty::Private::KeyFuncs<MapType>;
 
-		const _ValueType* Ptr = This->PerPlatform.Find(KeyFuncs::NameToKey(PlatformName));
-
+		const _ValueType* Ptr = PlatformName != NAME_None ? This->PerPlatform.Find(PlatformName) : nullptr;
 		if (Ptr == nullptr)
 		{
+#if WITH_EDITOR
 			const FDataDrivenPlatformInfo& Info = FDataDrivenPlatformInfoRegistry::GetPlatformInfo(PlatformName);
 			if (Info.PlatformGroupName != NAME_None)
 			{
-				Ptr = This->PerPlatform.Find(KeyFuncs::NameToKey(Info.PlatformGroupName));
+				Ptr = This->PerPlatform.Find(Info.PlatformGroupName);
 			}
+#else 
+			#if WITH_PREVIEW_PPX_DATA
+			if(PPXPreviewOverrides::PerPlatformGroupOverride != NAME_None)
+			{
+				Ptr = This->PerPlatform.Find(PPXPreviewOverrides::PerPlatformGroupOverride);
+			}
+			#endif
+#endif
 		}
 
 		return Ptr ? *Ptr : This->Default;
@@ -160,8 +127,17 @@ struct TPerPlatformProperty
 		else		
 #endif
 		{
-			const _StructType* This = StaticCast<const _StructType*>(this);
-			return This->Default;
+#if WITH_PREVIEW_PPX_DATA
+			if (PPXPreviewOverrides::PerPlatformOverride != NAME_None || PPXPreviewOverrides::PerPlatformGroupOverride != NAME_None)
+			{
+				return GetValueForPlatform(PPXPreviewOverrides::PerPlatformOverride);
+			}
+			else
+#endif
+			{
+				const _StructType* This = StaticCast<const _StructType*>(this);
+				return This->Default;
+			}
 		}
 	}
 
@@ -194,14 +170,11 @@ struct TPerPlatformProperty
 	}
 };
 
-
 template<typename _StructType, typename _ValueType, EName _BasePropertyName>
 FArchive& operator<<(FArchive& Ar, TPerPlatformProperty<_StructType, _ValueType, _BasePropertyName>& P);
 template<typename _StructType, typename _ValueType, EName _BasePropertyName>
 void operator<<(FStructuredArchive::FSlot Slot, TPerPlatformProperty<_StructType, _ValueType, _BasePropertyName>& P);
 
-struct FFreezablePerPlatformInt;
-
 /** FPerPlatformInt - int32 property with per-platform overrides */
 USTRUCT(BlueprintType)
 struct FPerPlatformInt
@@ -214,7 +187,7 @@ struct FPerPlatformInt
 	UPROPERTY(BlueprintReadOnly, EditAnywhere, Category = PerPlatform)
 	int32 Default;
 
-#if WITH_EDITORONLY_DATA
+#if WITH_EDITORONLY_DATA || WITH_PREVIEW_PPX_DATA
 	UPROPERTY(EditAnywhere, Category = PerPlatform, meta = (AllowedInOptional))
 	TMap<FName, int32> PerPlatform;
 #endif
@@ -231,53 +204,8 @@ struct FPerPlatformInt
 
 	COREUOBJECT_API FString ToString() const;
 
-	FPerPlatformInt(const FFreezablePerPlatformInt& Other);
-};
-
-USTRUCT()
-struct FFreezablePerPlatformInt
-#if CPP
-	: public TPerPlatformProperty<FFreezablePerPlatformInt, int32, NAME_IntProperty>
-#endif
-{
-	DECLARE_TYPE_LAYOUT(FFreezablePerPlatformInt, NonVirtual);
-
-	GENERATED_USTRUCT_BODY()
-
-public:
-	using KeyType = FMemoryImageName;
-	using FPerPlatformMap = TMemoryImageMap<FMemoryImageName, int32>;
-
-	LAYOUT_FIELD(int32, Default);
-	LAYOUT_FIELD_EDITORONLY(FPerPlatformMap, PerPlatform);
-
-	FFreezablePerPlatformInt() : Default(0) {}
-	FFreezablePerPlatformInt(int32 InDefaultValue) : Default(InDefaultValue) {}
-	FFreezablePerPlatformInt(const FPerPlatformInt& Other)
-		: Default(Other.Default)
-	{
-#if WITH_EDITORONLY_DATA
-		for (const TPair<FName, int32>& Pair : Other.PerPlatform)
-		{
-			PerPlatform.Add(FMemoryImageName(Pair.Key), Pair.Value);
-		}
-#endif
-	}
-
-	COREUOBJECT_API FString ToString() const;
 };
 
-inline FPerPlatformInt::FPerPlatformInt(const FFreezablePerPlatformInt& Other)
-	: Default(Other.Default)
-{
-#if WITH_EDITORONLY_DATA
-	for (const TPair<FMemoryImageName, int32>& Pair : Other.PerPlatform)
-	{
-		PerPlatform.Add(FName(Pair.Key), Pair.Value);
-	}
-#endif
-}
-
 extern template COREUOBJECT_API FArchive& operator<<(FArchive&, TPerPlatformProperty<FPerPlatformInt, int32, NAME_IntProperty>&);
 extern template COREUOBJECT_API void operator<<(FStructuredArchive::FSlot Slot, TPerPlatformProperty<FPerPlatformInt, int32, NAME_IntProperty>&);
 
@@ -293,8 +221,6 @@ struct TStructOpsTypeTraits<FPerPlatformInt>
 	static constexpr EPropertyObjectReferenceType WithSerializerObjectReferences = EPropertyObjectReferenceType::None;
 };
 
-struct FFreezablePerPlatformFloat;
-
 /** FPerPlatformFloat - float property with per-platform overrides */
 USTRUCT(BlueprintType, meta = (CanFlattenStruct))
 struct FPerPlatformFloat
@@ -307,7 +233,7 @@ struct FPerPlatformFloat
 	UPROPERTY(BlueprintReadOnly, EditAnywhere, Category = PerPlatform)
 	float Default;
 
-#if WITH_EDITORONLY_DATA
+#if WITH_EDITORONLY_DATA || WITH_PREVIEW_PPX_DATA
 	UPROPERTY(EditAnywhere, Category = PerPlatform, meta = (AllowedInOptional))
 	TMap<FName, float> PerPlatform;
 #endif
@@ -322,47 +248,9 @@ struct FPerPlatformFloat
 	{
 	}
 
-	FPerPlatformFloat(const FFreezablePerPlatformFloat& Other);
 };
 extern template COREUOBJECT_API FArchive& operator<<(FArchive&, TPerPlatformProperty<FPerPlatformFloat, float, NAME_FloatProperty>&);
 
-struct FFreezablePerPlatformFloat
-#if CPP
-	: public TPerPlatformProperty<FFreezablePerPlatformFloat, float, NAME_FloatProperty>
-#endif
-{
-	DECLARE_TYPE_LAYOUT(FFreezablePerPlatformFloat, NonVirtual);
-public:
-	using FPerPlatformMap = TMemoryImageMap<FMemoryImageName, float>;
-	
-	LAYOUT_FIELD(float, Default);
-	LAYOUT_FIELD_EDITORONLY(FPerPlatformMap, PerPlatform);
-
-	FFreezablePerPlatformFloat() : Default(0.0f) {}
-	FFreezablePerPlatformFloat(float InDefaultValue) : Default(InDefaultValue) {}
-	FFreezablePerPlatformFloat(const FPerPlatformFloat& Other)
-		: Default(Other.Default)
-	{
-#if WITH_EDITORONLY_DATA
-		for (const TPair<FName, float>& Pair : Other.PerPlatform)
-		{
-			PerPlatform.Add(FMemoryImageName(Pair.Key), Pair.Value);
-		}
-#endif
-}
-};
-
-inline FPerPlatformFloat::FPerPlatformFloat(const FFreezablePerPlatformFloat& Other)
-	: Default(Other.Default)
-{
-#if WITH_EDITORONLY_DATA
-	for (const TPair<FMemoryImageName, float>& Pair : Other.PerPlatform)
-	{
-		PerPlatform.Add(FName(Pair.Key), Pair.Value);
-	}
-#endif
-}
-
 template<>
 struct TStructOpsTypeTraits<FPerPlatformFloat>
 :	public TStructOpsTypeTraitsBase2<FPerPlatformFloat>
@@ -387,7 +275,7 @@ struct FPerPlatformBool
 	UPROPERTY(BlueprintReadOnly, EditAnywhere, Category = PerPlatform)
 	bool Default;
 
-#if WITH_EDITORONLY_DATA
+#if WITH_EDITORONLY_DATA || WITH_PREVIEW_PPX_DATA
 	UPROPERTY(EditAnywhere, Category = PerPlatform, meta = (AllowedInOptional))
 	TMap<FName, bool> PerPlatform;
 #endif
@@ -428,7 +316,7 @@ struct FPerPlatformFrameRate
 	UPROPERTY(EditAnywhere, Category = PerPlatform)
 	FFrameRate Default;
 
-#if WITH_EDITORONLY_DATA
+#if WITH_EDITORONLY_DATA || WITH_PREVIEW_PPX_DATA
 	UPROPERTY(EditAnywhere, Category = PerPlatform, meta = (AllowedInOptional))
 	TMap<FName, FFrameRate> PerPlatform;
 #endif
diff --git a/Engine/Source/Runtime/CoreUObject/Public/UObject/PerPlatformPropertiesImpl.inl b/Engine/Source/Runtime/CoreUObject/Public/UObject/PerPlatformPropertiesImpl.inl
--- a/Engine/Source/Runtime/CoreUObject/Public/UObject/PerPlatformPropertiesImpl.inl
+++ b/Engine/Source/Runtime/CoreUObject/Public/UObject/PerPlatformPropertiesImpl.inl
@@ -14,19 +14,37 @@
 
 #if WITH_EDITOR
 #include "Interfaces/ITargetPlatform.h"
+#include "Misc/DataDrivenPlatformInfoRegistry.h"
+#include "Misc/ConfigCacheIni.h"
 #endif
 
 /** Serializer to cook out the most appropriate platform override */
 template<typename StructType, typename ValueType, EName BasePropertyName>
 FArchive& operator<<(FArchive& Ar, TPerPlatformProperty<StructType, ValueType, BasePropertyName>& Property)
 {
-	bool bCooked = false;
+	bool bStripSerializeAllProperties = true;
+#if WITH_EDITOR
+	bool bClientUsePreviewPPXData = false;
+	if(Ar.IsCooking())
+	{
+		bool bIsClientOnly = Ar.CookingTarget()->IsClientOnly();
+		bool bIsServerOnly = Ar.CookingTarget()->IsServerOnly();
+		bool bHasEditorOnlyData = Ar.CookingTarget()->HasEditorOnlyData();
+		if(bIsClientOnly || !(bIsClientOnly || bIsServerOnly || bHasEditorOnlyData))
+		{
+			bClientUsePreviewPPXData = Ar.CookingTarget()->GetConfigSystem()->GetBoolOrDefault(TEXT("CookedPreviewData"), TEXT("bClientUsePreviewPPXData"), false, GEngineIni);		
+		}
+	}
+	bStripSerializeAllProperties = !bClientUsePreviewPPXData && Ar.IsCooking() && !Ar.IsSavingOptionalObject();
+#elif WITH_PREVIEW_PPX_DATA
+	bStripSerializeAllProperties = false;
+#endif
+	
+	StructType* This = StaticCast<StructType*>(&Property);
+	Ar << bStripSerializeAllProperties;
 #if WITH_EDITOR
-	// we don't want to lose this EditorOnly property when saving an Optional object (which normally would be stripped out when cooking)
 	if (Ar.IsCooking() && !Ar.IsSavingOptionalObject())
 	{
-		bCooked = true;
-		Ar << bCooked;
 		// Save out platform override if it exists and Default otherwise
 		ValueType Value = Property.GetValueForPlatform(*Ar.CookingTarget()->IniPlatformName());
 		Ar << Value;
@@ -34,18 +52,39 @@ FArchive& operator<<(FArchive& Ar, TPerPlatformProperty<StructType, ValueType, B
 	else
 #endif
 	{
-		StructType* This = StaticCast<StructType*>(&Property);
-		Ar << bCooked;
 		Ar << This->Default;
+	}
+
+	if (!bStripSerializeAllProperties)
+	{	
 #if WITH_EDITORONLY_DATA
-		if (!bCooked)
+		if (Ar.IsCooking() && !Ar.IsSavingOptionalObject())
+		{
+			TMap<FName, ValueType> ToBeCookedPerPlatform(This->PerPlatform);
+			for (auto It = ToBeCookedPerPlatform.CreateIterator(); It; ++It)
+			{
+				if (FDataDrivenPlatformInfoRegistry::GetPlatformInfo(It.Key()).bIsConfidential)
+				{
+					It.RemoveCurrent();
+				}
+			}
+			Ar << ToBeCookedPerPlatform;
+		}
+		else
 		{
-			using MapType = decltype(This->PerPlatform);
-			using KeyFuncs = typename PerPlatformProperty::Private::KeyFuncs<MapType>;
-			KeyFuncs::SerializePerPlatformMap(Ar, This->PerPlatform);
+			Ar << This->PerPlatform;
 		}
+#elif WITH_PREVIEW_PPX_DATA
+		{
+			Ar << This->PerPlatform;
+		}
+#else
+		// If we have old assets that were serialized with a PerPlatform map, but we no longer need it.
+		TMap<FName, ValueType> TempPerPlatform;
+		Ar << TempPerPlatform;
 #endif
 	}
+
 	return Ar;
 }
 
@@ -53,32 +92,65 @@ FArchive& operator<<(FArchive& Ar, TPerPlatformProperty<StructType, ValueType, B
 template<typename StructType, typename ValueType, EName BasePropertyName>
 void operator<<(FStructuredArchive::FSlot Slot, TPerPlatformProperty<StructType, ValueType, BasePropertyName>& Property)
 {
-	FArchive& UnderlyingArchive = Slot.GetUnderlyingArchive();
+	FArchive& Ar = Slot.GetUnderlyingArchive();
 	FStructuredArchive::FRecord Record = Slot.EnterRecord();
-
-	bool bCooked = false;
+	bool bStripSerializeAllProperties = true;
 #if WITH_EDITOR
-	if (UnderlyingArchive.IsCooking())
+	bool bClientUsePreviewPPXData = false;
+	if (Ar.IsCooking())
+	{
+		bool bIsClientOnly = Ar.CookingTarget()->IsClientOnly();
+		bool bIsServerOnly = Ar.CookingTarget()->IsServerOnly();
+		bool bHasEditorOnlyData = Ar.CookingTarget()->HasEditorOnlyData();
+		if (bIsClientOnly || !(bIsClientOnly || bIsServerOnly || bHasEditorOnlyData))
+		{
+			bClientUsePreviewPPXData = Ar.CookingTarget()->GetConfigSystem()->GetBoolOrDefault(TEXT("CookedPreviewData"), TEXT("bClientUsePreviewPPXData"), false, GEngineIni);
+		}
+	}
+	bStripSerializeAllProperties = !bClientUsePreviewPPXData && Ar.IsCooking() && !Ar.IsSavingOptionalObject();
+#elif WITH_PREVIEW_PPX_DATA
+	bStripSerializeAllProperties = false;
+#endif
+	StructType* This = StaticCast<StructType*>(&Property);
+	Record << SA_VALUE(TEXT("bCooked"), bStripSerializeAllProperties);
+#if WITH_EDITOR
+	if (Ar.IsCooking() && !Ar.IsSavingOptionalObject())
 	{
-		bCooked = true;
-		Record << SA_VALUE(TEXT("bCooked"), bCooked);
 		// Save out platform override if it exists and Default otherwise
-		ValueType Value = Property.GetValueForPlatform(*UnderlyingArchive.CookingTarget()->IniPlatformName());
+		ValueType Value = Property.GetValueForPlatform(*Ar.CookingTarget()->IniPlatformName());
 		Record << SA_VALUE(TEXT("Value"), Value);
 	}
 	else
 #endif
 	{
-		StructType* This = StaticCast<StructType*>(&Property);
-		Record << SA_VALUE(TEXT("bCooked"), bCooked);
 		Record << SA_VALUE(TEXT("Value"), This->Default);
+	}
+
+	if (!bStripSerializeAllProperties)
+	{
 #if WITH_EDITORONLY_DATA
-		if (!bCooked)
+		if (Ar.IsCooking() && !Ar.IsSavingOptionalObject())
+		{
+			TMap<FName, ValueType> ToBeCookedPerPlatform(This->PerPlatform);
+			for (auto It = ToBeCookedPerPlatform.CreateIterator(); It; ++It)
+			{
+				if (FDataDrivenPlatformInfoRegistry::GetPlatformInfo(It.Key()).bIsConfidential)
+				{
+					It.RemoveCurrent();
+				}
+			}
+			Record << SA_VALUE(TEXT("PerPlatform"), ToBeCookedPerPlatform);
+		}
+		else
 		{
-			using MapType = decltype(This->PerPlatform);
-			using KeyFuncs = typename PerPlatformProperty::Private::KeyFuncs<MapType>;
-			KeyFuncs::SerializePerPlatformMap(UnderlyingArchive, Record, This->PerPlatform);
+			Record << SA_VALUE(TEXT("PerPlatform"), This->PerPlatform);
 		}
+#elif WITH_PREVIEW_PPX_DATA
+		Record << SA_VALUE(TEXT("PerPlatform"), This->PerPlatform);
+#else
+		// If we have old assets that were serialized with a PerPlatform map, but we no longer need it.
+		TMap<FName, ValueType> TempPerPlatform;
+		Record << SA_VALUE(TEXT("PerPlatform"), TempPerPlatform);
 #endif
 	}
 }
