Remove FEditorBulkData::SetVirtualizationOptOut and all associated code.

#rnx

- This code has been deprecated since 5.0 to avoid people using it and since it is unlikely at this point that it will ever be turned into a proper feature it's time to just remove it.

#Synced-CL 49815329

diff --git a/Engine/Source/Runtime/CoreUObject/Private/Serialization/EditorBulkData.cpp b/Engine/Source/Runtime/CoreUObject/Private/Serialization/EditorBulkData.cpp
--- a/Engine/Source/Runtime/CoreUObject/Private/Serialization/EditorBulkData.cpp
+++ b/Engine/Source/Runtime/CoreUObject/Private/Serialization/EditorBulkData.cpp
@@ -113,23 +113,6 @@ static bool ShouldLogCookedStatus()
 	return ConfigSetting.bEnabled;
 }
 
-#if UE_ENABLE_VIRTUALIZATION_TOGGLE
-bool ShouldAllowVirtualizationOptOut()
-{
-	static struct FAllowVirtualizationOptOut
-	{
-		bool bEnabled = true;
-
-		FAllowVirtualizationOptOut()
-		{
-			GConfig->GetBool(TEXT("Core.System.Experimental"), TEXT("AllowVirtualizationOptOut"), bEnabled, GEngineIni);
-		}
-	} AllowVirtualizationOptOut;
-
-	return AllowVirtualizationOptOut.bEnabled;
-}
-#endif // UE_ENABLE_VIRTUALIZATION_TOGGLE
-
 /** Utility for logging extended error messages when we fail to open a package for reading */
 static void LogPackageOpenFailureMessage(const FPackagePath& PackagePath, EPackageSegment PackageSegment)
 {
@@ -1984,13 +1967,6 @@ void FEditorBulkData::SerializeToPackageTrailer(FLinkerSave& LinkerSave, FCompre
 		PayloadFilter = UE::Virtualization::IVirtualizationSystem::Get().FilterPayload(Owner);
 	}
 
-#if UE_ENABLE_VIRTUALIZATION_TOGGLE
-	if (bSkipVirtualization)
-	{
-		PayloadFilter |= UE::Virtualization::EPayloadFilterReason::Asset;
-	}
-#endif //UE_ENABLE_VIRTUALIZATION_TOGGLE
-
 	LinkerSave.PackageTrailerBuilder->AddPayload(PayloadContentId, MoveTemp(PayloadToSerialize), PayloadFilter, MoveTemp(OnPayloadWritten));
 }
 
@@ -2302,19 +2278,6 @@ bool FEditorBulkData::LocationMatches(const FEditorBulkData& Other) const
 	}
 }
 
-
-#if UE_ENABLE_VIRTUALIZATION_TOGGLE
-
-void FEditorBulkData::SetVirtualizationOptOut(bool bOptOut)
-{
-	if (ShouldAllowVirtualizationOptOut())
-	{
-		bSkipVirtualization = bOptOut;
-	}	
-}
-
-#endif //UE_ENABLE_VIRTUALIZATION_TOGGLE
-
 void FEditorBulkData::UpdateKeyIfNeeded()
 {
 	// If this was created from old BulkData then the key is generated from an older FGuid, we
diff --git a/Engine/Source/Runtime/CoreUObject/Public/Serialization/EditorBulkData.h b/Engine/Source/Runtime/CoreUObject/Public/Serialization/EditorBulkData.h
--- a/Engine/Source/Runtime/CoreUObject/Public/Serialization/EditorBulkData.h
+++ b/Engine/Source/Runtime/CoreUObject/Public/Serialization/EditorBulkData.h
@@ -27,10 +27,6 @@ struct FPackageFileVersion;
 
 namespace UE::BulkDataRegistry { enum class ERegisterResult : uint8; }
 
-// When enabled it will be possible for specific editor bulkdata objects to opt out of being virtualized
-// This is a development feature and not expected to be used!
-#define UE_ENABLE_VIRTUALIZATION_TOGGLE 1
-
 //TODO: At some point it might be a good idea to uncomment this to make sure that FEditorBulkData is
 //		never used at runtime (requires a little too much reworking of some assets for now though)
 //#if WITH_EDITORONLY_DATA
@@ -358,11 +354,6 @@ public:
 	 */
 	COREUOBJECT_API void UpdateRegistrationOwner(UObject* Owner);
 
-#if UE_ENABLE_VIRTUALIZATION_TOGGLE
-	UE_DEPRECATED(5.0, "SetVirtualizationOptOut is an internal feature for development and will be removed without warning!")
-	COREUOBJECT_API void SetVirtualizationOptOut(bool bOptOut);
-#endif //UE_ENABLE_VIRTUALIZATION_TOGGLE
-
 protected:
 	enum class ETornOff {};
 	COREUOBJECT_API FEditorBulkData(const FEditorBulkData& Other, ETornOff);
@@ -548,7 +539,7 @@ private:
 	 */
 	static COREUOBJECT_API void ValidatePackageTrailerBuilder(const FLinkerSave* LinkerSave, const FIoHash& Id, EFlags PayloadFlags);
 
-	/** Returns true if we should use legacy serialization instead of the FPackageTrailer system. This can be removed when UE_ENABLE_VIRTUALIZATION_TOGGLE is removed. */
+	/** Returns true if we should use legacy serialization instead of the FPackageTrailer system. */
 	COREUOBJECT_API bool ShouldUseLegacySerialization(const FLinkerSave* LinkerSave) const;
 
 	/** Unique identifier for the bulkdata object itself */
@@ -580,10 +571,6 @@ private:
 	// Use of IBulkDataRegistry requires a recursive mutex
 	mutable FRecursiveMutex Mutex;
 
-#if UE_ENABLE_VIRTUALIZATION_TOGGLE
-	bool bSkipVirtualization = false;
-#endif //UE_ENABLE_VIRTUALIZATION_TOGGLE
-
 	/** 
 	 * Compression settings to be applied to the payload when the package is next saved. The settings will be reset if
 	 * the payload is unloaded from memory during serialization (i.e. the payload was virtualized or the package was
