[GCObjectReferencer/ALT]
- Add missing lock preventing other threads to add and remove objects in the list during iteration

#Synced-CL 50201644

diff --git a/Engine/Source/Runtime/CoreUObject/Private/Misc/GCObjectReferencer.cpp b/Engine/Source/Runtime/CoreUObject/Private/Misc/GCObjectReferencer.cpp
--- a/Engine/Source/Runtime/CoreUObject/Private/Misc/GCObjectReferencer.cpp
+++ b/Engine/Source/Runtime/CoreUObject/Private/Misc/GCObjectReferencer.cpp
@@ -57,8 +57,8 @@ void UGCObjectReferencer::AddReferencedObjects(UObject* InThis, FReferenceCollec
 
 	check(!This->CurrentlySerializingObject);
 
-	// Note we're not locking ReferencedObjectsCritical here because we guard
-	// against adding new references during GC in AddObject and RemoveObject.
+	// Only GC protects us against adding new UObjects, but this is called through other means outside of GC that needs to be protected.
+	UE::TConditionalScopeLock ReferencedObjectsLock(This->Impl->ReferencedObjectsCritical, !IsGarbageCollecting());
 	// Let each registered object handle its AddReferencedObjects call
 	
 	if (Collector.NeedsInitialReferences())
