RemoteObjects:  Ensure that conversion from TArray<TObjectPtr<Blah>> to TArray<Blah*> results in a resolve on the elements.

#rnx

#Synced-CL 50214038

diff --git a/Engine/Source/Runtime/CoreUObject/Public/UObject/ObjectHandleTracking.h b/Engine/Source/Runtime/CoreUObject/Public/UObject/ObjectHandleTracking.h
--- a/Engine/Source/Runtime/CoreUObject/Public/UObject/ObjectHandleTracking.h
+++ b/Engine/Source/Runtime/CoreUObject/Public/UObject/ObjectHandleTracking.h
@@ -129,7 +129,7 @@ namespace UE::CoreUObject::Private
 namespace UE::CoreUObject::Private
 {
 	inline void OnHandleRead(const UObject* Object) { }
-	inline void OnHandleRead(TArrayView<const UObject*> Objects) { }
+	inline void OnHandleRead(TArrayView<const UObject* const> Objects) { }
 	inline void OnClassReferenceResolved(const FObjectRef& ObjectRef, UPackage* Package, UClass* Class) { }
 	inline void OnReferenceResolved(const FObjectRef& ObjectRef, UPackage* Package, UClass* Class) { }
 	inline void OnReferenceLoaded(const FObjectRef& ObjectRef, UPackage* Package, UObject* Object) { }
diff --git a/Engine/Source/Runtime/CoreUObject/Public/UObject/ObjectPtr.h b/Engine/Source/Runtime/CoreUObject/Public/UObject/ObjectPtr.h
--- a/Engine/Source/Runtime/CoreUObject/Public/UObject/ObjectPtr.h
+++ b/Engine/Source/Runtime/CoreUObject/Public/UObject/ObjectPtr.h
@@ -1175,11 +1175,11 @@ struct TContainerElementTypeCompatibility<TObjectPtr<T>>
 	UE_OBJPTR_DEPRECATED(5.0, "Reinterpretation between ranges of one type to another type is deprecated.")
 	static void ReinterpretRange(IterBeginType Iter, IterEndType IterEnd, OperatorType Operator = TContainerElementTypeCompatibility_DefaultOperatorFunc)
 	{
-#if UE_WITH_OBJECT_HANDLE_LATE_RESOLVE || UE_WITH_OBJECT_HANDLE_TRACKING
+#if UE_WITH_OBJECT_HANDLE_LATE_RESOLVE || UE_WITH_OBJECT_HANDLE_TRACKING || UE_WITH_REMOTE_OBJECT_HANDLE
 		while (Iter != IterEnd)
 		{
 			(void)Operator(Iter).Get();
-			check(Operator(Iter).IsResolved());
+			check(UE_WITH_REMOTE_OBJECT_HANDLE || Operator(Iter).IsResolved());
 			++Iter;
 		}
 #endif
@@ -1189,14 +1189,14 @@ struct TContainerElementTypeCompatibility<TObjectPtr<T>>
 	UE_OBJPTR_DEPRECATED(5.0, "Reinterpretation between ranges of one type to another type is deprecated.")
 	static void ReinterpretRangeContiguous(IterBeginType Iter, IterEndType IterEnd, SizeType Size, OperatorType Operator = TContainerElementTypeCompatibility_DefaultOperatorFunc)
 	{
-#if UE_WITH_OBJECT_HANDLE_LATE_RESOLVE || UE_WITH_OBJECT_HANDLE_TRACKING
+#if UE_WITH_OBJECT_HANDLE_LATE_RESOLVE || UE_WITH_OBJECT_HANDLE_TRACKING || UE_WITH_REMOTE_OBJECT_HANDLE
 		const TObjectPtr<T>* Begin = &*Iter;
 		while (Iter != IterEnd)
 		{
 			auto& Ptr = Operator(Iter);
 			const FObjectPtr& ObjPtr = reinterpret_cast<const FObjectPtr&>(Ptr);
 			UE::CoreUObject::Private::ResolveObjectHandleNoRead(ObjPtr.GetHandleRef());
-			check(ObjPtr.IsResolved());
+			check(UE_WITH_REMOTE_OBJECT_HANDLE || ObjPtr.IsResolved());
 			++Iter;
 		}
 		const UObject* const* ObjPtr = reinterpret_cast<const UObject* const*>(Begin);
@@ -1218,11 +1218,11 @@ struct TContainerElementTypeCompatibility<const TObjectPtr<T>>
 	UE_OBJPTR_DEPRECATED(5.0, "Reinterpretation between ranges of one type to another type is deprecated.")
 	static void ReinterpretRange(IterBeginType Iter, IterEndType IterEnd, OperatorType Operator = TContainerElementTypeCompatibility_DefaultOperatorFunc)
 	{
-#if UE_WITH_OBJECT_HANDLE_LATE_RESOLVE || UE_WITH_OBJECT_HANDLE_TRACKING
+#if UE_WITH_OBJECT_HANDLE_LATE_RESOLVE || UE_WITH_OBJECT_HANDLE_TRACKING || UE_WITH_REMOTE_OBJECT_HANDLE
 		while (Iter != IterEnd)
 		{
 			Operator(Iter).Get();
-			check(Operator(Iter).IsResolved());
+			check(UE_WITH_REMOTE_OBJECT_HANDLE || Operator(Iter).IsResolved());
 			++Iter;
 		}
 #endif
@@ -1232,14 +1232,14 @@ struct TContainerElementTypeCompatibility<const TObjectPtr<T>>
 	UE_OBJPTR_DEPRECATED(5.0, "Reinterpretation between ranges of one type to another type is deprecated.")
 	static void ReinterpretRangeContiguous(IterBeginType Iter, IterEndType IterEnd, SizeType Size, OperatorType Operator = TContainerElementTypeCompatibility_DefaultOperatorFunc)
 	{
-#if UE_WITH_OBJECT_HANDLE_LATE_RESOLVE || UE_WITH_OBJECT_HANDLE_TRACKING
+#if UE_WITH_OBJECT_HANDLE_LATE_RESOLVE || UE_WITH_OBJECT_HANDLE_TRACKING || UE_WITH_REMOTE_OBJECT_HANDLE
 		const TObjectPtr<T>* Begin = &*Iter;
 		while (Iter != IterEnd)
 		{
 			auto& Ptr = Operator(Iter);
 			const FObjectPtr& ObjPtr = reinterpret_cast<const FObjectPtr&>(Ptr);
 			UE::CoreUObject::Private::ResolveObjectHandleNoRead(ObjPtr.GetHandleRef());
-			check(ObjPtr.IsResolved());
+			check(UE_WITH_REMOTE_OBJECT_HANDLE || ObjPtr.IsResolved());
 			++Iter;
 		}
 		const UObject* const* ObjPtr = reinterpret_cast<const UObject* const*>(Begin);
