PackageName: Make a publicly accessible and refcounted structure for mountpoints, so that systems such as the AssetRegistry can hold a pointer to the structure to use as a key for data associated with the mountpoint.
#rnx

#Synced-CL 50149715

diff --git a/Engine/Source/Runtime/CoreUObject/Private/Misc/PackageName.cpp b/Engine/Source/Runtime/CoreUObject/Private/Misc/PackageName.cpp
--- a/Engine/Source/Runtime/CoreUObject/Private/Misc/PackageName.cpp
+++ b/Engine/Source/Runtime/CoreUObject/Private/Misc/PackageName.cpp
@@ -8,8 +8,8 @@
 
 #include "Algo/Find.h"
 #include "Algo/FindLast.h"
+#include "Algo/Sort.h"
 #include "AssetRegistry/AssetData.h"
-#include "Containers/DirectoryTree.h"
 #include "Containers/StringView.h"
 #include "Containers/VersePath.h"
 #include "GenericPlatform/GenericPlatformFile.h"
@@ -24,6 +24,7 @@
 #include "Misc/AssetRegistryInterface.h"
 #include "Misc/AutomationTest.h"
 #include "Misc/CoreDelegates.h"
+#include "Misc/CoreUObjectPluginManager.h"
 #include "Misc/EnumClassFlags.h"
 #include "Misc/PackageName/MountPointPrivate.h"
 #include "Misc/PackageName/PackageNamePrivate.h"
@@ -31,7 +32,6 @@
 #include "Misc/PackageSegment.h"
 #include "Misc/Paths.h"
 #include "Misc/PathViews.h"
-#include "Misc/CoreUObjectPluginManager.h"
 #include "Misc/ScopeRWLock.h"
 #include "Misc/TransactionallySafeRWLock.h"
 #include "Misc/StringBuilder.h"
@@ -185,61 +185,144 @@ void FLongPackagePathsSingleton::OnCoreUObjectInitialized()
 	UE::CoreUObject::Private::PluginHandler::Install();
 }
 
-UE::PackageName::FMountPoint FLongPackagePathsSingleton::ConstructMountPoint(const FString& RootPath, const FString& ContentPath,
-	UE::PackageName::EMountFlags Flags)
+bool FLongPackagePathsSingleton::TryNormalizePaths(const FString& InLongPackageName, const FString& InLocalPath,
+	FString& OutLongPackageName, FString& OutLocalPathStandard, FString& OutLocalPathAbsolute, bool bSilenceLogs)
 {
-	using namespace UE::PackageName;
+	if (!TryNormalizeLongPackageNamePath(InLongPackageName, OutLongPackageName, bSilenceLogs))
+	{
+		return false;
+	}
+	if (!TryNormalizeLocalPath(InLocalPath, OutLocalPathStandard, OutLocalPathAbsolute, bSilenceLogs))
+	{
+		return false;
+	}
+	return true;
+}
+
+bool FLongPackagePathsSingleton::TryNormalizeLongPackageNamePath(const FString& InLongPackageName,
+	FString& OutLongPackageName, bool bSilenceLogs)
+{
+	int32 NumMidSlashes = 0;
+	int32 NumBackSlashes = 0;
+	int32 LenLongPackageName = InLongPackageName.Len();
+	for (int32 N = 0; N < LenLongPackageName; ++N)
+	{
+		TCHAR C = (*InLongPackageName)[N];
+		if (N != 0 && N != LenLongPackageName - 1 && C == '/')
+		{
+			++NumMidSlashes;
+		}
+		else if (C == '\\')
+		{
+			++NumBackSlashes;
+		}
+	}
+	if (!InLongPackageName.StartsWith("/") || NumMidSlashes > 0 || NumBackSlashes > 0)
+	{
+		if (!bSilenceLogs)
+		{
+			UE_LOG(LogPackageName, Error,
+				TEXT("Invalid LongPackageName for mount point: \"%s\". Callstack:"), *InLongPackageName);
+			FDebug::DumpStackTraceToLog(ELogVerbosity::Display);
+		}
+		return false;
+	}
 
-	FString MountPath = ContentPath;
+	OutLongPackageName = InLongPackageName;
+	if (!OutLongPackageName.EndsWith("/"))
+	{
+		OutLongPackageName += TEXT("/");
+	}
+	return true;
+}
 
+bool FLongPackagePathsSingleton::TryNormalizeLocalPath(const FString& InLocalPath,
+	FString& OutLocalPathStandard, FString& OutLocalPathAbsolute, bool bSilenceLogs)
+{
 	// If a relative path is passed, convert to an absolute path 
-	if (FPaths::IsRelative(MountPath))
+	OutLocalPathAbsolute = InLocalPath;
+	if (FPaths::IsRelative(OutLocalPathAbsolute))
 	{
-		MountPath = FPaths::ConvertRelativePathToFull(ContentPath);
+		OutLocalPathAbsolute = FPaths::ConvertRelativePathToFull(InLocalPath);
 
 		// Revert to original path if unable to convert to full path
-		if (MountPath.Len() <= 1)
+		if (OutLocalPathAbsolute.Len() <= 1)
 		{
-			MountPath = ContentPath;
-			UE_LOG(LogPackageName, Warning, TEXT("Unable to convert mount point relative path: %s"), *ContentPath);
+			OutLocalPathAbsolute = InLocalPath;
+			if (!bSilenceLogs)
+			{
+				UE_LOG(LogPackageName, Warning, TEXT("Unable to convert mount point relative path: %s"), *InLocalPath);
+			}
 		}
 	}
 
 	// Convert to a relative path using the FileManager
-	FString RelativeContentPath = IFileManager::Get().ConvertToRelativePath(*MountPath);
+	OutLocalPathStandard = IFileManager::Get().ConvertToRelativePath(*OutLocalPathAbsolute);
 
 	// Make sure the path ends in a trailing path separator for consistency.
-	if (!RelativeContentPath.EndsWith(TEXT("/"), ESearchCase::CaseSensitive))
+	if (!OutLocalPathStandard.EndsWith(TEXT("/"), ESearchCase::CaseSensitive))
 	{
-		RelativeContentPath += TEXT("/");
+		OutLocalPathStandard += TEXT("/");
 	}
 
+	// Recompute the absolute path from the standard path, to prevent the possibility of there being any text
+	// difference between a path passed in as an absolute path and a path passed in as a standard path.
 	TStringBuilder<256> AbsolutePathBuilder;
-	FPathViews::ToAbsolutePath(RelativeContentPath, AbsolutePathBuilder);
+	FPathViews::ToAbsolutePath(OutLocalPathStandard, AbsolutePathBuilder);
+	OutLocalPathAbsolute = AbsolutePathBuilder;
 
-	return FMountPoint(FString(RootPath), MoveTemp(RelativeContentPath), FString(AbsolutePathBuilder), Flags);
+	return true;
 }
 
 void FLongPackagePathsSingleton::InsertMountPoint(const FString& RootPath, const FString& ContentPath)
-{	
+{
 	using namespace UE::PackageName;
 
-	TUniquePtr<FMountPoint> MountPointOwner(MakeUnique<FMountPoint>(
-		ConstructMountPoint(RootPath, ContentPath, EMountFlags::Removable)));
-	FMountPoint* MountPoint = MountPointOwner.Get();
+	FMountPointKeyByStrings SearchKey;
+	FString LocalPathAbsolute;
+	if (!TryNormalizePaths(RootPath, ContentPath, SearchKey.LongPackageName,
+		SearchKey.LocalPathStandard, LocalPathAbsolute, false /* bSilenceLogs */))
+	{
+		return;
+	}
+	uint32 SearchHash = GetTypeHash(SearchKey);
+
+	TRefCountPtr<FMountPoint> MountPoint;
 	{
 		UE::TWriteScopeLock ScopeLock(MountLock);
-		MountPoints.Insert(MoveTemp(MountPointOwner), 0);
 
-		RootPathTree.FindOrAdd(MountPoint->GetRootPath()) = MountPoint;
-		ContentPathTree.FindOrAdd(MountPoint->GetContentPathRelative()) = MountPoint;
-		ContentPathTree.FindOrAdd(MountPoint->GetContentPathAbsolute()) = MountPoint;
-	}		
+		// Look in the existing MountPoints to see if this one already exists, so we can update the existing pointer.
+		FMountPointKeyByPtr* Existing = MountPoints.FindByHash(SearchHash, SearchKey);
+		if (Existing)
+		{
+			MountPoint = Existing->GetMountPoint();
+		}
+		else
+		{
+			MountPoint = new FMountPoint(FString(SearchKey.LongPackageName), FString(SearchKey.LocalPathStandard),
+				FString(LocalPathAbsolute), EMountFlags::Removable);
+			MountPoints.AddByHash(SearchHash, FMountPointKeyByPtr(*MountPoint));
+		}
+		RootPathTree.FindOrAdd(MountPoint->GetRootPath()) = MountPoint.GetReference();
+		ContentPathTree.FindOrAdd(MountPoint->GetContentPathRelative()) = MountPoint.GetReference();
+		ContentPathTree.FindOrAdd(MountPoint->GetContentPathAbsolute()) = MountPoint.GetReference();
+
+		bool bRemoveFromContainer;
+		MountPoint->SetMounted(true, ScopeLock, MountPoint, bRemoveFromContainer);
+		// Set mount order to highest priority, replacing earlier mount order if it already existed.
+		MountPoint->SetMountOrder(NextMountOrder++);
+		if (NextMountOrder == 0)
+		{
+			// Overflow, due to e.g. removing and adding a mountpoint too many times.
+			ReassignMountOrders();
+		}
+		check(!bRemoveFromContainer); // Should only be settable to true when calling SetMounted(false).
+	}
 
 	// Downgrade the log messages on startup to Verbose to reduce startup spam
 #if !NO_LOGGING
 	FString Message = FString::Printf(TEXT("FPackageName: Mount point added: '%s' mounted to '%s'"),
-		*MountPoint->GetContentPathRelative(), *MountPoint->GetRootPath());
+		*SearchKey.LocalPathStandard, *SearchKey.LongPackageName);
 	if (IsEngineStartupModuleLoadingComplete())
 	{
 		UE_LOG(LogPackageName, Log, TEXT("%s"), *Message);
@@ -251,99 +334,148 @@ void FLongPackagePathsSingleton::InsertMountPoint(const FString& RootPath, const
 #endif
 
 	// Let subscribers know that a new content path was mounted
-	FPackageName::OnContentPathMounted().Broadcast( RootPath, MountPoint->GetContentPathRelative());
+	FPackageName::OnContentPathMounted().Broadcast( SearchKey.LongPackageName, SearchKey.LocalPathStandard);
 }
 
-void FLongPackagePathsSingleton::RemoveMountPoint(const FString& RootPath, const FString& ContentPath)
+void FLongPackagePathsSingleton::ReassignMountOrders()
 {
 	using namespace UE::PackageName;
 
-	RemoveMountPoint(ConstructMountPoint(RootPath, ContentPath, EMountFlags::Removable));
+	// Called within writelock on MountLock
+	TArray<FMountPoint*> SortedMountPoints;
+	SortedMountPoints.Reserve(MountPoints.Num());
+	for (const FMountPointKeyByPtr& Ptr : MountPoints)
+	{
+		SortedMountPoints.Add(Ptr.GetMountPoint());
+	}
+
+	// Sort the list from lowest in priority to highest in priority, so we can traverse front to
+	// back assigning increasing-in-value and therefore increasing-in-priority NextMountOrder.
+	Algo::Sort(SortedMountPoints, [](const FMountPoint* A, const FMountPoint* B)
+		{
+			return A->GetMountOrder() < B->GetMountOrder();
+		});
+	NextMountOrder = 0;
+	for (FMountPoint* MountPoint : SortedMountPoints)
+	{
+		MountPoint->SetMountOrder(NextMountOrder++);
+	}
 }
 
-void FLongPackagePathsSingleton::RemoveMountPoint(const UE::PackageName::FMountPoint& MountPoint)
+void FLongPackagePathsSingleton::RemoveMountPoint(const FString& RootPath, const FString& ContentPath)
 {
 	using namespace UE::PackageName;
 
-	bool bRemoved = false;
+	FMountPointKeyByStrings SearchKey;
+	FString LocalPathAbsolute;
+	if (!TryNormalizePaths(RootPath, ContentPath, SearchKey.LongPackageName,
+		SearchKey.LocalPathStandard, LocalPathAbsolute, true /* bSilenceLogs */))
+	{
+		UE_LOG(LogPackageName, Display,
+			TEXT("FPackageName: Mount point remove failed: invalid input LongPackageName ('%s') or LocalPath ('%s')."),
+			*RootPath, *ContentPath);
+		return;
+	}
+	uint32 SearchHash = GetTypeHash(SearchKey);
+
+	TRefCountPtr<FMountPoint> ExistingMount;
+	bool bUnmounted = false;
 	{
 		UE::TWriteScopeLock ScopeLock(MountLock);
-		for (int32 MountPointIndex = 0; MountPointIndex < MountPoints.Num(); /* Conditionally incremented in loop */)
+		FMountPointKeyByPtr* ExistingPtr = MountPoints.FindByHash(SearchHash, SearchKey);
+		if (!ExistingPtr)
 		{
-			FMountPoint& ExistingMount = *MountPoints[MountPointIndex];
-			if (ExistingMount.GetRootPath() == MountPoint.GetRootPath() &&
-				ExistingMount.GetContentPathRelative() == MountPoint.GetContentPathRelative() &&
-				ExistingMount.IsRemovable())
+			UE_LOG(LogPackageName, Display,
+				TEXT("FPackageName: Mount point remove failed: no MountPoint found mapping '%s' to '%s'"),
+				*SearchKey.LocalPathStandard, *SearchKey.LongPackageName);
+			return;
+		}
+		ExistingMount = ExistingPtr->GetMountPoint();
+		if (!ExistingMount->IsRemovable())
+		{
+			UE_LOG(LogPackageName, Display, TEXT("FPackageName: Mount point remove failed: MountPoint mapping '%s' to '%s' is a default mapping and is not removable."),
+				*SearchKey.LocalPathStandard, *SearchKey.LongPackageName);
+			return;
+		}
+
+		bool bRootPathExisted;
+		bool bContentPathRelativeExisted;
+		bool bContentPathAbsoluteExisted;
+		RootPathTree.Remove(ExistingMount->GetRootPath(), &bRootPathExisted);
+		ContentPathTree.Remove(ExistingMount->GetContentPathRelative(), &bContentPathRelativeExisted);
+		ContentPathTree.Remove(ExistingMount->GetContentPathAbsolute(), &bContentPathAbsoluteExisted);
+		if (bRootPathExisted || bContentPathRelativeExisted || bContentPathAbsoluteExisted)
+		{
+			// See if there is an existing mounted MountPoint that was overridden by the removed mount point and now
+			// needs to be published as the MountPoint for one of its paths.
+			FMountPoint* NewRootPath = nullptr;
+			FMountPoint* NewContentPathRelative = nullptr;
+			FMountPoint* NewContentPathAbsolute = nullptr;
+			for (const FMountPointKeyByPtr& OtherMountPtr : MountPoints)
 			{
-				bool bRootPathExisted;
-				RootPathTree.Remove(ExistingMount.GetRootPath(), &bRootPathExisted);
-				if (bRootPathExisted)
+				FMountPoint* OtherMount = OtherMountPtr.GetMountPoint();
+				CA_ASSUME(OtherMount != nullptr);
+				if (OtherMount == ExistingMount.GetReference()
+					|| !OtherMount->IsMounted())
 				{
-					// See if there is a lower priority path that replaces it
-					for (TUniquePtr<FMountPoint>& OtherMount : TArrayView<TUniquePtr<FMountPoint>>(MountPoints).RightChop(MountPointIndex + 1))
+					continue;
+				}
+				if (bRootPathExisted && OtherMount->GetRootPath().Equals(
+					SearchKey.LongPackageName, ESearchCase::IgnoreCase))
+				{
+					if (NewRootPath == nullptr
+						|| OtherMount->GetMountOrder() > NewRootPath->GetMountOrder())
 					{
-						if (OtherMount->GetRootPath() == MountPoint.GetRootPath())
-						{
-							RootPathTree.FindOrAdd(OtherMount->GetRootPath()) = OtherMount.Get();
-							break;
-						}
+						NewRootPath = OtherMount;
 					}
 				}
-				bool bContentPathRelativeExisted;
-				ContentPathTree.Remove(ExistingMount.GetContentPathRelative(), &bContentPathRelativeExisted);
-				if (bContentPathRelativeExisted)
+				if (bContentPathRelativeExisted && OtherMount->GetContentPathRelative().Equals(
+					SearchKey.LocalPathStandard, ESearchCase::IgnoreCase))
 				{
-					// See if there is a lower priority path that replaces it
-					for (TUniquePtr<FMountPoint>& OtherMount :
-						TArrayView<TUniquePtr<FMountPoint>>(MountPoints).RightChop(MountPointIndex+1))
+					if (NewContentPathRelative == nullptr
+						|| OtherMount->GetMountOrder() > NewContentPathRelative->GetMountOrder())
 					{
-						if (OtherMount->GetContentPathRelative() == MountPoint.GetContentPathRelative())
-						{
-							ContentPathTree.FindOrAdd(OtherMount->GetContentPathRelative()) = OtherMount.Get();
-							break;
-						}
+						NewContentPathRelative = OtherMount;
 					}
 				}
-				bool bContentPathAbsoluteExisted;
-				ContentPathTree.Remove(ExistingMount.GetContentPathAbsolute(), &bContentPathAbsoluteExisted);
-				if (bContentPathAbsoluteExisted)
+				if (bContentPathAbsoluteExisted && OtherMount->GetContentPathAbsolute().Equals(
+					LocalPathAbsolute, ESearchCase::IgnoreCase))
 				{
-					// See if there is a lower priority path that replaces it
-					for (TUniquePtr<FMountPoint>& OtherMount :
-						TArrayView<TUniquePtr<FMountPoint>>(MountPoints).RightChop(MountPointIndex + 1))
+					if (NewContentPathAbsolute == nullptr
+						|| OtherMount->GetMountOrder() > NewContentPathAbsolute->GetMountOrder())
 					{
-						if (OtherMount->GetContentPathAbsolute() == MountPoint.GetContentPathAbsolute())
-						{
-							ContentPathTree.FindOrAdd(OtherMount->GetContentPathAbsolute()) = OtherMount.Get();
-							break;
-						}
+						NewContentPathAbsolute = OtherMount;
 					}
 				}
-
-				MountPoints.RemoveAt(MountPointIndex);
-				bRemoved = true;
 			}
-			else
+			if (NewRootPath)
 			{
-				++MountPointIndex;
+				RootPathTree.FindOrAdd(NewRootPath->GetRootPath()) = NewRootPath;
+			}
+			if (NewContentPathRelative)
+			{
+				ContentPathTree.FindOrAdd(NewContentPathRelative->GetContentPathRelative()) = NewContentPathRelative;
+			}
+			if (NewContentPathAbsolute)
+			{
+				ContentPathTree.FindOrAdd(NewContentPathAbsolute->GetContentPathAbsolute()) = NewContentPathAbsolute;
 			}
 		}
-		if (bRemoved)
-		{
-			UE_LOG(LogPackageName, Display, TEXT("FPackageName: Mount point removed: '%s' unmounted from '%s'"),
-				*MountPoint.GetContentPathRelative(), *MountPoint.GetRootPath());
-		}
-		else
+
+		bool bRemoveFromContainer;
+		ExistingMount->SetMounted(false, ScopeLock, ExistingMount, bRemoveFromContainer);
+		if (bRemoveFromContainer)
 		{
-			UE_LOG(LogPackageName, Display, TEXT("FPackageName: Mount point remove failed: no MountPoint found mapping '%s' to '%s'"),
-				*MountPoint.GetContentPathRelative(), *MountPoint.GetRootPath());
+			MountPoints.RemoveByHash(SearchHash, SearchKey);
 		}
 	}
 
-	// Let subscribers know that a new content path was unmounted
-	if (bRemoved)
+	if (ExistingMount)
 	{
-		FPackageName::OnContentPathDismounted().Broadcast(MountPoint.GetRootPath(), MountPoint.GetContentPathRelative());
+		UE_LOG(LogPackageName, Display, TEXT("FPackageName: Mount point removed: '%s' unmounted from '%s'"),
+			*SearchKey.LocalPathStandard, *SearchKey.LongPackageName);
+		// Let subscribers know that a content path was unmounted
+		FPackageName::OnContentPathDismounted().Broadcast(SearchKey.LongPackageName, SearchKey.LocalPathStandard);
 	}
 }
 
@@ -397,34 +529,67 @@ FLongPackagePathsSingleton::FLongPackagePathsSingleton()
 	GameConfigPathRebased  = RebasedGameDir / TEXT("Config/");
 	GameScriptPathRebased  = RebasedGameDir / TEXT("Script/");
 	GameSavedPathRebased   = RebasedGameDir / TEXT("Saved/");
+	UE::TWriteScopeLock ScopeLock(MountLock);
 
-	auto LocalInsertMountPoint = [this](const FString& RootPath, const FString& ContentPathRelative, EMountFlags InFlags)
+	auto LocalInsertMountPoint = [this, &ScopeLock](const FString& RootPath, const FString& ContentPathRelative,
+		EMountFlags InFlags)
 	{
-		// Don't call ConstructMountPoint, because that calls FPaths::ConvertRelativePathToFull followed by
+		// Don't call TryNormalizePaths, because that calls FPaths::ConvertRelativePathToFull followed by
 		// IFileManager::Get().ConvertToRelativePath, and that combination will convert ../../../Engine/Content
 		// into ../../Content, which we don't want here; we want to enter in the full ../../.. path for some entries.
 		TStringBuilder<256> AbsolutePathBuilder;
 		FPathViews::ToAbsolutePath(ContentPathRelative, AbsolutePathBuilder);
-		FMountPoint& MountPoint = *MountPoints.Add_GetRef(MakeUnique<FMountPoint>(
-			FString(RootPath), FString(ContentPathRelative), FString(AbsolutePathBuilder),
-			InFlags));
+		TRefCountPtr<FMountPoint> MountPoint = new FMountPoint(FString(RootPath), FString(ContentPathRelative),
+			FString(AbsolutePathBuilder), InFlags);
+		MountPoint->SetMountOrder(NextMountOrder++);
+		bool bAlreadyExisted;
+		MountPoints.Add(FMountPointKeyByPtr(*MountPoint), &bAlreadyExisted);
+		// We do not handle deduplication here; we expect that the calls to LocalInsertMountPoint is a manual list of
+		// non-duplicate entries.
+		check(!bAlreadyExisted);
+
+		bool bRemoveFromContainer;
+		MountPoint->SetMounted(true, ScopeLock, MountPoint, bRemoveFromContainer);
+		check(!bRemoveFromContainer); // Should only be settable to true when calling SetMounted(false).
 
 		// We are appending to the MountPoint list, and will have only a few duplicates. If any duplicates do occur,
-		// keep the MountPoint that was appended earlier as the registered mountpoint.
-		FMountPoint*& ExistingRootPath = RootPathTree.FindOrAdd(MountPoint.GetRootPath());
+		// keep the MountPoint that was appended earlier as the registered mountpoint, and mark it as higher priority
+		// than the current one.
+		FMountPoint*& ExistingRootPath = RootPathTree.FindOrAdd(MountPoint->GetRootPath());
 		if (!ExistingRootPath)
 		{
-			ExistingRootPath = &MountPoint;
+			ExistingRootPath = MountPoint.GetReference();
 		}
-		FMountPoint*& ExistingContentPathRelative = ContentPathTree.FindOrAdd(MountPoint.GetContentPathRelative());
+		else
+		{
+			if (ExistingRootPath->GetMountOrder() < MountPoint->GetMountOrder())
+			{
+				ExistingRootPath->SetMountOrder(NextMountOrder++);
+			}
+		}
+		FMountPoint*& ExistingContentPathRelative = ContentPathTree.FindOrAdd(MountPoint->GetContentPathRelative());
 		if (!ExistingContentPathRelative)
 		{
-			ExistingContentPathRelative = &MountPoint;
+			ExistingContentPathRelative = MountPoint.GetReference();
 		}
-		FMountPoint*& ExistingContentPathAbsolute = ContentPathTree.FindOrAdd(MountPoint.GetContentPathAbsolute());
+		else
+		{
+			if (ExistingContentPathRelative->GetMountOrder() < MountPoint->GetMountOrder())
+			{
+				ExistingContentPathRelative->SetMountOrder(NextMountOrder++);
+			}
+		}
+		FMountPoint*& ExistingContentPathAbsolute = ContentPathTree.FindOrAdd(MountPoint->GetContentPathAbsolute());
 		if (!ExistingContentPathAbsolute)
 		{
-			ExistingContentPathAbsolute = &MountPoint;
+			ExistingContentPathAbsolute = MountPoint.GetReference();
+		}
+		else
+		{
+			if (ExistingContentPathAbsolute->GetMountOrder() < MountPoint->GetMountOrder())
+			{
+				ExistingContentPathAbsolute->SetMountOrder(NextMountOrder++);
+			}
 		}
 	};
 		
@@ -448,10 +613,22 @@ FLongPackagePathsSingleton::FLongPackagePathsSingleton()
 	{
 		LocalInsertMountPoint(EngineRootPath, ContentPathShort, EMountFlags::NotReadOnly | EMountFlags::Alias);
 	}
-	LocalInsertMountPoint(GameRootPath,   GameContentPathRebased, EMountFlags::NotReadOnly | EMountFlags::Alias);
-	LocalInsertMountPoint(ScriptRootPath, GameScriptPathRebased, EMountFlags::ReadOnly | EMountFlags::Alias);
-	LocalInsertMountPoint(TempRootPath,   GameSavedPathRebased, EMountFlags::ReadOnly | EMountFlags::Alias);
-	LocalInsertMountPoint(ConfigRootPath, GameConfigPathRebased, EMountFlags::ReadOnly | EMountFlags::Alias);
+	if (GameContentPathRebased != GameContentPath)
+	{
+		LocalInsertMountPoint(GameRootPath, GameContentPathRebased, EMountFlags::NotReadOnly | EMountFlags::Alias);
+	}
+	if (GameScriptPathRebased != GameScriptPath)
+	{
+		LocalInsertMountPoint(ScriptRootPath, GameScriptPathRebased, EMountFlags::ReadOnly | EMountFlags::Alias);
+	}
+	if (GameSavedPathRebased != GameSavedPath)
+	{
+		LocalInsertMountPoint(TempRootPath, GameSavedPathRebased, EMountFlags::ReadOnly | EMountFlags::Alias);
+	}
+	if (GameConfigPathRebased != GameConfigPath)
+	{
+		LocalInsertMountPoint(ConfigRootPath, GameConfigPathRebased, EMountFlags::ReadOnly | EMountFlags::Alias);
+	}
 }
 
 #if !UE_BUILD_SHIPPING
@@ -461,14 +638,31 @@ void FLongPackagePathsSingleton::ExecDumpMountPoints(const TArray<FString>& Args
 
 	UE_LOG(LogPackageName, Log, TEXT("Valid mount points:"));
 
-	UE::TReadScopeLock ScopeLock(MountLock);
-	for (const TUniquePtr<FMountPoint>& MountPoint : MountPoints)
+	TArray<TRefCountPtr<FMountPoint>> ReportMountPoints;
+	{
+		UE::TReadScopeLock ScopeLock(MountLock);
+		ReportMountPoints.Reserve(MountPoints.Num());
+		for (const FMountPointKeyByPtr& MountPoint : MountPoints)
+		{
+			if (MountPoint.GetMountPoint()->IsMounted())
+			{
+				ReportMountPoints.Emplace(MountPoint.GetMountPoint());
+			}
+		}
+	}
+
+	// Sort the mount points from highest priority to lowest priority.
+	Algo::Sort(ReportMountPoints, [](const TRefCountPtr<FMountPoint>& A, const TRefCountPtr<FMountPoint>& B)
+		{
+			return A->GetMountOrder() > B->GetMountOrder();
+		});
+	for (TRefCountPtr<FMountPoint>& MountPoint : ReportMountPoints)
 	{
 		UE_LOG(LogPackageName, Log, TEXT("	'%s' -> '%s'"), *MountPoint->GetRootPath(), *MountPoint->GetContentPathRelative());
 	}
 
 	UE_LOG(LogPackageName, Log, TEXT("Removable mount points:"));
-	for (const TUniquePtr<FMountPoint>& MountPoint : MountPoints)
+	for (TRefCountPtr<FMountPoint>& MountPoint : ReportMountPoints)
 	{
 		if (MountPoint->IsRemovable())
 		{
@@ -540,31 +734,38 @@ void FLongPackagePathsSingleton::ExecRemoveMountPoint(const TArray<FString>& Arg
 		const FString& Path = Args[0];
 		const bool IsRootPath = Path[0] == TEXT('/');
 
-		FString RootPath;
-		FString ContentPath;
-			
+		FString LongPackageName;
+		FString LocalPathStandard;
+
 		if (IsRootPath)
 		{
+			if (!TryNormalizeLongPackageNamePath(Path, LongPackageName, false /* bSilenceLogs */))
+			{
+				return;
+			}
 			UE::TReadScopeLock ScopeLock(MountLock);
-				
-			RootPath = Path;
+
 			TArray<FString> ContentPathRelatives;
-			for (TUniquePtr<FMountPoint>& ExistingMount : Get().MountPoints)
+			for (const FMountPointKeyByPtr& ExistingKey : MountPoints)
 			{
-				if (ExistingMount->GetRootPath() == RootPath)
+				FMountPoint* ExistingMount = ExistingKey.GetMountPoint();
+				if (ExistingMount->IsMounted() && ExistingMount->GetRootPath() == LongPackageName)
 				{
 					ContentPathRelatives.Add(ExistingMount->GetContentPathRelative());
 				}
 			}
 			if (ContentPathRelatives.Num() == 0)
 			{
-				UE_LOG(LogPackageName, Error, TEXT("PackageName.UnregisterMountPoint: Root path '%s' is not mounted!"), *RootPath);
+				UE_LOG(LogPackageName, Error, TEXT("PackageName.UnregisterMountPoint: Root path '%s' is not mounted!"),
+					*LongPackageName);
 				return;
 			}
 
 			if (ContentPathRelatives.Num() > 1)
 			{
-				UE_LOG(LogPackageName, Error, TEXT("PackageName.UnregisterMountPoint: Root path '%s' is mounted to multiple content paths, specify content path to unmount explicitly!"), *RootPath);
+				UE_LOG(LogPackageName, Error,
+					TEXT("PackageName.UnregisterMountPoint: Root path '%s' is mounted to multiple content paths, specify content path to unmount explicitly!"),
+					*LongPackageName);
 				for (const FString& ContentPathRelative : ContentPathRelatives)
 				{
 					UE_LOG(LogPackageName, Error, TEXT("- %s"), *ContentPathRelative);
@@ -572,17 +773,23 @@ void FLongPackagePathsSingleton::ExecRemoveMountPoint(const TArray<FString>& Arg
 				return;
 			}
 
-			ContentPath = ContentPathRelatives[0];
+			LocalPathStandard = ContentPathRelatives[0];
 		}
 		else
 		{
+			FString LocalPathAbsolute;
+			if (!TryNormalizeLocalPath(Path, LocalPathStandard, LocalPathAbsolute, false /* bSilenceLogs */))
+			{
+				return;
+			}
+
 			UE::TReadScopeLock ScopeLock(MountLock);
-			
-			ContentPath = Path;
+
 			TArray<FString> RootPaths;
-			for (TUniquePtr<FMountPoint>& ExistingMount : Get().MountPoints)
+			for (const FMountPointKeyByPtr& ExistingKey : MountPoints)
 			{
-				if (ExistingMount->GetContentPathRelative() == ContentPath)
+				FMountPoint* ExistingMount = ExistingKey.GetMountPoint();
+				if (ExistingMount->IsMounted() && ExistingMount->GetContentPathRelative() == LocalPathStandard)
 				{
 					RootPaths.Add(ExistingMount->GetRootPath());
 				}
@@ -590,13 +797,13 @@ void FLongPackagePathsSingleton::ExecRemoveMountPoint(const TArray<FString>& Arg
 
 			if (RootPaths.Num() == 0)
 			{
-				UE_LOG(LogPackageName, Error, TEXT("PackageName.UnregisterMountPoint: Content path '%s' is not mounted!"), *ContentPath);
+				UE_LOG(LogPackageName, Error, TEXT("PackageName.UnregisterMountPoint: Content path '%s' is not mounted!"), *LocalPathStandard);
 				return;
 			}
 
 			if (RootPaths.Num() > 1)
 			{
-				UE_LOG(LogPackageName, Error, TEXT("PackageName.UnregisterMountPoint: Content path '%s' is mounted to multiple root paths, specify root path to unmount explicitly!"), *ContentPath);
+				UE_LOG(LogPackageName, Error, TEXT("PackageName.UnregisterMountPoint: Content path '%s' is mounted to multiple root paths, specify root path to unmount explicitly!"), *LocalPathStandard);
 				for (const FString& MountRootPath: RootPaths)
 				{
 					UE_LOG(LogPackageName, Error, TEXT("- %s"), *MountRootPath);
@@ -604,28 +811,14 @@ void FLongPackagePathsSingleton::ExecRemoveMountPoint(const TArray<FString>& Arg
 				return;
 			}
 
-			RootPath = RootPaths[0];
+			LongPackageName = RootPaths[0];
 		}
-			
-		Get().RemoveMountPoint(RootPath, ContentPath);
+
+		Get().RemoveMountPoint(LongPackageName, LocalPathStandard);
 	}
 	else
 	{
-		const FString& RootPath = Args[0];
-		const FString& ContentPath = Args[1];
-		if (ContentPath[0] == TEXT('/'))
-		{
-			UE_LOG(LogPackageName, Error, TEXT("PackageName.UnregisterMountPoint: Invalid ContentPath, should not start with '/'! Example: '../../../ProjectName/Content/'"));
-			return;
-		}
-
-		if (RootPath[0] != TEXT('/'))
-		{
-			UE_LOG(LogPackageName, Error, TEXT("PackageName.UnregisterMountPoint: Invalid RootPath, should start with '/'! Example: '/Game/'"));
-			return;
-		}
-
-		Get().RemoveMountPoint(RootPath, ContentPath);
+		Get().RemoveMountPoint(Args[0], Args[1]);
 	}
 }
 #endif
@@ -838,10 +1031,14 @@ FString FPackageName::FilenameToLongPackageName(const FString& InFilename)
 		{
 			const auto& Paths = FLongPackagePathsSingleton::Get();
 			UE::TReadScopeLock ScopeLock(Paths.MountLock);
-			for (const TUniquePtr<FMountPoint>& MountPoint : Paths.MountPoints)
+			for (const FMountPointKeyByPtr& MountPointKey : Paths.MountPoints)
 			{
-				ContentRootsArrayRelative.Add(MountPoint->GetContentPathRelative());
-				ContentRootsArrayAbsolute.Add(MountPoint->GetContentPathAbsolute());
+				FMountPoint* MountPoint = MountPointKey.GetMountPoint();
+				if (MountPoint->IsMounted())
+				{
+					ContentRootsArrayRelative.Add(MountPoint->GetContentPathRelative());
+					ContentRootsArrayAbsolute.Add(MountPoint->GetContentPathAbsolute());
+				}
 			}
 		}
 
@@ -1201,9 +1398,10 @@ bool FPackageName::IsValidLongPackageName(FStringView InLongPackageName, bool bI
 
 				UE::TReadScopeLock ScopeLock(Paths.MountLock);
 				TStringBuilder<2048> ValidRootsString;
-				for (const TUniquePtr<FMountPoint>& MountPoint : Paths.MountPoints)
+				for (const FMountPointKeyByPtr& MountPointKey : Paths.MountPoints)
 				{
-					if (bIncludeReadOnlyRoots || !MountPoint->IsReadOnly())
+					FMountPoint* MountPoint = MountPointKey.GetMountPoint();
+					if (MountPoint->IsMounted() && (bIncludeReadOnlyRoots || !MountPoint->IsReadOnly()))
 					{
 						ValidRootsString << MountPoint->GetRootPath() << TEXT(", ");
 					}
@@ -2673,17 +2871,35 @@ void FPackageName::QueryRootContentPaths(TArray<FString>& OutRootContentPaths, b
 	const FLongPackagePathsSingleton& Paths = FLongPackagePathsSingleton::Get();
 
 	{
-		UE::TReadScopeLock ScopeLock(Paths.MountLock);
-		OutRootContentPaths.Reserve(Paths.MountPoints.Num() + (bIncludeReadOnlyRoots ? 1 : 0));
+		TArray<TRefCountPtr<FMountPoint>> ReportMountPoints;
+		{
+			UE::TReadScopeLock ScopeLock(Paths.MountLock);
+			ReportMountPoints.Reserve(Paths.MountPoints.Num());
+			for (const FMountPointKeyByPtr& MountPointKey : Paths.MountPoints)
+			{
+				FMountPoint* MountPoint = MountPointKey.GetMountPoint();
+				if (MountPoint->IsMounted())
+				{
+					ReportMountPoints.Emplace(MountPoint);
+				}
+			}
+		}
+		// Sort ReportMountPoints from highest priority to lowest priority. 
+		Algo::Sort(ReportMountPoints, [](const TRefCountPtr<FMountPoint>& A, const TRefCountPtr<FMountPoint>& B)
+			{
+				return A->GetMountOrder() > B->GetMountOrder();
+			});
+
+		OutRootContentPaths.Reserve(ReportMountPoints.Num() + (bIncludeReadOnlyRoots ? 1 : 0));
 		// Maintain the legacy order: non-removable non-readonly, then removable, then non-removable readonly
-		for (const TUniquePtr<FMountPoint>& MountPoint : Paths.MountPoints)
+		for (TRefCountPtr<FMountPoint>& MountPoint : ReportMountPoints)
 		{
 			if (!MountPoint->IsRemovable() && !MountPoint->IsReadOnly() && !MountPoint->IsAlias())
 			{
 				OutRootContentPaths.Add(MountPoint->GetRootPath());
 			}
 		}
-		for (const TUniquePtr<FMountPoint>& MountPoint : Paths.MountPoints)
+		for (TRefCountPtr<FMountPoint>& MountPoint : ReportMountPoints)
 		{
 			if (MountPoint->IsRemovable())
 			{
@@ -2692,7 +2908,7 @@ void FPackageName::QueryRootContentPaths(TArray<FString>& OutRootContentPaths, b
 		}
 		if (bIncludeReadOnlyRoots)
 		{
-			for (const TUniquePtr<FMountPoint>& MountPoint : Paths.MountPoints)
+			for (TRefCountPtr<FMountPoint>& MountPoint : ReportMountPoints)
 			{
 				if (!MountPoint->IsRemovable() && MountPoint->IsReadOnly() && !MountPoint->IsAlias())
 				{
@@ -2728,17 +2944,33 @@ TArray<FString> FPackageName::QueryMountPointLocalAbsPaths()
 	const FLongPackagePathsSingleton& Paths = FLongPackagePathsSingleton::Get();
 	TArray<FString> OutAbsPaths;
 
+	TArray<TRefCountPtr<FMountPoint>> ReportMountPoints;
 	{
 		UE::TReadScopeLock ScopeLock(Paths.MountLock);
-		OutAbsPaths.Reserve(Paths.MountPoints.Num() + 1);
-		for (const TUniquePtr<FMountPoint>& MountPoint : Paths.MountPoints)
+		ReportMountPoints.Reserve(Paths.MountPoints.Num());
+		for (const FMountPointKeyByPtr& MountPointKey : Paths.MountPoints)
 		{
-			if (!MountPoint->IsAlias() && !MountPoint->GetContentPathAbsolute().IsEmpty())
+			FMountPoint* MountPoint = MountPointKey.GetMountPoint();
+			if (MountPoint->IsMounted())
 			{
-				OutAbsPaths.Add(MountPoint->GetContentPathAbsolute());
+				ReportMountPoints.Emplace(MountPoint);
 			}
 		}
 	}
+	// Sort ReportMountPoints from highest priority to lowest priority.
+	Algo::Sort(ReportMountPoints, [](const TRefCountPtr<FMountPoint>& A, const TRefCountPtr<FMountPoint>& B)
+		{
+			return A->GetMountOrder() > B->GetMountOrder();
+		});
+
+	OutAbsPaths.Reserve(ReportMountPoints.Num());
+	for (TRefCountPtr<FMountPoint>& MountPoint : ReportMountPoints)
+	{
+		if (!MountPoint->IsAlias() && !MountPoint->GetContentPathAbsolute().IsEmpty())
+		{
+			OutAbsPaths.Add(MountPoint->GetContentPathAbsolute());
+		}
+	}
 	return OutAbsPaths;
 }
 
diff --git a/Engine/Source/Runtime/CoreUObject/Private/Misc/PackageName/MountPoint.cpp b/Engine/Source/Runtime/CoreUObject/Private/Misc/PackageName/MountPoint.cpp
--- a/Engine/Source/Runtime/CoreUObject/Private/Misc/PackageName/MountPoint.cpp
+++ b/Engine/Source/Runtime/CoreUObject/Private/Misc/PackageName/MountPoint.cpp
@@ -1,19 +1,119 @@
 // Copyright Epic Games, Inc. All Rights Reserved.
 
+#include "Misc/MountPoint.h"
+
 #include "MountPointPrivate.h"
+#include "PackageNamePrivate.h"
 
 namespace UE::PackageName
 {
 
-FMountPoint::FMountPoint(FString&& InRootPath, FString&& InContentPathRelative,
-	FString&& InContentPathAbsolute, EMountFlags Flags)
-	: RootPath(MoveTemp(InRootPath))
-	, ContentPathRelative(MoveTemp(InContentPathRelative))
-	, ContentPathAbsolute(MoveTemp(InContentPathAbsolute))
+FMountPoint::FMountPoint(FString&& InLongPackageName, FString&& InLocalPathStandard,
+	FString&& InLocalPathAbsolute, EMountFlags Flags)
+	: LongPackageName(MoveTemp(InLongPackageName))
+	, LocalPathStandard(MoveTemp(InLocalPathStandard))
+	, LocalPathAbsolute(MoveTemp(InLocalPathAbsolute))
 	, bRemovable(!!(Flags & EMountFlags::Removable))
 	, bReadOnly(!!(Flags & EMountFlags::ReadOnly))
 	, bAlias(!!(Flags & EMountFlags::Alias))
 {
 }
 
-}
\ No newline at end of file
+bool FMountPoint::IsMounted() const
+{
+	return bMounted.load(std::memory_order_relaxed);
+}
+
+void FMountPoint::SetMounted(bool bInMounted, UE::TWriteScopeLock<FTransactionallySafeRWLock>& /* ProofOfWriteLock */,
+	const TRefCountPtr<FMountPoint>& /* ProofOfCallerRefCount */, bool& bOutRemoveFromContainer)
+{
+	bOutRemoveFromContainer = false;
+
+	// bMounted is only modified within a write lock, so we do not need compare_exchange
+	bool bWasMounted = bMounted.load(std::memory_order_relaxed);
+
+	// Our caller holds a refcount, to prevent the possiblility of refcount going to zero when we update mounted
+	// status; refcount should only go to zero once we leave the WriteLock. If we are already mounted, we should
+	// additionally hold a refcount for being mounted.
+	check(RefCount.load(std::memory_order_relaxed) >= (bWasMounted ? 2U : 1U));
+	if (bWasMounted == bInMounted)
+	{
+		return;
+	}
+
+	bMounted.store(bInMounted, std::memory_order_relaxed);
+
+	// While mounted, we hold an internal refcount on the MountPoint.
+	if (bInMounted)
+	{
+		AddRef();
+	}
+	else
+	{
+		Release();
+
+		// If there is no further external refcount, as an optimization, allow our caller to remove
+		// us from the FLongPackagePathsSingleton.MountPoints container, possibly more efficiently
+		// than we will be able to do it in the upcoming TryUnregisterAndDelete.
+		bOutRemoveFromContainer = RefCount.load(std::memory_order_relaxed) == 1;
+		bAlreadyRemovedFromContainer = bOutRemoveFromContainer;
+	}
+}
+
+void FMountPoint::AddRef() const
+{
+	RefCount.fetch_add(1, std::memory_order_relaxed);
+}
+
+FReturnedRefCountValue FMountPoint::Release() const
+{
+	uint32 OldValue = RefCount.fetch_add(-1, std::memory_order_relaxed);
+	check(OldValue > 0);
+	if (OldValue == 1)
+	{
+		const_cast<FMountPoint*>(this)->TryUnregisterAndDelete();
+	}
+	return FReturnedRefCountValue{ OldValue - 1 };
+}
+
+void FMountPoint::TryUnregisterAndDelete()
+{
+	// Another thread may be in the process of finding the pointer to this from
+	// FLongPackagePathsSingleton.RootPathTree, and it might AddRef after the decrement in our caller function.
+	// It does that Find and AddRef inside a WriteLock on FLongPackagePathsSingleton.MountLock.
+	// Enter the WriteLock ourselves and check that the RefCount is still 0 within the writelock
+	// before unregister and delete.
+	FLongPackagePathsSingleton& Paths = FLongPackagePathsSingleton::Get();
+	UE::TWriteScopeLock<FTransactionallySafeRWLock> ScopeLock(Paths.MountLock);
+	if (RefCount.load(std::memory_order_relaxed) != 0)
+	{
+		return;
+	}
+
+	// We may have set bAlreadyRemovedFromContainer when unmounting, in which case we can skip the effort of removing
+	// from MountPoints
+	if (!bAlreadyRemovedFromContainer)
+	{
+		// It should be impossible for an FMountPoint to be in one of the mounted containers when refcount is 0, because
+		// we AddRef it when mounted and do not Release until unmounted. Assert that it is not in any of the moounted containers.
+#if DO_CHECK
+		FMountPoint** Existing = Paths.RootPathTree.Find(LongPackageName);
+		check(!Existing || *Existing != this);
+		Existing = Paths.ContentPathTree.Find(LocalPathStandard);
+		check(!Existing || *Existing != this);
+		Existing = Paths.ContentPathTree.Find(LocalPathAbsolute);
+		check(!Existing || *Existing != this);
+#endif
+
+		Paths.MountPoints.Remove(FMountPointKeyByPtr(*this));
+	}
+
+	delete this;
+}
+
+FReturnedRefCountValue FMountPoint::GetRefCount() const
+{
+	return FReturnedRefCountValue{ RefCount.load(std::memory_order_relaxed) };
+}
+
+} // namespace UE::PackageName
\ No newline at end of file
diff --git a/Engine/Source/Runtime/CoreUObject/Private/Misc/PackageName/MountPointPrivate.h b/Engine/Source/Runtime/CoreUObject/Private/Misc/PackageName/MountPointPrivate.h
--- a/Engine/Source/Runtime/CoreUObject/Private/Misc/PackageName/MountPointPrivate.h
+++ b/Engine/Source/Runtime/CoreUObject/Private/Misc/PackageName/MountPointPrivate.h
@@ -2,8 +2,14 @@
 
 #pragma once
 
+#include "Containers/StringView.h"
 #include "Containers/UnrealString.h"
 #include "Misc/EnumClassFlags.h"
+#include "Misc/MountPoint.h"
+#include "Misc/ScopeRWLock.h"
+#include "Misc/TransactionallySafeRWLock.h"
+
+#include <atomic>
 
 namespace UE::PackageName
 {
@@ -19,33 +25,140 @@ enum class EMountFlags
 };
 ENUM_CLASS_FLAGS(EMountFlags);
 
-class FMountPoint
+// Keytypes for use in TSets that look up mountpoints by their paths. We support both a KeyType struct with a pointer,
+// to save memory, and a KeyType struct with strings, for when we have the paths but not an FMountPoint.
+struct FMountPointKeyByPtr;
+struct FMountPointKeyByStrings;
+extern uint32 GetTypeHash(const FMountPointKeyByPtr& Key);
+extern uint32 GetTypeHash(const FMountPointKeyByStrings& Key);
+
+
+/**
+ * The internal implementation of IMountPoint. FMountPoints are present in FLongPackagePathsSingleton.MountPoints
+ * while allocated; they are removed from that set, inside a critical section, when their refcount goes to zero.
+ */
+class FMountPoint final : public IMountPoint
 {
 public:
-	FMountPoint(FString&& InRootPath, FString&& InContentPathRelative,
-		FString&& InContentPathAbsolute, EMountFlags Flags);
+	FMountPoint(FString&& InLongPackageName, FString&& InLocalPathStandard,
+		FString&& InLocalPathAbsolute, EMountFlags Flags);
+
+	// IMountPoint API
+	virtual FStringView GetLongPackageName() const override;
+	virtual FStringView GetLocalPathStandard() const override;
+	virtual FStringView GetLocalPathAbsolute() const override;
+
+	virtual bool IsRemovable() const override;
+	virtual bool IsReadOnly() const override;
+	virtual bool IsMounted() const override;
+
+	virtual void AddRef() const override;
+	virtual FReturnedRefCountValue Release() const override;
+	virtual FReturnedRefCountValue GetRefCount() const override;
 
-	/** The LongPackageName path(e.g., "/Engine/") */
+	// FMountPoint additional API
+
+	/** Same value as GetLongPackageName */
 	const FString& GetRootPath() const;
-	/** The LocalPath relative path(e.g., "../../../Engine/Content/") */
+	/** Same value as GetLocalPathStandard */
 	const FString& GetContentPathRelative() const;
-	/** The LocalPath absolute path */
+	/** Same value as GetLocalPathAbsolute */
 	const FString& GetContentPathAbsolute() const;
 
-	/** Whether the mountpoint was created by InsertMountPoint and can be removed by RemoveMountPoint. */
-	bool IsRemovable() const;
-	/** Whether new Packages can be saved into the mountpoint. */
-	bool IsReadOnly() const;
 	/** Whether this mountpoint is the same RootPath as another mountpoint but bound to a different ContentPath. */
 	bool IsAlias() const;
 
+	/**
+	 * Mark mounted or unmounted. Must only be called within Paths.MountLock, so that the
+	 * change in value is atomic with external queries about whether a path is mounted.
+	 * When unmounting, caller must hold a refcount to prevent the possibility of deletion.
+	 *
+	 * Caller should hold exactly 1 refcount when unmounting. In that case bOutRemoveFromContainer will
+	 * be set to true and we will skip the effort in TryUnregisterAndDelete of removing it from the container.
+	 */
+	void SetMounted(bool bInMounted, UE::TWriteScopeLock<FTransactionallySafeRWLock>& ProofOfWriteLock,
+		const TRefCountPtr<FMountPoint>& ProofOfCallerRefCount, bool& bOutRemoveFromContainer);
+
+	/**
+	 * Get/Set the Ephemeral ordering value. MountPoints that have the same RootPath or ContentPath as other
+	 * mountpoints are selectively returned from queries where they overlap based on higher priority == higher value.
+	 * Setting higher values higher in priority allows us to easily implement our contract that MountPoints added
+	 * later are higher in priority than MountPoints added earlier.
+	 * The exact value of this field can change arbitrarily in MountPoints, when e.g. we reorder the list. Only
+	 * the relative ordering is guaranteed unchanged for an existing MountPoint.
+	 */
+	uint32 GetMountOrder() const;
+	void SetMountOrder(uint32 MountOrder);
+
+	bool Equals(const FMountPointKeyByStrings& Key) const;
+	static uint32 GetTypeHashKey(const FString& LongPackageNameOfKey, const FString& LocalPathStandardOfKey);
+	static bool KeysAreEqual(const FString& LongPackageNameA, const FString& LocalPathStandardA,
+		const FString& LongPackageNameB, const FString& LocalPathStandardB);
+
 private:
-	const FString RootPath;
-	const FString ContentPathRelative;
-	const FString ContentPathAbsolute;
+	/**
+	 * Called when refcount goes to zero. Enters the critical section and if still zero in critical section,
+	 * removes this from FMountPoint's registration maps and deletes this.
+	 */
+	void TryUnregisterAndDelete();
+
+	/** Const field set in Constructor */
+	const FString LongPackageName;
+	/** Const field set in Constructor */
+	const FString LocalPathStandard;
+	/** Const field set in Constructor */
+	const FString LocalPathAbsolute;
+	/** Const field set in Constructor */
 	const bool bRemovable : 1;
+	/** Const field set in Constructor */
 	const bool bReadOnly : 1;
+	/** Const field set in Constructor */
 	const bool bAlias : 1;
+	/** Read/Write only under FLongPackagePathsSingleton.MountLock */
+	bool bAlreadyRemovedFromContainer : 1 = false;
+	/** Read/Write only under FLongPackagePathsSingleton.MountLock */
+	uint32 MountOrder = 0;
+	/** Atomic field read/write from all threads */
+	mutable std::atomic<uint32> RefCount{ 0 };
+	/** Atomic field readable from all threads, writable only under FLongPackagePathsSingleton.MountLock */
+	std::atomic<bool> bMounted{ false };
+};
+
+/**
+ * A wrapper around FMountPoint* that provides the KeyFuncs for use in a TSet.
+ * This type must only be used in the reference maps for FMountPoints, and removed before FMountPoint is destructed,
+ * so that we can keep a raw pointer and not have to handle the complexity of accounting for this structure's
+ * TRefCountPtr when deciding whether an FMountPoint is no longer referenced.
+ */
+struct FMountPointKeyByPtr
+{
+public:
+	explicit FMountPointKeyByPtr(FMountPoint& InMountPoint);
+	bool operator==(const FMountPointKeyByPtr& Other) const;
+	bool operator==(const FMountPointKeyByStrings& Other) const;
+	FMountPoint* GetMountPoint() const;
+
+private:
+	FMountPoint* MountPoint = nullptr;
+	friend uint32 UE::PackageName::GetTypeHash(const FMountPointKeyByPtr& Key);
+	friend UE::PackageName::FMountPointKeyByStrings;
+};
+
+/**
+ * A structure that holds the same strings held by an FMountPoint that is comparable with FMountPointKeyByPtr,
+ * to support queries against a TSet<FMountPointKeyByPtr>.
+ */
+struct FMountPointKeyByStrings
+{
+	FMountPointKeyByStrings() = default;
+	FMountPointKeyByStrings(const FMountPoint& InMountPoint);
+	FMountPointKeyByStrings(FString InLongPackageName, FString InLocalPathStandard);
+
+	bool operator==(const FMountPointKeyByPtr& Other) const;
+	bool operator==(const FMountPointKeyByStrings& Other) const;
+
+	FString LongPackageName;
+	FString LocalPathStandard;
 };
 
 
@@ -54,19 +167,19 @@ private:
 ///////////////////////////////////////////////////////
 
 
-inline const FString& FMountPoint::GetRootPath() const
+inline FStringView FMountPoint::GetLongPackageName() const
 {
-	return RootPath;
+	return LongPackageName;
 }
 
-inline const FString& FMountPoint::GetContentPathRelative() const
+inline FStringView FMountPoint::GetLocalPathStandard() const
 {
-	return ContentPathRelative;
+	return LocalPathStandard;
 }
 
-inline const FString& FMountPoint::GetContentPathAbsolute() const
+inline FStringView FMountPoint::GetLocalPathAbsolute() const
 {
-	return ContentPathAbsolute;
+	return LocalPathAbsolute;
 }
 
 inline bool FMountPoint::IsRemovable() const
@@ -79,9 +192,107 @@ inline bool FMountPoint::IsReadOnly() const
 	return bReadOnly;
 }
 
+inline const FString& FMountPoint::GetRootPath() const
+{
+	return LongPackageName;
+}
+
+inline const FString& FMountPoint::GetContentPathRelative() const
+{
+	return LocalPathStandard;
+}
+
+inline const FString& FMountPoint::GetContentPathAbsolute() const
+{
+	return LocalPathAbsolute;
+}
+
 inline bool FMountPoint::IsAlias() const
 {
 	return bAlias;
 }
 
+inline uint32 FMountPoint::GetMountOrder() const
+{
+	return MountOrder;
+}
+
+inline void FMountPoint::SetMountOrder(uint32 InMountOrder)
+{
+	MountOrder = InMountOrder;
+}
+
+inline FMountPointKeyByPtr::FMountPointKeyByPtr(FMountPoint& InMountPoint)
+	: MountPoint(&InMountPoint)
+{
+}
+
+inline bool FMountPointKeyByPtr::operator==(const FMountPointKeyByPtr& Other) const
+{
+	return FMountPoint::KeysAreEqual(MountPoint->GetRootPath(), MountPoint->GetContentPathRelative(),
+		Other.MountPoint->GetRootPath(), Other.MountPoint->GetContentPathRelative());
+}
+
+inline bool FMountPointKeyByPtr::operator==(const FMountPointKeyByStrings& Other) const
+{
+	return FMountPoint::KeysAreEqual(MountPoint->GetRootPath(), MountPoint->GetContentPathRelative(),
+		Other.LongPackageName, Other.LocalPathStandard);
+}
+
+inline FMountPoint* FMountPointKeyByPtr::GetMountPoint() const
+{
+	return MountPoint;
+}
+
+inline uint32 GetTypeHash(const FMountPointKeyByPtr& Key)
+{
+	return FMountPoint::GetTypeHashKey(Key.MountPoint->GetRootPath(), Key.MountPoint->GetContentPathRelative());
+}
+
+inline FMountPointKeyByStrings::FMountPointKeyByStrings(const FMountPoint& InMountPoint)
+	: LongPackageName(InMountPoint.GetLongPackageName())
+	, LocalPathStandard(InMountPoint.GetLocalPathStandard())
+{
+}
+
+inline FMountPointKeyByStrings::FMountPointKeyByStrings(FString InLongPackageName, FString InLocalPathStandard)
+	: LongPackageName(MoveTemp(InLongPackageName))
+	, LocalPathStandard(MoveTemp(InLocalPathStandard))
+{
+}
+
+inline bool FMountPointKeyByStrings::operator==(const FMountPointKeyByPtr& Other) const
+{
+	return FMountPoint::KeysAreEqual(LongPackageName, LocalPathStandard,
+		Other.MountPoint->GetRootPath(), Other.MountPoint->GetContentPathRelative());
+}
+
+inline bool FMountPointKeyByStrings::operator==(const FMountPointKeyByStrings& Other) const
+{
+	return FMountPoint::KeysAreEqual(LongPackageName, LocalPathStandard,
+		Other.LongPackageName, Other.LocalPathStandard);
+}
+
+inline uint32 GetTypeHash(const FMountPointKeyByStrings& Key)
+{
+	return FMountPoint::GetTypeHashKey(Key.LongPackageName, Key.LocalPathStandard);
+}
+
+inline bool FMountPoint::Equals(const FMountPointKeyByStrings& Key) const
+{
+	return KeysAreEqual(LongPackageName, LocalPathStandard, Key.LongPackageName, Key.LocalPathStandard);
+}
+
+inline uint32 FMountPoint::GetTypeHashKey(const FString& LongPackageNameOfKey, const FString& LocalPathStandardOfKey)
+{
+	return 101 * GetTypeHash(LongPackageNameOfKey) + GetTypeHash(LocalPathStandardOfKey);
+}
+
+inline bool FMountPoint::KeysAreEqual(const FString& LongPackageNameA, const FString& LocalPathStandardA,
+	const FString& LongPackageNameB, const FString& LocalPathStandardB)
+{
+	return LongPackageNameA.Equals(LongPackageNameB, ESearchCase::IgnoreCase)
+		&& LocalPathStandardA.Equals(LocalPathStandardB, ESearchCase::IgnoreCase);
+}
+
 } // namespace UE::PackageName
\ No newline at end of file
diff --git a/Engine/Source/Runtime/CoreUObject/Private/Misc/PackageName/PackageNamePrivate.h b/Engine/Source/Runtime/CoreUObject/Private/Misc/PackageName/PackageNamePrivate.h
--- a/Engine/Source/Runtime/CoreUObject/Private/Misc/PackageName/PackageNamePrivate.h
+++ b/Engine/Source/Runtime/CoreUObject/Private/Misc/PackageName/PackageNamePrivate.h
@@ -40,7 +40,17 @@ struct FLongPackagePathsSingleton
 
 	TDirectoryTree<UE::PackageName::FMountPoint*> RootPathTree;
 	TDirectoryTree<UE::PackageName::FMountPoint*> ContentPathTree;
-	TArray<TUniquePtr<UE::PackageName::FMountPoint>> MountPoints;
+	/**
+	 * Set of all MountPoints, possibly unmounted, looked up by SearchKey={LongPackageName,LocalPath}. 
+	 * Elements added later are assigned a higher priority and will be the ones findable in RootPathTree
+	 * and ContentPathTree if they have the same LongPackageName or LocalPath as another MountPoint.
+	 * This override system allows aliasing both of two LocalPaths that go to the same LongPackageName and two
+	 * LongPackageNames that go to the same LocalPath.
+	 */
+	TSet<UE::PackageName::FMountPointKeyByPtr> MountPoints;
+
+	/** Assigned and incremented to each new mountpoint, to establish override priority. */
+	uint32 NextMountOrder = 0;
 
 	// singleton
 	static FLongPackagePathsSingleton& Get();
@@ -49,17 +59,22 @@ struct FLongPackagePathsSingleton
 	void OnCoreUObjectInitialized();
 
 	/**
-	 * Given a content path return the consistently-formatted RelativePath and AbsolutePath, suitable for lookup based only on text
-	 * The RelativePath is consistent with FileManager relative paths.
+	 * Given LongPackageName and LocalPath, return the consistently-formatted LongPackageName, LocalPathStandard, and
+	 * LocalPathAbsolute, suitable for lookup in our maps of string to MountPoint.
+	 * LocalPathStandard is a relative path and is consistent with FileManager relative paths.
+	 * @return false if the input strings are not recognizable as mountpoint strings, else true.
 	 */
-	static UE::PackageName::FMountPoint ConstructMountPoint(const FString& RootPath, const FString& ContentPath,
-		UE::PackageName::EMountFlags Flags);
+	bool TryNormalizePaths(const FString& InLongPackageName, const FString& InLocalPath,
+		FString& OutLongPackageName, FString& OutLocalPathStandard, FString& OutLocalPathAbsolute, bool bSilenceLogs);
+	bool TryNormalizeLongPackageNamePath(const FString& InLongPackageName, FString& OutLongPackageName,
+		bool bSilenceLogs);
+	bool TryNormalizeLocalPath(const FString& InLocalPath, FString& OutLocalPathStandard,
+		FString& OutLocalPathAbsolute, bool bSilenceLogs);
 
 	void InsertMountPoint(const FString& RootPath, const FString& ContentPath);
 
 	/** This will remove a previously inserted mount point. */
 	void RemoveMountPoint(const FString& RootPath, const FString& ContentPath);
-	void RemoveMountPoint(const UE::PackageName::FMountPoint& MountPoint);
 
 	/** Checks whether the specific root path is a valid mount point. */
 	bool MountPointExists(FStringView RootPath);
@@ -72,6 +87,7 @@ private:
 #endif
 	
 	FLongPackagePathsSingleton();
+	void ReassignMountOrders();
 
 #if !UE_BUILD_SHIPPING
 	void ExecDumpMountPoints(const TArray<FString>& Args);
diff --git a/Engine/Source/Runtime/CoreUObject/Public/Misc/MountPoint.h b/Engine/Source/Runtime/CoreUObject/Public/Misc/MountPoint.h
new file mode 100644
--- /dev/null
+++ b/Engine/Source/Runtime/CoreUObject/Public/Misc/MountPoint.h
@@ -0,0 +1,65 @@
+// Copyright Epic Games, Inc. All Rights Reserved.
+
+#pragma once
+
+#include "Containers/StringView.h"
+#include "Templates/RefCounting.h"
+
+namespace UE::PackageName
+{
+
+/**
+ * A reference-counted mountpoint that has been registered with FPackageName. It might not be mounted; mountpoints can
+ * be preregistered before mounting, or can be kept allocated after unmounting.
+ * If mounted, a MountPoint provides a location on disk from which packages can be loaded (and saved in most cases),
+ * and provides a two-way mapping from LongPackageName <-> LocalPath.
+ * 
+ * Multiple MountPoints can be mapped from a single LongPackageName to multiple LocalPaths, or from a single LocalPath
+ * to multiple LongPackageNames. Only the last MountPoint mounted for a LongPackageName or LocalPath is returned from
+ * queries for those paths, but all mounted MountPoints can still be used to load packages, they report
+ * IsMounted == true, and they are reported in query functions that report the list of mounted MountPoints.
+ */
+class IMountPoint : public IRefCountedObject
+{
+public:
+	/**
+	 * PackageName path that is the parent directory of packagenames in this mountpoint. Has a trailing slash.
+	 * e.g. "/Game/", "/PluginName/"
+	 * Immutable, can be read from any thread.
+	 */
+	virtual FStringView GetLongPackageName() const = 0;
+	/**
+	 * LocalPath, in Unreal standard path form (normalized and relative to Engine/Binaries/<Platform> directory), that
+	 * is the parent directory of filenames in this mountpoint. Has a trailing slash.
+	 * e.g. "../../../EngineTest/Content/", "../../../EngineTest/Plugins/PluginName/Content/"
+	 * Immutable, can be read from any thread.
+	 */
+	virtual FStringView GetLocalPathStandard() const = 0;
+	/**
+	 * LocalPath, absolute, that is the parent directory of filenames in this mountpoint. Has a trailing slash.
+	 * e.g. "d:/UnrealRoot/EngineTest/Content/", "d:/UnrealRoot/EngineTest/Plugins/PluginName/Content/"
+	 * Immutable, can be read from any thread.
+	 */
+	virtual FStringView GetLocalPathAbsolute() const = 0;
+
+	/**
+	 * Whether the mountpoint was created by InsertMountPoint and can be removed by RemoveMountPoint.
+	 * Immutable, can be read from any thread.
+	 */
+	virtual bool IsRemovable() const = 0;
+	/**
+	 * Whether new Packages are unable to be saved into the mountpoint.
+	 * e.g. returns false for /Game/, returns true for /Script/.
+	 * Immutable, can be read from any thread.
+	 */
+	virtual bool IsReadOnly() const = 0;
+	/**
+	 * Whether the MountPoint is currently mounted. IMountPoints can be registered with PackageName before they are mounted,
+	 * and IMountPoints will remain allocated, and be used by PackageName rather than replaced with a new allocation when they
+	 * are mounted, as long as they have a non-zero refcount.
+	 * Mutable but atomic, can be read from any thread.
+	 */
+	virtual bool IsMounted() const = 0;
+};
+
+} // namespace UE::PackageName
\ No newline at end of file
diff --git a/Engine/Source/Runtime/CoreUObject/Public/Misc/PackageName.h b/Engine/Source/Runtime/CoreUObject/Public/Misc/PackageName.h
--- a/Engine/Source/Runtime/CoreUObject/Public/Misc/PackageName.h
+++ b/Engine/Source/Runtime/CoreUObject/Public/Misc/PackageName.h
@@ -17,6 +17,7 @@
 #include "HAL/PlatformMath.h"
 #include "Internationalization/Text.h"
 #include "Logging/LogMacros.h"
+#include "Misc/MountPoint.h"
 #include "Misc/PackagePath.h"
 #include "Templates/Function.h"
 #include "UObject/NameTypes.h"
