PackageName: When constructing initial default MountPoints, which depend upon project name and location, handle the possibility of duplicates in our hardcoded list and issue a warning rather than an assertion.
#rnx

#Synced-CL 50273579

diff --git a/Engine/Source/Runtime/CoreUObject/Private/Misc/PackageName.cpp b/Engine/Source/Runtime/CoreUObject/Private/Misc/PackageName.cpp
--- a/Engine/Source/Runtime/CoreUObject/Private/Misc/PackageName.cpp
+++ b/Engine/Source/Runtime/CoreUObject/Private/Misc/PackageName.cpp
@@ -539,14 +539,21 @@ FLongPackagePathsSingleton::FLongPackagePathsSingleton()
 		// into ../../Content, which we don't want here; we want to enter in the full ../../.. path for some entries.
 		TStringBuilder<256> AbsolutePathBuilder;
 		FPathViews::ToAbsolutePath(ContentPathRelative, AbsolutePathBuilder);
+		FMountPointKeyByStrings SearchKey(RootPath, ContentPathRelative);
+		FMountPointKeyByPtr* ExistingFullKey = MountPoints.FindByHash(GetTypeHash(SearchKey), SearchKey);
+		if (ExistingFullKey)
+		{
+			UE_LOG(LogPackageName, Warning, TEXT("Unexpected duplicate initial MountPoint added during PackageName startup: %s, %s."),
+				*RootPath, *ContentPathRelative);
+			return;
+		}
+
 		TRefCountPtr<FMountPoint> MountPoint = new FMountPoint(FString(RootPath), FString(ContentPathRelative),
 			FString(AbsolutePathBuilder), InFlags);
 		MountPoint->SetMountOrder(NextMountOrder++);
 		bool bAlreadyExisted;
 		MountPoints.Add(FMountPointKeyByPtr(*MountPoint), &bAlreadyExisted);
-		// We do not handle deduplication here; we expect that the calls to LocalInsertMountPoint is a manual list of
-		// non-duplicate entries.
-		check(!bAlreadyExisted);
+		check(!bAlreadyExisted); // Impossible for it to already exist, because we called Find above.
 
 		bool bRemoveFromContainer;
 		MountPoint->SetMounted(true, ScopeLock, MountPoint, bRemoveFromContainer);
