PackageName: Fix leak at shutdown with new reference counted FMountPoints. We have to drop the reference from being mounted, since most mountpoints are not unmounted before shutdown.
#rnx

#Synced-CL 50273579

diff --git a/Engine/Source/Runtime/CoreUObject/Private/Misc/PackageName.cpp b/Engine/Source/Runtime/CoreUObject/Private/Misc/PackageName.cpp
--- a/Engine/Source/Runtime/CoreUObject/Private/Misc/PackageName.cpp
+++ b/Engine/Source/Runtime/CoreUObject/Private/Misc/PackageName.cpp
@@ -638,6 +638,44 @@ FLongPackagePathsSingleton::FLongPackagePathsSingleton()
 	}
 }
 
+FLongPackagePathsSingleton::~FLongPackagePathsSingleton()
+{
+	using namespace UE::PackageName;
+
+	// We require that our external callers drop their IMountPoint refcounts, but we do not require that they unmount
+	// all of the MountPoints they registered with RegisterMountPoint. So we may have MountPoints remaining with that
+	// our still mounted and therefore we still hold a refcount. Unmount them all now to drop that final refcount and
+	// free the memory. This happens at end of process, so this is important mainly for leak detection.
+	TArray<TRefCountPtr<FMountPoint>> RemainingMountPoints;
+	{
+		UE::TWriteScopeLock ScopeLock(MountLock);
+		RemainingMountPoints.Reserve(MountPoints.Num());
+		for (FMountPointKeyByPtr& Key : MountPoints)
+		{
+			TRefCountPtr<FMountPoint>& MountPoint = RemainingMountPoints.Emplace_GetRef(Key.GetMountPoint());
+			bool bRemoveFromContainer;
+			MountPoint->SetMounted(false, ScopeLock, MountPoint, bRemoveFromContainer);
+			if (!bRemoveFromContainer)
+			{
+				// UE_LOG is likely no longer available; we will have to rely on leak detector reporting this instead.
+				constexpr bool bAllowedToCallLog = false;
+				if (bAllowedToCallLog)
+				{
+					UE_LOG(LogPackageName, Error,
+						TEXT("MountPoint is still referenced during FLongPackagePathsSingleton destruction; this will leak memory or will cause a crash if the MountPoint is released later. MountPoint: %s, %s"),
+						*MountPoint->GetRootPath(), *MountPoint->GetContentPathRelative());
+				}
+			}
+		}
+		RootPathTree.Empty();
+		ContentPathTree.Empty();
+		MountPoints.Empty();
+	}
+
+	// Emptying the list will drop the final refcount for all of the remaining MountPoints and will cause their destruction.
+	RemainingMountPoints.Empty();
+}
+
 #if !UE_BUILD_SHIPPING
 void FLongPackagePathsSingleton::ExecDumpMountPoints(const TArray<FString>& Args)
 {
diff --git a/Engine/Source/Runtime/CoreUObject/Private/Misc/PackageName/PackageNamePrivate.h b/Engine/Source/Runtime/CoreUObject/Private/Misc/PackageName/PackageNamePrivate.h
--- a/Engine/Source/Runtime/CoreUObject/Private/Misc/PackageName/PackageNamePrivate.h
+++ b/Engine/Source/Runtime/CoreUObject/Private/Misc/PackageName/PackageNamePrivate.h
@@ -87,6 +87,7 @@ private:
 #endif
 	
 	FLongPackagePathsSingleton();
+	~FLongPackagePathsSingleton();
 	void ReassignMountOrders();
 
 #if !UE_BUILD_SHIPPING
