PackageName: When constructing initial default MountPoints,
handle any duplicate MountPoints by skipping the duplicate rather than trying to manually check for duplicates. Remove the warning that we had in the code to handle duplicates since it is now expected to occur for some projects.
#rnx

#Synced-CL 50273579

diff --git a/Engine/Source/Runtime/CoreUObject/Private/Misc/PackageName.cpp b/Engine/Source/Runtime/CoreUObject/Private/Misc/PackageName.cpp
--- a/Engine/Source/Runtime/CoreUObject/Private/Misc/PackageName.cpp
+++ b/Engine/Source/Runtime/CoreUObject/Private/Misc/PackageName.cpp
@@ -543,8 +543,9 @@ FLongPackagePathsSingleton::FLongPackagePathsSingleton()
 		FMountPointKeyByPtr* ExistingFullKey = MountPoints.FindByHash(GetTypeHash(SearchKey), SearchKey);
 		if (ExistingFullKey)
 		{
-			UE_LOG(LogPackageName, Warning, TEXT("Unexpected duplicate initial MountPoint added during PackageName startup: %s, %s."),
-				*RootPath, *ContentPathRelative);
+			// Duplicates can occur in our manual list if e.g. GameContentPathRebased == GameContentPath.
+			// If we call LocalInsertMountPoint with the same pair of LongPackageName, LocalPath as one we've done
+			// earlier, just skip it.
 			return;
 		}
 
@@ -620,22 +621,10 @@ FLongPackagePathsSingleton::FLongPackagePathsSingleton()
 	{
 		LocalInsertMountPoint(EngineRootPath, ContentPathShort, EMountFlags::NotReadOnly | EMountFlags::Alias);
 	}
-	if (GameContentPathRebased != GameContentPath)
-	{
-		LocalInsertMountPoint(GameRootPath, GameContentPathRebased, EMountFlags::NotReadOnly | EMountFlags::Alias);
-	}
-	if (GameScriptPathRebased != GameScriptPath)
-	{
-		LocalInsertMountPoint(ScriptRootPath, GameScriptPathRebased, EMountFlags::ReadOnly | EMountFlags::Alias);
-	}
-	if (GameSavedPathRebased != GameSavedPath)
-	{
-		LocalInsertMountPoint(TempRootPath, GameSavedPathRebased, EMountFlags::ReadOnly | EMountFlags::Alias);
-	}
-	if (GameConfigPathRebased != GameConfigPath)
-	{
-		LocalInsertMountPoint(ConfigRootPath, GameConfigPathRebased, EMountFlags::ReadOnly | EMountFlags::Alias);
-	}
+	LocalInsertMountPoint(GameRootPath,   GameContentPathRebased, EMountFlags::NotReadOnly | EMountFlags::Alias);
+	LocalInsertMountPoint(ScriptRootPath, GameScriptPathRebased, EMountFlags::ReadOnly | EMountFlags::Alias);
+	LocalInsertMountPoint(TempRootPath,   GameSavedPathRebased, EMountFlags::ReadOnly | EMountFlags::Alias);
+	LocalInsertMountPoint(ConfigRootPath, GameConfigPathRebased, EMountFlags::ReadOnly | EMountFlags::Alias);
 }
 
 FLongPackagePathsSingleton::~FLongPackagePathsSingleton()
