PackageName: Remove incorrectly-added restriction that MountPoints can not have LongPackageName roots with more than one path segment. Multi-segment paths such as /Game/Weapons/ are used by the Template system.
#rnx

#Synced-CL 50309879

diff --git a/Engine/Source/Runtime/CoreUObject/Private/Misc/PackageName.cpp b/Engine/Source/Runtime/CoreUObject/Private/Misc/PackageName.cpp
--- a/Engine/Source/Runtime/CoreUObject/Private/Misc/PackageName.cpp
+++ b/Engine/Source/Runtime/CoreUObject/Private/Misc/PackageName.cpp
@@ -202,22 +202,17 @@ bool FLongPackagePathsSingleton::TryNormalizePaths(const FString& InLongPackageN
 bool FLongPackagePathsSingleton::TryNormalizeLongPackageNamePath(const FString& InLongPackageName,
 	FString& OutLongPackageName, bool bSilenceLogs)
 {
-	int32 NumMidSlashes = 0;
 	int32 NumBackSlashes = 0;
 	int32 LenLongPackageName = InLongPackageName.Len();
 	for (int32 N = 0; N < LenLongPackageName; ++N)
 	{
 		TCHAR C = (*InLongPackageName)[N];
-		if (N != 0 && N != LenLongPackageName - 1 && C == '/')
-		{
-			++NumMidSlashes;
-		}
-		else if (C == '\\')
+		if (C == '\\')
 		{
 			++NumBackSlashes;
 		}
 	}
-	if (!InLongPackageName.StartsWith("/") || NumMidSlashes > 0 || NumBackSlashes > 0)
+	if (!InLongPackageName.StartsWith("/") || NumBackSlashes > 0)
 	{
 		if (!bSilenceLogs)
 		{
diff --git a/Engine/Source/Runtime/CoreUObject/Public/Misc/MountPoint.h b/Engine/Source/Runtime/CoreUObject/Public/Misc/MountPoint.h
--- a/Engine/Source/Runtime/CoreUObject/Public/Misc/MountPoint.h
+++ b/Engine/Source/Runtime/CoreUObject/Public/Misc/MountPoint.h
@@ -18,26 +18,43 @@ namespace UE::PackageName
  * to multiple LongPackageNames. Only the last MountPoint mounted for a LongPackageName or LocalPath is returned from
  * queries for those paths, but all mounted MountPoints can still be used to load packages, they report
  * IsMounted == true, and they are reported in query functions that report the list of mounted MountPoints.
+ * 
+ * Most MountPoints have a one-segment length path in their LongPackageName, such as /Game/ or /PluginName/, but the
+ * system supports paths of any length as long as they are rooted. MountPoints that are subpaths of existing
+ * MountPoints will take precedence for mapping files under their path, even if their parent MountPoint is higher
+ * priority because it was added later. This feature is used for example by the UnrealEd templates system in
+ * FUnrealEdMisc::MountTemplateSharedPaths, which mounts some providers for feature templates in a mountpoint
+ * subdirectory of /Game/ such as /Game/Weapons/.
  */
 class IMountPoint : public IRefCountedObject
 {
 public:
 	/**
 	 * PackageName path that is the parent directory of packagenames in this mountpoint. Has a trailing slash.
-	 * e.g. "/Game/", "/PluginName/"
+	 * e.g.
+	 *     "/Game/",
+	 *     "/PluginName/",
+	 *     "/Game/Weapons/"
 	 * Immutable, can be read from any thread.
 	 */
 	virtual FStringView GetLongPackageName() const = 0;
 	/**
 	 * LocalPath, in Unreal standard path form (normalized and relative to Engine/Binaries/<Platform> directory), that
 	 * is the parent directory of filenames in this mountpoint. Has a trailing slash.
-	 * e.g. "../../../EngineTest/Content/", "../../../EngineTest/Plugins/PluginName/Content/"
+	 * e.g.
+	 *     "../../../EngineTest/Content/",
+	 *     "../../../EngineTest/Plugins/PluginName/Content/",
+	 *     "../../../Templates/TemplateResources/Standard/Weapons/Content/"
 	 * Immutable, can be read from any thread.
 	 */
 	virtual FStringView GetLocalPathStandard() const = 0;
 	/**
 	 * LocalPath, absolute, that is the parent directory of filenames in this mountpoint. Has a trailing slash.
-	 * e.g. "d:/UnrealRoot/EngineTest/Content/", "d:/UnrealRoot/EngineTest/Plugins/PluginName/Content/"
+	 * e.g.
+	 *     "d:/UnrealRoot/EngineTest/Content/",
+	 *     "d:/UnrealRoot/EngineTest/Plugins/PluginName/Content/"
+	 *     "d:/UnrealRoot/Templates/TemplateResources/Standard/Weapons/Content/",
+	 *
 	 * Immutable, can be read from any thread.
 	 */
 	virtual FStringView GetLocalPathAbsolute() const = 0;
