[ZenLoader]
- Fix timeout detection not accounting for activity in worker threads

#Synced-CL 50254652

diff --git a/Engine/Source/Runtime/CoreUObject/Private/Serialization/AsyncLoading2.cpp b/Engine/Source/Runtime/CoreUObject/Private/Serialization/AsyncLoading2.cpp
--- a/Engine/Source/Runtime/CoreUObject/Private/Serialization/AsyncLoading2.cpp
+++ b/Engine/Source/Runtime/CoreUObject/Private/Serialization/AsyncLoading2.cpp
@@ -3467,7 +3467,7 @@ private:
 
 	void ProcessDependencies(FAsyncLoadingThreadState2& ThreadState);
 	void Fire(FAsyncLoadingThreadState2* ThreadState);
-	void ParallelLoadingWork();
+	void ParallelLoadingWork(FAsyncLoadingThreadState2* InParentState);
 
 	const FAsyncLoadEventSpec* Spec = nullptr;
 	FAsyncPackage2* Package = nullptr;
@@ -3701,7 +3701,14 @@ struct FAsyncLoadingThreadState2
 
 	void MarkAsActive()
 	{
-		bIsActive.store(true, std::memory_order_relaxed);
+		if (ParentState)
+		{
+			ParentState->MarkAsActive();
+		}
+		else
+		{
+			bIsActive.store(true, std::memory_order_relaxed);
+		}
 	}
 
 	void ResetActivity()
@@ -3720,6 +3727,7 @@ struct FAsyncLoadingThreadState2
 	TArray<FAsyncPackage2*> PackagesOnStack;
 	TMpscQueue<FAsyncLoadingSyncLoadContext*> SyncLoadContextsCreatedOnOtherThreads;
 	TMpscQueue<FAsyncPackage2*> PackagesToReprioritize;
+	FAsyncLoadingThreadState2* ParentState = nullptr;
 	bool bIsAsyncLoadingThread = false;
 	bool bCanAccessAsyncLoadingThreadData = false;
 	// used to probe activity for stall detection by the game thread
@@ -6397,7 +6405,7 @@ public:
 	}
 };
 
-void FEventLoadNode2::ParallelLoadingWork()
+void FEventLoadNode2::ParallelLoadingWork(FAsyncLoadingThreadState2* InParentState)
 {
 #if USING_INSTRUMENTATION
 	if (GDetectRaceDuringLoading)
@@ -6409,7 +6417,7 @@ void FEventLoadNode2::ParallelLoadingWork()
 
 	LLM_SCOPE(ELLMTag::AsyncLoading);
 	FTaskTagScope Scope(ETaskTag::EParallelLoadingThread);
-	FAsyncLoadingThreadState2 LocalThreadState{ .bIsWorkerThread = true };
+	FAsyncLoadingThreadState2 LocalThreadState{ .ParentState = InParentState, .bIsWorkerThread = true };
 
 	FAsyncLoadingThreadState2Scope WorkerStateScope(LocalThreadState);
 
@@ -6458,9 +6466,9 @@ void FEventLoadNode2::Fire(FAsyncLoadingThreadState2* ThreadState)
 		{
 			UE::Tasks::Launch(
 				TEXT("ParallelLoadingWork"),
-				[this]()
+				[this, ParentState = ThreadState]()
 				{
-					ParallelLoadingWork();
+					ParallelLoadingWork(ParentState);
 
 					FGCCSyncObject::Get().UnlockAsync(bShouldLock);
 				}
@@ -7918,6 +7926,9 @@ bool FAsyncPackage2::CreateLinkerLoadExports(FAsyncLoadingThreadState2& ThreadSt
 			continue;
 		}
 #endif // WITH_METADATA
+
+		ThreadState.MarkAsActive();
+
 		FObjectExport& LinkerExport = LinkerLoadState->Linker->ExportMap[ExportIndex];
 		FExportObject& ExportObject = Data.Exports[ExportIndex];
 
