Remove old entries when bReplaceMap is true in FMapProperty::ConvertFromType

#rnx
#Synced-CL 50152864

diff --git a/Engine/Source/Runtime/CoreUObject/Private/UObject/PropertyMap.cpp b/Engine/Source/Runtime/CoreUObject/Private/UObject/PropertyMap.cpp
--- a/Engine/Source/Runtime/CoreUObject/Private/UObject/PropertyMap.cpp
+++ b/Engine/Source/Runtime/CoreUObject/Private/UObject/PropertyMap.cpp
@@ -1652,6 +1652,11 @@ EConvertFromTypeResult FMapProperty::ConvertFromType(const FPropertyTag& Tag, FS
 
 	if (bConversionSucceeded)
 	{
+		if (bReplaceMap)
+		{
+			MapHelper.EmptyValues(Num);
+		}
+
 		if (Num != 0)
 		{
 			if (TempKeyValueStorage == nullptr)
