[AsyncLoading2]

Allows GlobalImportStore to be resized after we remove unreachable objects to reduce the amount of memory allocated by the loader. Reduces peak memory usage from the loader by 5-6MB on consoles but may add 100s of microseconds on the GT when collecting unreachable objects. Disabled by default, enableable via s.ShrinkGlobalImportStoreOnGT

#rnx

#Synced-CL 50345503

diff --git a/Engine/Source/Runtime/CoreUObject/Private/Serialization/AsyncLoading2.cpp b/Engine/Source/Runtime/CoreUObject/Private/Serialization/AsyncLoading2.cpp
--- a/Engine/Source/Runtime/CoreUObject/Private/Serialization/AsyncLoading2.cpp
+++ b/Engine/Source/Runtime/CoreUObject/Private/Serialization/AsyncLoading2.cpp
@@ -502,6 +502,14 @@ static FAutoConsoleVariableRef CVarGRemoveUnreachableObjectsOnGT(
 	ECVF_Default
 );
 
+static bool GShrinkGlobalImportStoreOnGT = false;
+static FAutoConsoleVariableRef CVarGShrinkGlobalImportStoreOnGT(
+	TEXT("s.ShrinkGlobalImportStoreOnGT"),
+	GShrinkGlobalImportStoreOnGT,
+	TEXT("Shrink the GlobalImportStore's working set of memory on the GT from the GC callback NotifyUnreachableObjects."),
+	ECVF_Default
+);
+
 #if UE_BUILD_DEBUG || (UE_BUILD_DEVELOPMENT && !WITH_EDITOR)
 static bool GVerifyUnreachableObjects = true;
 static bool GVerifyObjectLoadFlags = true;
@@ -2736,6 +2744,17 @@ public:
 		UE_CLOG(RemovedCount, LogStreaming, Display, TEXT("RemoveUnreachableScriptObjects: removed %d entries"), RemovedCount);
 	}
 
+	void Shrink()
+	{
+		LLM_SCOPE_BYTAG(AsyncLoading_GlobalImportStore);
+		TRACE_CPUPROFILER_EVENT_SCOPE(FGlobalImportStore::Shrink);
+		Packages.Shrink();
+		PackageCells.Shrink();
+		ScriptObjects.Shrink();
+		ScriptCells.Shrink();
+		ObjectIndexToPublicExport.Shrink();
+	}
+
 #if WITH_VERSE_VM || defined(__INTELLISENSE__)
 	void AddScriptCellPackage(Verse::VPackage* Package)
 	{
@@ -11703,6 +11722,11 @@ void FAsyncLoadingThread2::NotifyUnreachableObjects(const TArrayView<FUObjectIte
 		RemoveUnreachableObjects(UnreachableObjects);
 	}
 
+	if (GShrinkGlobalImportStoreOnGT)
+	{
+		GlobalImportStore.Shrink();
+	}
+
 	// wake up ALT to remove unreachable objects
 	AltZenaphore.NotifyAll();
 }
