* Changed ImportTypeHierarchies to contain the hierarchies for all types that necessitate it (usage of a DynamicImport callback or can create PropertyBagPlaceholder), this correctly inserts the hierarchy for BP derived classes (like prefabs).
* Changed LinkerLoad to consider exports of dynamically imported packages as potential exports of the package doing the import, this saving objects in external packages (loaded through a dynamic import) without having to resave their referencers in other packages.

#Synced-CL 50398382

diff --git a/Engine/Source/Runtime/CoreUObject/Private/UObject/LinkerLoad.cpp b/Engine/Source/Runtime/CoreUObject/Private/UObject/LinkerLoad.cpp
--- a/Engine/Source/Runtime/CoreUObject/Private/UObject/LinkerLoad.cpp
+++ b/Engine/Source/Runtime/CoreUObject/Private/UObject/LinkerLoad.cpp
@@ -1201,6 +1201,9 @@ FLinkerLoad::FLinkerLoad(UPackage* InParent, const FPackagePath& InPackagePath,
 , bIsPackageRelocated(false)
 , bIsSerializingScriptProperties(false)
 #endif // WITH_EDITOR
+#if WITH_EDITORONLY_DATA
+, DynamicImportsIndex(INDEX_NONE)
+#endif // WITH_EDITORONLY_DATA
 , IsTimeLimitExceededCallCount(0)
 , TimeLimit(0.0f)
 , TickStartTime(0.0)
@@ -2906,6 +2909,12 @@ FLinkerLoad::ELinkerStatus FLinkerLoad::AddDynamicImports()
 						continue;
 					}
 
+					// set the dynamic import index if necessary 
+					if (DynamicImportsIndex == INDEX_NONE)
+					{
+						DynamicImportsIndex = ImportMap.Num();
+					}
+
 					FObjectImport& Import = ImportMap.AddDefaulted_GetRef();
 					Import.ObjectName = DynamicImport.PackageName;
 					Import.ClassPackage = GLongCoreUObjectPackageName;
@@ -4218,6 +4227,68 @@ bool FLinkerLoad::VerifyImportInner(const int32 ImportIndex, FString& WarningSuf
 #endif
 			}
 
+#if WITH_EDITORONLY_DATA
+			if (!Import.SourceLinker && OuterImport.SourceLinker && OuterImport.SourceLinker->DynamicImportsIndex != INDEX_NONE)
+			{
+				// We consider the exports of the dynamic imports of the outer as if they were exports of the outer's linker, so we
+				// explore them to find a matching export before we default back to using the the outer's linker
+				//
+				// The reasoning behind that is that dynamic imports are used to move objects in external packages 
+				// and we do not want all packages importing them to have to be resaved when this happens since we know 
+				// how to resolve the situation efficiently
+				const int32 ObjectHash = GetHashBucket( Import.ObjectName );
+
+				// Prevents any other serialization function from modifying this linker while we're reading from it
+				UE::TUniqueLock UniqueLock(OuterImport.SourceLinker->Mutex);
+
+				for (int32 i = OuterImport.SourceLinker->DynamicImportsIndex; i < OuterImport.SourceLinker->ImportMap.Num(); i++)
+				{
+					FObjectImport& ObjImport = OuterImport.SourceLinker->Imp(FPackageIndex::FromImport(i));
+					
+					if (ObjImport.SourceLinker && ObjImport.SourceLinker->ExportHash)
+					{
+						FLinkerLoad* DynamicImportLinker = ObjImport.SourceLinker;
+						
+						// Prevents any other serialization function from modifying this linker while we're reading from it
+						UE::TUniqueLock DynImportLinkerUniqueLock(DynamicImportLinker->Mutex);
+
+						// go through exports to find a matching object 
+						for (int32 j= DynamicImportLinker->ExportHash[ObjectHash]; j!=INDEX_NONE; j=DynamicImportLinker->ExportMap[j].HashNext )
+						{
+							FObjectExport& ObjExport = DynamicImportLinker->ExportMap[j];
+
+							if (ObjExport.ObjectName == Import.ObjectName)
+							{
+								// walk up the outer chains to compare & validate we've got right object
+								FObjectResource* ImportOuterResource = ImpExpPtr(Import.OuterIndex);
+								FObjectResource* ExportOuterResource = DynamicImportLinker->ImpExpPtr(ObjExport.OuterIndex);
+
+								while( (ImportOuterResource && ExportOuterResource) 
+									    && ImportOuterResource->ObjectName == ExportOuterResource->ObjectName)
+								{
+									ImportOuterResource = ImpExpPtr(ImportOuterResource->OuterIndex);
+									ExportOuterResource = DynamicImportLinker->ImpExpPtr(ExportOuterResource->OuterIndex);
+								} 
+
+								// This means we managed to walk up the chain and all objects were a match
+								if (!ImportOuterResource && !ExportOuterResource)
+								{
+									UE_LOG(LogLinker, Verbose, TEXT("Assigning %s as linker to object %s in %s"), *DynamicImportLinker->GetLinker()->LinkerRoot->GetName(), *Import.ObjectName.ToString(), *LinkerRoot->GetName());
+									Import.SourceLinker = DynamicImportLinker;
+									break;
+								}
+							}
+						}
+
+						if (Import.SourceLinker)
+						{
+							break;
+						}
+					}
+				}
+			}
+#endif // WITH_EDITORONLY_DATA
+
 			// Copy the SourceLinker from the FObjectImport for our Outer if the SourceLinker hasn't been set yet,
 			// Otherwise we may be overwriting a re-directed linker and SourceIndex is already from the redirected one
 			// or we had an assigned package and our linker is already set.
@@ -8056,4 +8127,4 @@ void FLinkerLoad::SerializeBulkMeta(UE::BulkData::Private::FBulkMetaData& Meta,
 #endif // WITH_EDITOR
 }
 
-#undef LOCTEXT_NAMESPACE
+#undef LOCTEXT_NAMESPACE
\ No newline at end of file
diff --git a/Engine/Source/Runtime/CoreUObject/Private/UObject/SavePackage2.cpp b/Engine/Source/Runtime/CoreUObject/Private/UObject/SavePackage2.cpp
--- a/Engine/Source/Runtime/CoreUObject/Private/UObject/SavePackage2.cpp
+++ b/Engine/Source/Runtime/CoreUObject/Private/UObject/SavePackage2.cpp
@@ -1680,17 +1680,27 @@ ESavePackageResult BuildLinker(FSaveContext& SaveContext)
 
 			for (FObjectExport& Export : Linker->ExportMap)
 			{
-				if (!Export.ClassIndex.IsNull() && Export.ClassIndex.IsImport())
+				if (!Export.ClassIndex.IsNull())
 				{
-					FObjectImport& ClassImport = Linker->ImportMap[Export.ClassIndex.ToImport()];
-
-					if (ClassImport.XObject)
+					UObject* ClassObject = nullptr;
+					if (Export.ClassIndex.IsExport())
+					{
+						FObjectExport& ClassExport = Linker->ExportMap[Export.ClassIndex.ToExport()];
+						ClassObject = ClassExport.Object;
+					}
+					else
+					{
+						FObjectImport& ClassImport = Linker->ImportMap[Export.ClassIndex.ToImport()];
+						ClassObject = ClassImport.XObject;
+					}
+					
+					if (ClassObject)
 					{
-						UStruct* Struct = Cast<UStruct>(ClassImport.XObject);
+						UStruct* Struct = Cast<UStruct>(ClassObject);
 						if (Struct)
 						{
-							const FPackageIndex PackageIndex = Linker->MapObject(ClassImport.XObject);
-							// This should always be in the import table already.
+							const FPackageIndex PackageIndex = Linker->MapObject(ClassObject);
+							// This should always be in the linker tables already.
 							if (ensureMsgf(!PackageIndex.IsNull(), TEXT("%s: Struct object: '%s' is not mapped"), *Linker->LinkerRoot->GetFullName(), *Struct->GetFullName()))
 							{
 								// Skip already processed class imports.
diff --git a/Engine/Source/Runtime/CoreUObject/Public/UObject/LinkerLoad.h b/Engine/Source/Runtime/CoreUObject/Public/UObject/LinkerLoad.h
--- a/Engine/Source/Runtime/CoreUObject/Public/UObject/LinkerLoad.h
+++ b/Engine/Source/Runtime/CoreUObject/Public/UObject/LinkerLoad.h
@@ -485,6 +485,11 @@ private:
 	bool bIsSerializingScriptProperties : 1;
 #endif // WITH_EDITOR
 
+#if WITH_EDITORONLY_DATA
+	/** Index of the 1st dynamic import that was added (all dynamic imports are consecutive at the end of the import table  */
+	int32					DynamicImportsIndex;
+#endif //#if WITH_EDITORONLY_DATA
+
 	/** Call count of IsTimeLimitExceeded.																					*/
 	int32					IsTimeLimitExceededCallCount;
 	/** Current time limit to use if bUseTimeLimit is true.																	*/
