CoreUObject: Change the FUNC_AllFlags enumerator to be a define just like the other FUNC_ flag combinations.
Add the comment found on core enums in ObjectMacrors.h to EFunctionFlags to clarify that composite flags should not be made a enumerators within the enums.

#rnx

#Synced-CL 50399333

diff --git a/Engine/Source/Runtime/CoreUObject/Public/UObject/ObjectMacros.h b/Engine/Source/Runtime/CoreUObject/Public/UObject/ObjectMacros.h
--- a/Engine/Source/Runtime/CoreUObject/Public/UObject/ObjectMacros.h
+++ b/Engine/Source/Runtime/CoreUObject/Public/UObject/ObjectMacros.h
@@ -163,7 +163,7 @@ enum EPackageFlags
 	PKG_FilterEditorOnly			= 0x80000000,	///< Package has editor-only data filtered out
 
 	// !!!!!!!!!!!!!!!!!!!!!!!!!!!!!!!!!!!!!!!!!!!!!!!!!!!!!!!!!!!!!!!!!!!!!!!!!!!!!!!!!!!!!!!!!!!!!!!!!!!!!!!!!!!!!!!!!!!!!!!!!!!!
-	// DO NOT ADD composite flags to EPackageFlags. Debugger visualisations have trouble displaying composite flag values as text.
+	// DO NOT ADD composite flags to EPackageFlags. Debugger visualizations have trouble displaying composite flag values as text.
 	// !!!!!!!!!!!!!!!!!!!!!!!!!!!!!!!!!!!!!!!!!!!!!!!!!!!!!!!!!!!!!!!!!!!!!!!!!!!!!!!!!!!!!!!!!!!!!!!!!!!!!!!!!!!!!!!!!!!!!!!!!!!!
 };
 
@@ -280,7 +280,7 @@ enum EClassFlags
 	CLASS_NewerVersionExists  = 0x80000000u,
 
 	// !!!!!!!!!!!!!!!!!!!!!!!!!!!!!!!!!!!!!!!!!!!!!!!!!!!!!!!!!!!!!!!!!!!!!!!!!!!!!!!!!!!!!!!!!!!!!!!!!!!!!!!!!!!!!!!!!!!!!!!!!!
-	// DO NOT ADD composite flags to EClassFlags. Debugger visualisations have trouble displaying composite flag values as text.
+	// DO NOT ADD composite flags to EClassFlags. Debugger visualizations have trouble displaying composite flag values as text.
 	// !!!!!!!!!!!!!!!!!!!!!!!!!!!!!!!!!!!!!!!!!!!!!!!!!!!!!!!!!!!!!!!!!!!!!!!!!!!!!!!!!!!!!!!!!!!!!!!!!!!!!!!!!!!!!!!!!!!!!!!!!!
 };
 
@@ -409,7 +409,7 @@ enum EClassCastFlags : uint64
 	//										= 0x8000000000000000,
 
 	// !!!!!!!!!!!!!!!!!!!!!!!!!!!!!!!!!!!!!!!!!!!!!!!!!!!!!!!!!!!!!!!!!!!!!!!!!!!!!!!!!!!!!!!!!!!!!!!!!!!!!!!!!!!!!!!!!!!!!!!!!!!!!!
-	// DO NOT ADD composite flags to EClassCastFlags. Debugger visualisations have trouble displaying composite flag values as text.
+	// DO NOT ADD composite flags to EClassCastFlags. Debugger visualizations have trouble displaying composite flag values as text.
 	// !!!!!!!!!!!!!!!!!!!!!!!!!!!!!!!!!!!!!!!!!!!!!!!!!!!!!!!!!!!!!!!!!!!!!!!!!!!!!!!!!!!!!!!!!!!!!!!!!!!!!!!!!!!!!!!!!!!!!!!!!!!!!!
 };
 
@@ -494,7 +494,7 @@ enum EPropertyFlags : uint64
 	CPF_ForcePostConstructLink			= 0x2000000000000000,	///< Use to force native UCLASS properties to be included in the PostConstructLink chain in a derived Blueprint class, regardless of whether it differs from the native class constructor-initialized value
 
 	// !!!!!!!!!!!!!!!!!!!!!!!!!!!!!!!!!!!!!!!!!!!!!!!!!!!!!!!!!!!!!!!!!!!!!!!!!!!!!!!!!!!!!!!!!!!!!!!!!!!!!!!!!!!!!!!!!!!!!!!!!!!!!
-	// DO NOT ADD composite flags to EPropertyFlags. Debugger visualisations have trouble displaying composite flag values as text.
+	// DO NOT ADD composite flags to EPropertyFlags. Debugger visualizations have trouble displaying composite flag values as text.
 	// !!!!!!!!!!!!!!!!!!!!!!!!!!!!!!!!!!!!!!!!!!!!!!!!!!!!!!!!!!!!!!!!!!!!!!!!!!!!!!!!!!!!!!!!!!!!!!!!!!!!!!!!!!!!!!!!!!!!!!!!!!!!!
 };
 
@@ -535,7 +535,7 @@ enum class EArrayPropertyFlags : uint8
 	UsesMemoryImageAllocator
 
 	// !!!!!!!!!!!!!!!!!!!!!!!!!!!!!!!!!!!!!!!!!!!!!!!!!!!!!!!!!!!!!!!!!!!!!!!!!!!!!!!!!!!!!!!!!!!!!!!!!!!!!!!!!!!!!!!!!!!!!!!!!!!!!!!!!!
-	// DO NOT ADD composite flags to EArrayPropertyFlags. Debugger visualisations have trouble displaying composite flag values as text.
+	// DO NOT ADD composite flags to EArrayPropertyFlags. Debugger visualizations have trouble displaying composite flag values as text.
 	// !!!!!!!!!!!!!!!!!!!!!!!!!!!!!!!!!!!!!!!!!!!!!!!!!!!!!!!!!!!!!!!!!!!!!!!!!!!!!!!!!!!!!!!!!!!!!!!!!!!!!!!!!!!!!!!!!!!!!!!!!!!!!!!!!!
 };
 
@@ -550,7 +550,7 @@ enum class EMapPropertyFlags : uint8
 	UsesMemoryImageAllocator
 
 	// !!!!!!!!!!!!!!!!!!!!!!!!!!!!!!!!!!!!!!!!!!!!!!!!!!!!!!!!!!!!!!!!!!!!!!!!!!!!!!!!!!!!!!!!!!!!!!!!!!!!!!!!!!!!!!!!!!!!!!!!!!!!!!!!
-	// DO NOT ADD composite flags to EMapPropertyFlags. Debugger visualisations have trouble displaying composite flag values as text.
+	// DO NOT ADD composite flags to EMapPropertyFlags. Debugger visualizations have trouble displaying composite flag values as text.
 	// !!!!!!!!!!!!!!!!!!!!!!!!!!!!!!!!!!!!!!!!!!!!!!!!!!!!!!!!!!!!!!!!!!!!!!!!!!!!!!!!!!!!!!!!!!!!!!!!!!!!!!!!!!!!!!!!!!!!!!!!!!!!!!!!
 };
 
@@ -565,7 +565,7 @@ enum class EPropertyObjectReferenceType : uint32
 	Conservative = 1 << 3,	 // Not a real reference type, used to mark native struct serializers which may serialize unknown reference types and to conservatively populate RefLink with such struct properties.
 
 	// !!!!!!!!!!!!!!!!!!!!!!!!!!!!!!!!!!!!!!!!!!!!!!!!!!!!!!!!!!!!!!!!!!!!!!!!!!!!!!!!!!!!!!!!!!!!!!!!!!!!!!!!!!!!!!!!!!!!!!!!!!!!!!!!!!!!!!!!!!!
-	// DO NOT ADD composite flags to EPropertyObjectReferenceType. Debugger visualisations have trouble displaying composite flag values as text.
+	// DO NOT ADD composite flags to EPropertyObjectReferenceType. Debugger visualizations have trouble displaying composite flag values as text.
 	// !!!!!!!!!!!!!!!!!!!!!!!!!!!!!!!!!!!!!!!!!!!!!!!!!!!!!!!!!!!!!!!!!!!!!!!!!!!!!!!!!!!!!!!!!!!!!!!!!!!!!!!!!!!!!!!!!!!!!!!!!!!!!!!!!!!!!!!!!!!
 };
 ENUM_CLASS_FLAGS(EPropertyObjectReferenceType);
@@ -630,7 +630,7 @@ enum EObjectFlags
 	RF_AllocatedInSharedPage UE_DEPRECATED(5.8, "This flag is now obsolete as the functionality is unnecessary.") =0x80000000,	///< Allocated from a ref-counted page shared with other UObjects
 
 	// !!!!!!!!!!!!!!!!!!!!!!!!!!!!!!!!!!!!!!!!!!!!!!!!!!!!!!!!!!!!!!!!!!!!!!!!!!!!!!!!!!!!!!!!!!!!!!!!!!!!!!!!!!!!!!!!!!!!!!!!!!!
-	// DO NOT ADD composite flags to EObjectFlags. Debugger visualisations have trouble displaying composite flag values as text.
+	// DO NOT ADD composite flags to EObjectFlags. Debugger visualizations have trouble displaying composite flag values as text.
 	// !!!!!!!!!!!!!!!!!!!!!!!!!!!!!!!!!!!!!!!!!!!!!!!!!!!!!!!!!!!!!!!!!!!!!!!!!!!!!!!!!!!!!!!!!!!!!!!!!!!!!!!!!!!!!!!!!!!!!!!!!!!
 };
 
@@ -701,7 +701,7 @@ enum class EInternalObjectFlags : int32
 	PendingConstruction = 1 << 31, ///< Object didn't have its class constructor called yet (only the UObjectBase one to initialize its most basic members)
 
 	// !!!!!!!!!!!!!!!!!!!!!!!!!!!!!!!!!!!!!!!!!!!!!!!!!!!!!!!!!!!!!!!!!!!!!!!!!!!!!!!!!!!!!!!!!!!!!!!!!!!!!!!!!!!!!!!!!!!!!!!!!!!!!!!!!!!
-	// DO NOT ADD composite flags to EInternalObjectFlags. Debugger visualisations have trouble displaying composite flag values as text.
+	// DO NOT ADD composite flags to EInternalObjectFlags. Debugger visualizations have trouble displaying composite flag values as text.
 	// !!!!!!!!!!!!!!!!!!!!!!!!!!!!!!!!!!!!!!!!!!!!!!!!!!!!!!!!!!!!!!!!!!!!!!!!!!!!!!!!!!!!!!!!!!!!!!!!!!!!!!!!!!!!!!!!!!!!!!!!!!!!!!!!!!!
 };
 ENUM_CLASS_FLAGS(EInternalObjectFlags);
@@ -729,7 +729,7 @@ enum class EEnumFlags : uint8
 	NewerVersionExists = 0x00000002, // If set, this UEnum has been replaced by a newer version
 
 	// !!!!!!!!!!!!!!!!!!!!!!!!!!!!!!!!!!!!!!!!!!!!!!!!!!!!!!!!!!!!!!!!!!!!!!!!!!!!!!!!!!!!!!!!!!!!!!!!!!!!!!!!!!!!!!!!!!!!!!!!!
-	// DO NOT ADD composite flags to EEnumFlags. Debugger visualisations have trouble displaying composite flag values as text.
+	// DO NOT ADD composite flags to EEnumFlags. Debugger visualizations have trouble displaying composite flag values as text.
 	// !!!!!!!!!!!!!!!!!!!!!!!!!!!!!!!!!!!!!!!!!!!!!!!!!!!!!!!!!!!!!!!!!!!!!!!!!!!!!!!!!!!!!!!!!!!!!!!!!!!!!!!!!!!!!!!!!!!!!!!!!
 };
 
diff --git a/Engine/Source/Runtime/CoreUObject/Public/UObject/Script.h b/Engine/Source/Runtime/CoreUObject/Public/UObject/Script.h
--- a/Engine/Source/Runtime/CoreUObject/Public/UObject/Script.h
+++ b/Engine/Source/Runtime/CoreUObject/Public/UObject/Script.h
@@ -167,7 +167,9 @@ enum EFunctionFlags : uint32
 	FUNC_Const				= 0x40000000,	// function can be called from blueprint code, and only reads state (never writes state)
 	FUNC_NetValidate		= 0x80000000,	// function must supply a _Validate implementation
 
-	FUNC_AllFlags		= 0xFFFFFFFF,
+	// !!!!!!!!!!!!!!!!!!!!!!!!!!!!!!!!!!!!!!!!!!!!!!!!!!!!!!!!!!!!!!!!!!!!!!!!!!!!!!!!!!!!!!!!!!!!!!!!!!!!!!!!!!!!!!!!!!!!!!!!!!!
+	// DO NOT ADD composite flags to EFunctionFlags. Debugger visualizations have trouble displaying composite flag values as text.
+	// !!!!!!!!!!!!!!!!!!!!!!!!!!!!!!!!!!!!!!!!!!!!!!!!!!!!!!!!!!!!!!!!!!!!!!!!!!!!!!!!!!!!!!!!!!!!!!!!!!!!!!!!!!!!!!!!!!!!!!!!!!!
 };
 
 UE_FORCEINLINE_HINT FArchive& operator<<(FArchive& Ar, EFunctionFlags& Flags)
@@ -183,6 +185,7 @@ ENUM_CLASS_FLAGS(EFunctionFlags)
 #define FUNC_FuncOverrideMatch ((EFunctionFlags)(FUNC_Exec | FUNC_Final | FUNC_Static | FUNC_Public | FUNC_Protected | FUNC_Private))
 #define FUNC_NetFuncFlags      ((EFunctionFlags)(FUNC_Net | FUNC_NetReliable | FUNC_NetServer | FUNC_NetClient | FUNC_NetMulticast))
 #define FUNC_AccessSpecifiers  ((EFunctionFlags)(FUNC_Public | FUNC_Private | FUNC_Protected))
+#define FUNC_AllFlags          ((EFunctionFlags)(0xFFFFFFFF))
 
 //
 // Evaluatable expression item types.
