[Mass] Extensions to TStructTypeBitSet:
* made TStructTypeBitSet(const int32) constructor public
* added GetAtIndex that lets caller check value at specified index
* added AndNot function, that calculates C = A & !B, without intermediate variable storing !B

#Synced-CL 50246518

diff --git a/Engine/Source/Runtime/CoreUObject/Public/StructUtils/StructTypeBitSet.h b/Engine/Source/Runtime/CoreUObject/Public/StructUtils/StructTypeBitSet.h
--- a/Engine/Source/Runtime/CoreUObject/Public/StructUtils/StructTypeBitSet.h
+++ b/Engine/Source/Runtime/CoreUObject/Public/StructUtils/StructTypeBitSet.h
@@ -111,14 +111,6 @@ struct FConstBitSetContainer : TBitArray<TInlineAllocator<8>>
 
 			FWordIterator IteratorResult(Result);
 
-			/*const int32 MinNumBits = FMath::Min(Num(), Other.Num());
-
-			for (int32 WordIndex = 0; IteratorResult; ++IteratorResult, ++IteratorA, ++IteratorB)
-			{
-				const uint32 NewValue = IteratorA.GetWord() | IteratorB.GetWord();
-				IteratorResult.SetWord(NewValue);
-			}*/
-
 			for (; IteratorA && IteratorB; ++IteratorResult, ++IteratorA, ++IteratorB)
 			{
 				const uint32 NewValue = IteratorA.GetWord() | IteratorB.GetWord();
@@ -141,6 +133,32 @@ struct FConstBitSetContainer : TBitArray<TInlineAllocator<8>>
 		return {};
 	}
 
+	FConstBitSetContainer AndNot(const FBitArrayBase& Other) const
+	{
+		const int32 NumResultBits = Num();
+		if (NumResultBits > 0)
+		{
+			FBitArrayBase Result(false, NumResultBits);
+
+			FConstWordIterator IteratorA(*this);
+			FConstWordIterator IteratorB(Other);
+
+			FWordIterator IteratorResult(Result);
+
+			for ( ; IteratorResult; ++IteratorResult, ++IteratorA, ++IteratorB)
+			{
+				const uint32 NewValue = IteratorB
+					? IteratorA.GetWord() & ~IteratorB.GetWord()
+					: IteratorA.GetWord();
+				IteratorResult.SetWord(NewValue);
+			}
+
+			return FConstBitSetContainer(MoveTemp(Result));
+		}
+
+		return {};
+	}
+
 	/**
 	 * @return whether the bitset is empty (no bits are set)
 	 */
@@ -699,6 +717,15 @@ struct TTypeBitSetBase
 		return TImplementation(StructTypesBitArray | Other.StructTypesBitArray);
 	}
 
+	/**
+	 * Performs a bitwise AND operation with the negation of the Other bitset.
+	 * @return A new bitset representing the result of the operation.
+	 */
+	UE_FORCEINLINE_HINT TImplementation AndNot(const TImplementation& Other) const
+	{
+		return TImplementation(StructTypesBitArray.AndNot(Other.StructTypesBitArray));
+	}
+
 	UE_FORCEINLINE_HINT TImplementation GetOverlap(const TImplementation& Other) const
 	{
 		return GetImplementation() & Other;
@@ -1048,13 +1075,12 @@ private:
 		StructTypesBitArray = Source;
 	}
 
+public:
 	explicit TStructTypeBitSet(const int32 BitToSet)
 	{
 		StructTypesBitArray.AddAtIndex(BitToSet);
 	}
 
-public:
-
 	/**
 	 * Retrieves the FStructTracker associated with this bitset.
 	 * @return Reference to the FStructTracker.
@@ -1171,6 +1197,15 @@ public:
 		return StructTypesBitArray.Contains(StructTypeIndex);
 	}
 
+	/**
+	 * Checks whether the bitset has a bit set at Index.
+	 * Indices past the end of the bitset are handled seamlessly (returning false).
+	 */
+	bool GetAtIndex(const int32 Index) const
+	{
+		return (Index >=0) && StructTypesBitArray.Contains(Index);
+	}
+
 	/**
 	 * Performs a union operation with another bitset.
 	 * @param Other - The other bitset to combine with.
