Removed WITH_ENGINE block in CoreNetTypes.h (used for C++ UHT) so that StaticEnum gets defined for ELifetimeCondition in that header.

diff --git a/Engine/Source/Runtime/CoreUObject/Public/UObject/CoreNetTypes.h b/Engine/Source/Runtime/CoreUObject/Public/UObject/CoreNetTypes.h
--- a/Engine/Source/Runtime/CoreUObject/Public/UObject/CoreNetTypes.h
+++ b/Engine/Source/Runtime/CoreUObject/Public/UObject/CoreNetTypes.h
@@ -3,10 +3,7 @@
 #pragma once
 
 #include "UObject/ObjectMacros.h"
-
-#if WITH_ENGINE 
-	#include "CoreNetTypes.generated.h"
-#endif
+#include "CoreNetTypes.generated.h"
 
 /** Whether or not caching of actor/subobject names to the stack should be enabled, for async demo recording crashdumps */
 #ifndef UE_NET_REPACTOR_NAME_DEBUG
