[StructUtils] Use ArePropertyGuidsAvailable for UserDefinedStruct and PropertyBag. Save property IDs for UserDefinedStruct. They are also saved in the mangled name of the property if the property is created in the UserDefinedStruct editor. This is useful when a property is renamed.

#Synced-CL 50306593

diff --git a/Engine/Source/Editor/KismetCompiler/Private/UserDefinedStructureCompilerUtils.cpp b/Engine/Source/Editor/KismetCompiler/Private/UserDefinedStructureCompilerUtils.cpp
--- a/Engine/Source/Editor/KismetCompiler/Private/UserDefinedStructureCompilerUtils.cpp
+++ b/Engine/Source/Editor/KismetCompiler/Private/UserDefinedStructureCompilerUtils.cpp
@@ -226,6 +226,7 @@ struct FUserDefinedStructureCompilerInner
 			StructToClean->ScriptAndPropertyObjectReferences.Empty();
 			StructToClean->ErrorMessage.Empty();
 			StructToClean->SetStructTrashed(true);
+			StructToClean->PropertyIDs.Empty();
 		}
 	}
 
@@ -238,10 +239,8 @@ struct FUserDefinedStructureCompilerInner
 		}
 	}
 
-	static void CreateVariables(UUserDefinedStruct* Struct, const class UEdGraphSchema_K2* Schema, FCompilerResultsLog& MessageLog)
+	static void CreateVariables(TNotNull<UUserDefinedStruct*> Struct, TNotNull<const UEdGraphSchema_K2*> Schema, FCompilerResultsLog& MessageLog)
 	{
-		check(Struct && Schema);
-
 		//FKismetCompilerUtilities::LinkAddedProperty push property to begin, so we revert the order
 		for (int32 VarDescIdx = FStructureEditorUtils::GetVarDesc(Struct).Num() - 1; VarDescIdx >= 0; --VarDescIdx)
 		{
@@ -302,6 +301,32 @@ struct FUserDefinedStructureCompilerInner
 				VarProperty->SetFlags(RF_LoadCompleted);
 				FKismetCompilerUtilities::LinkAddedProperty(Struct, VarProperty);
 			}
+
+			if (VarDesc.VarGuid.IsValid())
+			{
+				const UE::StructUtils::FUserDefinedPropertyID* Found = Struct->PropertyIDs.FindByPredicate(
+					[ID = VarDesc.VarGuid](const UE::StructUtils::FUserDefinedPropertyID& Other)
+					{
+						return Other.ID == ID;
+					});
+				if (Found != nullptr && bIsNewVariable)
+				{
+					LogError(
+						Struct,
+						MessageLog,
+						FText::Format(
+							LOCTEXT("VariableDuplicateID_ErrorFmt", "The variable {0} declared in {1} has the same ID as variable {2}."),
+							FText::FromName(VarDesc.VarName),
+							FText::FromString(Struct->GetName()),
+							FText::FromName(Found->Name)
+						).ToString()
+					);
+				}
+				else if (Found == nullptr)
+				{
+					Struct->PropertyIDs.Add(UE::StructUtils::FUserDefinedPropertyID{ .ID = VarDesc.VarGuid , .Name = VarProperty->GetFName() });
+				}
+			}
 			
 			VarProperty->SetPropertyFlags(CPF_Edit | CPF_BlueprintVisible);
 			if (VarDesc.bDontEditOnInstance)
diff --git a/Engine/Source/Runtime/CoreUObject/Private/StructUtils/PropertyBag.cpp b/Engine/Source/Runtime/CoreUObject/Private/StructUtils/PropertyBag.cpp
--- a/Engine/Source/Runtime/CoreUObject/Private/StructUtils/PropertyBag.cpp
+++ b/Engine/Source/Runtime/CoreUObject/Private/StructUtils/PropertyBag.cpp
@@ -4638,3 +4638,15 @@ bool UPropertyBag::OwnsPropertyDesc(const FPropertyBagPropertyDesc& Desc) const
 	const FPropertyBagPropertyDesc* EndDesc = StartDesc + PropertyDescs.Num();
 	return &Desc >= StartDesc && &Desc < EndDesc;
 }
+
+FName UPropertyBag::FindPropertyNameFromGuid(const FGuid& PropertyID) const
+{
+	const FPropertyBagPropertyDesc* Desc = FindPropertyDescByID(PropertyID);
+	return Desc ? Desc->Name : FName();
+}
+
+FGuid UPropertyBag::FindPropertyGuidFromName(const FName PropertyName) const
+{
+	const FPropertyBagPropertyDesc* Desc = FindPropertyDescByPropertyName(PropertyName);
+	return Desc ? Desc->ID : FGuid();
+}
diff --git a/Engine/Source/Runtime/CoreUObject/Private/StructUtils/UserDefinedStruct.cpp b/Engine/Source/Runtime/CoreUObject/Private/StructUtils/UserDefinedStruct.cpp
--- a/Engine/Source/Runtime/CoreUObject/Private/StructUtils/UserDefinedStruct.cpp
+++ b/Engine/Source/Runtime/CoreUObject/Private/StructUtils/UserDefinedStruct.cpp
@@ -62,6 +62,30 @@ void UUserDefinedStruct::PrepareCppStructOps()
 	bPrepareCppStructOpsCompleted = true;
 }
 
+FName UUserDefinedStruct::FindPropertyNameFromGuid(const FGuid& PropertyID) const
+{
+	for (const UE::StructUtils::FUserDefinedPropertyID& Property : PropertyIDs)
+	{
+		if (Property.ID == PropertyID)
+		{
+			return Property.Name;
+		}
+	}
+	return FName();
+}
+
+FGuid UUserDefinedStruct::FindPropertyGuidFromName(const FName PropertyName) const
+{
+	for (const UE::StructUtils::FUserDefinedPropertyID& Property : PropertyIDs)
+	{
+		if (Property.Name == PropertyName)
+		{
+			return Property.ID;
+		}
+	}
+	return FGuid();
+}
+
 void UUserDefinedStruct::Serialize(FStructuredArchive::FRecord Record)
 {
 	Super::Serialize(Record);
diff --git a/Engine/Source/Runtime/CoreUObject/Public/StructUtils/PropertyBag.h b/Engine/Source/Runtime/CoreUObject/Public/StructUtils/PropertyBag.h
--- a/Engine/Source/Runtime/CoreUObject/Public/StructUtils/PropertyBag.h
+++ b/Engine/Source/Runtime/CoreUObject/Public/StructUtils/PropertyBag.h
@@ -1409,7 +1409,15 @@ protected:
 	UE_API virtual void InitializeStruct(void* Dest, int32 ArrayDim = 1) const override;
 	UE_API virtual void DestroyStruct(void* Dest, int32 ArrayDim = 1) const override;
 	UE_API virtual void FinishDestroy() override;
-	
+
+	virtual bool ArePropertyGuidsAvailable() const override
+	{
+		return PropertyDescs.Num() > 0;
+	}
+
+	UE_API virtual FName FindPropertyNameFromGuid(const FGuid& PropertyID) const override;
+	UE_API virtual FGuid FindPropertyGuidFromName(const FName PropertyName) const override;
+
 	UPROPERTY()
 	TArray<FPropertyBagPropertyDesc> PropertyDescs;
 
diff --git a/Engine/Source/Runtime/CoreUObject/Public/StructUtils/UserDefinedStruct.h b/Engine/Source/Runtime/CoreUObject/Public/StructUtils/UserDefinedStruct.h
--- a/Engine/Source/Runtime/CoreUObject/Public/StructUtils/UserDefinedStruct.h
+++ b/Engine/Source/Runtime/CoreUObject/Public/StructUtils/UserDefinedStruct.h
@@ -15,6 +15,7 @@
 class UUserDefinedStruct;
 class UUserDefinedStructEditorData;
 class UStructCookedMetaData;
+struct FUserDefinedStructureCompilerInner;
 
 DECLARE_MULTICAST_DELEGATE_OneParam(FOnStructChanged, UUserDefinedStruct*);
 
@@ -52,6 +53,20 @@ public:
 	COREUOBJECT_API virtual void Initialize() override;
 };
 
+namespace UE::StructUtils
+{
+	USTRUCT()
+	struct FUserDefinedPropertyID
+	{
+		GENERATED_BODY()
+
+		UPROPERTY()
+		FGuid ID;
+		UPROPERTY()
+		FName Name;
+	};
+} // namespace
+
 UCLASS(MinimalAPI)
 class UUserDefinedStruct : public UScriptStruct
 {
@@ -78,6 +93,10 @@ public:
 	UPROPERTY(AssetRegistrySearchable)
 	FGuid Guid;
 
+private:
+	UPROPERTY()
+	TArray<UE::StructUtils::FUserDefinedPropertyID> PropertyIDs;
+
 protected:
 	/** Default instance of this struct with default values filled in, used to initialize structure */
 	FUserStructOnScopeIgnoreDefaults DefaultStructInstance;
@@ -87,7 +106,7 @@ protected:
 
 public:
 #if WITH_EDITOR
-	// UObject interface.
+	//~ UObject interface.
 	COREUOBJECT_API virtual void PostDuplicate(bool bDuplicateForPIE) override;
 	COREUOBJECT_API virtual void GetAssetRegistryTags(FAssetRegistryTagsContext Context) const override;
 	UE_DEPRECATED(5.4, "Implement the version that takes FAssetRegistryTagsContext instead.")
@@ -95,7 +114,7 @@ public:
 	COREUOBJECT_API virtual void PostLoad() override;
 	COREUOBJECT_API virtual void PreSaveRoot(FObjectPreSaveRootContext ObjectSaveContext) override;
 	COREUOBJECT_API virtual void PostSaveRoot(FObjectPostSaveRootContext ObjectSaveContext) override;
-	// End of UObject interface.
+	//~ End of UObject interface.
 
 	/** Creates a new guid if needed */
 	COREUOBJECT_API void ValidateGuid();
@@ -103,15 +122,16 @@ public:
 	COREUOBJECT_API virtual void OnChanged();
 
 	friend UUserDefinedStructEditorData;
+	friend FUserDefinedStructureCompilerInner;
 #endif	// WITH_EDITOR
 
-	// UObject interface.
+	//~ UObject interface.
 	COREUOBJECT_API virtual void Serialize(FStructuredArchive::FRecord Record) override;
 	COREUOBJECT_API virtual void SerializeTaggedProperties(FStructuredArchive::FSlot Slot, TNotNull<uint8*> Data, const UStruct* DefaultsStruct, const uint8* Defaults, const UObject* BreakRecursionIfFullyLoad = nullptr) const override;
 	COREUOBJECT_API virtual FString GetAuthoredNameForField(const FField* Field) const override;
-	// End of UObject interface.
+	//~ End of UObject interface.
 
-	// UScriptStruct interface.
+	//~ UScriptStruct interface.
 	COREUOBJECT_API virtual void InitializeStruct(void* Dest, int32 ArrayDim = 1) const override;
 	COREUOBJECT_API virtual uint32 GetStructTypeHash(const void* Src) const override;
 	COREUOBJECT_API virtual void RecursivelyPreload() override;
@@ -119,7 +139,13 @@ public:
 	COREUOBJECT_API virtual FString GetStructCPPName(uint32 CPPExportFlags) const override;
 	COREUOBJECT_API virtual FProperty* CustomFindProperty(const FName Name) const override;
 	COREUOBJECT_API virtual void PrepareCppStructOps() override;
-	// End of  UScriptStruct interface.
+	virtual bool ArePropertyGuidsAvailable() const override
+	{
+		return PropertyIDs.Num() > 0;
+	}
+	COREUOBJECT_API virtual FName FindPropertyNameFromGuid(const FGuid& PropertyID) const override;
+	COREUOBJECT_API virtual FGuid FindPropertyGuidFromName(const FName PropertyName) const override;
+	//~ End of  UScriptStruct interface.
 
 	/** Returns the raw memory of the default instance */
 	COREUOBJECT_API const uint8* GetDefaultInstance() const;
