RemoteObjects:  Actor should serialize its OwnedComponents when migrating.  This will ensure they're properly picked up as Subobjects (we do not GetObjectsWithOuter).  Warn against missing SubObjects when migrating-out.  Some minor fixes for the MultiServerProxy flow are also included.

#rnx

#Synced-CL 50269217

diff --git a/Engine/Source/Runtime/CoreUObject/Private/UObject/RemoteObjectSerialization.cpp b/Engine/Source/Runtime/CoreUObject/Private/UObject/RemoteObjectSerialization.cpp
--- a/Engine/Source/Runtime/CoreUObject/Private/UObject/RemoteObjectSerialization.cpp
+++ b/Engine/Source/Runtime/CoreUObject/Private/UObject/RemoteObjectSerialization.cpp
@@ -1216,6 +1216,17 @@ FRemoteObjectData SerializeObjectData(UObject* InObject, TSet<UObject*>& OutObje
 		SerializeObjectDataInternal(Ar, InObject, RequestedObjectId, OutObjects);
 	}
 	checkf(ObjectData.SerializedObjectHeaders.Num(), TEXT("Serialized object data for object %s is empty"), *RequestedObjectId.ToString());
+
+#if !(UE_BUILD_TEST || UE_BUILD_SHIPPING)
+	// We are doing an aggressive warning log here because these missing reference issues are hard to track down:
+	TArray<UObject*> AllSubobjects;
+	GetObjectsWithOuter(InObject, AllSubobjects);
+	AllSubobjects.RemoveAll([OutObjects](UObject* SubObj) { return OutObjects.Contains(SubObj); });
+
+	auto WrappedGetNameSafe = [](const UObject* O) { return FString::Printf(TEXT("%s (%s)"), *GetNameSafe(O), *(FObjectPtr{ const_cast<UObject*>(O) }.GetRemoteId().ToString())); };
+	UE_CLOG(!AllSubobjects.IsEmpty(), LogRemoteSerialization, Warning, TEXT("SerializeObjectData for %s did not serialize all of its SubObjects. Missed:\n  %s"), *GetNameSafe(InObject), *FString::JoinBy(AllSubobjects, TEXT(", "), WrappedGetNameSafe));
+#endif
+
 	return MoveTemp(ObjectData);
 }
 
diff --git a/Engine/Source/Runtime/Engine/Private/Actor.cpp b/Engine/Source/Runtime/Engine/Private/Actor.cpp
--- a/Engine/Source/Runtime/Engine/Private/Actor.cpp
+++ b/Engine/Source/Runtime/Engine/Private/Actor.cpp
@@ -1087,6 +1087,7 @@ void AActor::Serialize(FArchive& Ar)
 		Ar << CreationTime;
 		Ar << CurrentLifespan;
 		Ar << MigrationFlags;
+		Ar << OwnedComponents;
 
 		if (Ar.IsLoading())
 		{
