Prevent all external packages assets from being considered main package assets (was already the case for Actors but not entities). This prevents external entities from being gathered when doing a cooked editor build.

#Synced-CL 50438933

diff --git a/Engine/Source/Runtime/CoreUObject/Private/AssetRegistry/AssetData.cpp b/Engine/Source/Runtime/CoreUObject/Private/AssetRegistry/AssetData.cpp
--- a/Engine/Source/Runtime/CoreUObject/Private/AssetRegistry/AssetData.cpp
+++ b/Engine/Source/Runtime/CoreUObject/Private/AssetRegistry/AssetData.cpp
@@ -459,6 +459,13 @@ bool FAssetData::IsUAsset(UObject* InAsset)
 		return false;
 	}
 
+#if WITH_EDITORONLY_DATA
+	if (InAsset->IsPackageExternal())
+	{
+		return false;
+	}
+#endif //WITH_EDITORONLY_DATA
+
 	const UPackage* Package = InAsset->GetPackage();
 
 	TStringBuilder<FName::StringBufferSize> AssetNameStrBuilder;
diff --git a/Engine/Source/Runtime/CoreUObject/Public/AssetRegistry/AssetData.h b/Engine/Source/Runtime/CoreUObject/Public/AssetRegistry/AssetData.h
--- a/Engine/Source/Runtime/CoreUObject/Public/AssetRegistry/AssetData.h
+++ b/Engine/Source/Runtime/CoreUObject/Public/AssetRegistry/AssetData.h
@@ -322,6 +322,13 @@ PRAGMA_ENABLE_DEPRECATION_WARNINGS
 			return false;
 		}
 
+#if WITH_EDITORONLY_DATA
+		if (!OptionalOuterPath.IsNone())
+		{
+			return false;
+		}
+#endif //WITH_EDITORONLY_DATA
+
 		TStringBuilder<FName::StringBufferSize> AssetNameStrBuilder;
 		AssetName.ToString(AssetNameStrBuilder);
 
@@ -331,7 +338,7 @@ PRAGMA_ENABLE_DEPRECATION_WARNINGS
 	}
 
 	/** Convert to a SoftObjectPath. PackageName.AssetName for ToplevelAssets. OptionalOuterName and AssetName joined by : or . for sub-object assets */
-	COREUOBJECT_API FSoftObjectPath GetSoftObjectPath() const;
+	COREUOBJECT_API FSoftObjectPath GetSoftObjectPath() const; 
 	
 	// TODO: Deprecate in favor of GetSoftObjectPath
 	FSoftObjectPath ToSoftObjectPath() const
