[Feature] Placeholder Round Trip Tests
[BugFix] Fixed several bugs that occur when serializing placeholders that have ExperimentalOverridableLogic properties on them

This should theoretically make placeholders stable enough that they can be round tripped allowing us to remove the scary popup without risk of data loss

[BugFix] in InstanceDataObjectUtils.cpp we need to apply MarkPropertyAsLoose to container sub-properties
[BugFix] in FLinkerLoad.cpp if a placeholder name collided with a pre-existing struct name, we were putting the placeholder in the transient package instead. Unfortunately this made the placeholder unsavable so instead we're keeping it in the same package but giving it a unique name now.

#Synced-CL 50282601

#Synced-CL 50282601

diff --git a/Engine/Source/Runtime/CoreUObject/Internal/UObject/PropertyBagRepository.h b/Engine/Source/Runtime/CoreUObject/Internal/UObject/PropertyBagRepository.h
--- a/Engine/Source/Runtime/CoreUObject/Internal/UObject/PropertyBagRepository.h
+++ b/Engine/Source/Runtime/CoreUObject/Internal/UObject/PropertyBagRepository.h
@@ -150,6 +150,8 @@ public:
 		ReplaceDeadClassInstanceTypes,
 		/** Always replace missing type imports with a placeholder type on load, regardless of the underlying UClass derivative type */
 		DoNotRestrictMissingTypeImports,
+		/** Disable validation check that stops placeholders from being saved */
+		AllowSavingPlaceholderExports,
 	};
 
 	// query for whether or not the given struct/class is a placeholder type
diff --git a/Engine/Source/Runtime/CoreUObject/Private/UObject/InstanceDataObjectUtils.cpp b/Engine/Source/Runtime/CoreUObject/Private/UObject/InstanceDataObjectUtils.cpp
--- a/Engine/Source/Runtime/CoreUObject/Private/UObject/InstanceDataObjectUtils.cpp
+++ b/Engine/Source/Runtime/CoreUObject/Private/UObject/InstanceDataObjectUtils.cpp
@@ -1184,17 +1184,17 @@ namespace UE
 			if (PropertyFlags & CPF_ExperimentalOverridableLogic)
 			{
 				PropertyFlags &= ~CPF_ExperimentalOverridableLogic;
-				if (ensureMsgf(AsArrayProperty->Inner->IsA<FObjectProperty>(), TEXT("Expected array inner type to be an object property (%s: %s)"), *AsArrayProperty->GetPathName(), *AsArrayProperty->Inner->GetClass()->GetName()))
+				if (AsArrayProperty->Inner->IsA<FObjectProperty>())
 				{
 					PropertyFlags |= CPF_InstancedReference | CPF_PersistentInstance;
 				}
 			}
 
-			AsArrayProperty->Inner->SetPropertyFlags(CPF_Edit | CPF_EditConst | PropertyFlags);
+			MarkPropertyAsLoose(AsArrayProperty->Inner, PropertyFlags);
 		}
 		else if (const FSetProperty* AsSetProperty = CastField<FSetProperty>(Property))
 		{
-			AsSetProperty->ElementProp->SetPropertyFlags(CPF_Edit | CPF_EditConst);
+			MarkPropertyAsLoose(AsSetProperty->ElementProp);
 		}
 		else if (const FMapProperty* AsMapProperty = CastField<FMapProperty>(Property))
 		{
@@ -1202,18 +1202,18 @@ namespace UE
 			if (PropertyFlags & CPF_ExperimentalOverridableLogic)
 			{
 				PropertyFlags &= ~CPF_ExperimentalOverridableLogic;
-				if (ensureMsgf(AsMapProperty->KeyProp->IsA<FObjectProperty>(), TEXT("Expected map key type to be an object property (%s: %s)"), *AsMapProperty->GetPathName(), *AsMapProperty->KeyProp->GetClass()->GetName()))
+				if (AsMapProperty->KeyProp->IsA<FObjectProperty>())
 				{
 					PropertyFlags |= CPF_InstancedReference | CPF_PersistentInstance;
 				}
 			}
 
-			AsMapProperty->KeyProp->SetPropertyFlags(CPF_Edit | CPF_EditConst | PropertyFlags);
-			AsMapProperty->ValueProp->SetPropertyFlags(CPF_Edit | CPF_EditConst); // override serialization doesn't require any flags on the value property
+			MarkPropertyAsLoose(AsMapProperty->KeyProp, PropertyFlags);
+			MarkPropertyAsLoose(AsMapProperty->ValueProp);	// override serialization doesn't require any flags on the value property
 		}
 		else if (const FOptionalProperty* AsOptionalProperty = CastField<FOptionalProperty>(Property))
 		{
-			AsOptionalProperty->GetValueProperty()->SetPropertyFlags(CPF_Edit | CPF_EditConst);
+			MarkPropertyAsLoose(AsOptionalProperty->GetValueProperty());
 		}
 		else if (FObjectProperty* AsObjectProperty = CastField<FObjectProperty>(Property))
 		{
diff --git a/Engine/Source/Runtime/CoreUObject/Private/UObject/LinkerLoad.cpp b/Engine/Source/Runtime/CoreUObject/Private/UObject/LinkerLoad.cpp
--- a/Engine/Source/Runtime/CoreUObject/Private/UObject/LinkerLoad.cpp
+++ b/Engine/Source/Runtime/CoreUObject/Private/UObject/LinkerLoad.cpp
@@ -4700,13 +4700,11 @@ UClass* FLinkerLoad::TryCreatePlaceholderClassImport(int32 ImportIndex)
 					FName ClassObjectName = Import.ObjectName;
 					EObjectFlags ClassObjectFlags = RF_Public;
 
-					// If it's now a non-native struct type, use a transient placeholder class type so exports can still be serialized.
+					// If it's now a non-native struct type, use a placeholder class type with a unique name so exports can still be serialized.
 					const UScriptStruct* ExistingType = FindObjectFast<UScriptStruct>(ClassObjectPackage, Import.ObjectName, EFindObjectFlags::ExactClass);
 					if (ExistingType)
 					{
-						ClassObjectName = NAME_None;
-						ClassObjectFlags |= RF_Transient;
-						ClassObjectPackage = GetTransientPackage();
+						ClassObjectName = MakeUniqueObjectName(ClassObjectPackage, ImportClass, ClassObjectName);
 					}
 
 					UStruct* SuperStruct = FindKnownSuperStruct(ImportTypeHierarchy);
diff --git a/Engine/Source/Runtime/CoreUObject/Private/UObject/OverriddenPropertySet.cpp b/Engine/Source/Runtime/CoreUObject/Private/UObject/OverriddenPropertySet.cpp
--- a/Engine/Source/Runtime/CoreUObject/Private/UObject/OverriddenPropertySet.cpp
+++ b/Engine/Source/Runtime/CoreUObject/Private/UObject/OverriddenPropertySet.cpp
@@ -422,6 +422,36 @@ int32 FOverriddenPropertyNodeID::ToMapInternalIndex(FScriptMapHelper& MapHelper)
 	return INDEX_NONE;
 }
 
+FScopedPropertyValue FOverriddenPropertyNodeID::ToMapKey(FScriptMapHelper& MapHelper) const
+{
+	// Special case for object we didn't use the pointer to create the key
+	if (const FObjectPropertyBase* KeyObjectProperty = CastField<FObjectPropertyBase>(MapHelper.KeyProp))
+	{
+		for (FScriptMapHelper::FIterator It(MapHelper); It; ++It)
+		{
+			void* KeyPtr = MapHelper.GetKeyPtr(It);
+			if (UObject* CurrentObject = KeyObjectProperty->GetObjectPropertyValue(KeyPtr))
+			{
+				if ((*this) == FOverriddenPropertyNodeID(CurrentObject))
+				{
+					FScopedPropertyValue Result(MapHelper.KeyProp);
+					MapHelper.KeyProp->CopyCompleteValue(Result.GetData(), KeyPtr);
+					return MoveTemp(Result);
+				}
+			}
+		}
+	}
+	else
+	{
+		// Default case, just import the text as key value for comparison
+		FString KeyToFind(ToString());
+		FScopedPropertyValue Result(MapHelper.KeyProp);
+		MapHelper.KeyProp->ImportText_Direct(*KeyToFind, Result.GetData(), nullptr, PPF_None);
+		return MoveTemp(Result);
+	}
+	return FScopedPropertyValue();
+}
+
 bool FOverriddenPropertyNodeID::operator==(const FOverriddenPropertyNodeID& Other) const
 {
 	if (Path == Other.Path)
@@ -1956,3 +1986,57 @@ const FOverriddenPropertyNode* FOverriddenPropertySet::GetOverriddenPropertyNode
 
 	return OverriddenPropertyNode;
 }
+
+FScopedPropertyValue::FScopedPropertyValue(FProperty* InProperty)
+	: Property(InProperty)
+	, Data(FMemory::Malloc(InProperty->GetSize(), InProperty->GetMinAlignment()))
+{
+	Property->InitializeValue(Data);
+}
+
+FScopedPropertyValue::~FScopedPropertyValue()
+{
+	if (Data)
+	{
+		Property->DestroyAndFreeValue(Data);
+	}
+}
+
+FScopedPropertyValue::FScopedPropertyValue(FScopedPropertyValue&& Other)
+	: Property(Other.Property)
+	, Data(Other.Data)
+{
+	if (this != &Other)
+	{
+		Other.Data = nullptr;
+		Other.Property = nullptr;
+	}
+}
+
+FScopedPropertyValue& FScopedPropertyValue::operator=(FScopedPropertyValue&& Other)
+{
+	if (this != &Other)
+	{
+		if (Data)
+		{
+			Property->DestroyAndFreeValue(Data);
+		}
+
+		Data = Other.Data;
+		Property = Other.Property;
+
+		Other.Data = nullptr;
+		Other.Property = nullptr;
+	}
+	return *this;
+}
+
+FScopedPropertyValue::operator bool() const
+{
+	return !!Data;
+}
+
+void* FScopedPropertyValue::GetData() const
+{
+	return Data;
+}
diff --git a/Engine/Source/Runtime/CoreUObject/Private/UObject/PropertyBagRepository.cpp b/Engine/Source/Runtime/CoreUObject/Private/UObject/PropertyBagRepository.cpp
--- a/Engine/Source/Runtime/CoreUObject/Private/UObject/PropertyBagRepository.cpp
+++ b/Engine/Source/Runtime/CoreUObject/Private/UObject/PropertyBagRepository.cpp
@@ -1029,6 +1029,7 @@ namespace Private
 	DEFINE_IDO_PLACEHOLDER_FEATURE_FLAG_CVAR(ReplaceMissingReinstancedTypes);
 	DEFINE_IDO_PLACEHOLDER_FEATURE_FLAG_CVAR(ReplaceDeadClassInstanceTypes);
 	DEFINE_IDO_PLACEHOLDER_FEATURE_FLAG_CVAR(DoNotRestrictMissingTypeImports);
+	DEFINE_IDO_PLACEHOLDER_FEATURE_FLAG_CVAR(AllowSavingPlaceholderExports);
 #endif
 }
 
diff --git a/Engine/Source/Runtime/CoreUObject/Private/UObject/PropertyMap.cpp b/Engine/Source/Runtime/CoreUObject/Private/UObject/PropertyMap.cpp
--- a/Engine/Source/Runtime/CoreUObject/Private/UObject/PropertyMap.cpp
+++ b/Engine/Source/Runtime/CoreUObject/Private/UObject/PropertyMap.cpp
@@ -11,6 +11,7 @@
 #include "UObject/PropertyTag.h"
 #include "UObject/UnrealTypePrivate.h"
 #include "UObject/UObjectThreadContext.h"
+#include "UObject/CommonFNames.h"
 
 namespace UEMapProperty_Private
 {
@@ -366,7 +367,17 @@ void FMapProperty::SerializeItem(FStructuredArchive::FSlot Slot, TNotNull<void*>
 								KeyProp->SerializeItem(RemovedArray.EnterElement().EnterRecord().EnterField(TEXT("Key")), TempKeyValueStorage);
 							}
 
-							if (MapHelper.RemovePair(TempKeyValueStorage))
+							bool bPersistOverride = MapHelper.RemovePair(TempKeyValueStorage);
+
+#if WITH_EDITORONLY_DATA
+							// remove operations for unknown properties should be persisted regardless of whether they can be found in the map because it's possible
+							// the lack of CDO constructor is causing an element to be missing that otherwise would've been present
+							if (!bPersistOverride)
+							{
+								bPersistOverride |= this->GetBoolMetaData(UE::NAME_IsLooseMetadata);
+							}
+#endif
+							if (bPersistOverride)
 							{
 								// Need to fetch the MapOverriddenPropertyNode every loop as the previous might have reallocated the node.
 								if (FOverriddenPropertyNode* MapOverriddenPropertyNode = OverriddenProperties ? OverriddenProperties->RestoreOverriddenPropertyOperation(EOverriddenPropertyOperation::Modified, UnderlyingArchive.GetSerializedPropertyChain(), /*Property*/nullptr) : nullptr)
@@ -510,7 +521,7 @@ void FMapProperty::SerializeItem(FStructuredArchive::FSlot Slot, TNotNull<void*>
 			else
 			{
 				// Container for temporarily tracking some indices
-				TArray<int32> RemovedIndices;
+				TArray<FScopedPropertyValue> RemovedKeys;
 				TArray<int32> AddedIndices;
 				TSet<int32> ModifiedIndices;
 				TSet<int32> ShadowedIndices;
@@ -561,10 +572,9 @@ void FMapProperty::SerializeItem(FStructuredArchive::FSlot Slot, TNotNull<void*>
 									{
 										case EOverriddenPropertyOperation::Remove:
 										{
-											const int32 InternalIndex = SubNode.GetNodeID().ToMapInternalIndex(DefaultsMapHelper);
-											if (InternalIndex != INDEX_NONE)
+											if (FScopedPropertyValue Found = SubNode.GetNodeID().ToMapKey(DefaultsMapHelper))
 											{
-												RemovedIndices.Add(InternalIndex);
+												RemovedKeys.Emplace(MoveTemp(Found));
 											}
 											break;
 										}
@@ -627,14 +637,14 @@ void FMapProperty::SerializeItem(FStructuredArchive::FSlot Slot, TNotNull<void*>
 					checkf(Defaults, TEXT("Expecting overridable serialization to have defaults to compare to"));
 					FScriptMapHelper DefaultsMapHelper(this, Defaults);
 
-					int32 NumRemoved = bReplaceMap ? INDEX_NONE : RemovedIndices.Num();
+					int32 NumRemoved = bReplaceMap ? INDEX_NONE : RemovedKeys.Num();
 					FStructuredArchive::FArray RemovedArray = Record.EnterArray(TEXT("Removed"), NumRemoved);
-					for (int32 i : RemovedIndices)
+					for (const FScopedPropertyValue& Key : RemovedKeys)
 					{
 						FStructuredArchive::FRecord EntryRecord = RemovedArray.EnterElement().EnterRecord();
 						UE::FSerializedPropertyPathScope SerializedPropertyPath(Context, {UE::NAME_Key});
 						FSerializedPropertyScope SerializedProperty(UnderlyingArchive, KeyProp, this);
-						KeyProp->SerializeItem(EntryRecord.EnterField(TEXT("Key")), DefaultsMapHelper.GetKeyPtr(i));
+						KeyProp->SerializeItem(EntryRecord.EnterField(TEXT("Key")), Key.GetData());
 					}
 
 					int32 NumModified = ModifiedIndices.Num();
@@ -680,7 +690,7 @@ void FMapProperty::SerializeItem(FStructuredArchive::FSlot Slot, TNotNull<void*>
 		}
 
 		const bool bIsLoadingToUnknownPropertyTree =
-#if WITH_EDITOR
+#if WITH_EDITORONLY_DATA
 			FUObjectThreadContext::Get().GetSerializeContext()->bIsLoadingToUnknownPropertyTree;
 #else
 			false;
diff --git a/Engine/Source/Runtime/CoreUObject/Private/UObject/SavePackage2.cpp b/Engine/Source/Runtime/CoreUObject/Private/UObject/SavePackage2.cpp
--- a/Engine/Source/Runtime/CoreUObject/Private/UObject/SavePackage2.cpp
+++ b/Engine/Source/Runtime/CoreUObject/Private/UObject/SavePackage2.cpp
@@ -714,46 +714,49 @@ ESavePackageResult ValidateExports(FSaveContext& SaveContext)
 #if WITH_EDITOR
 	if (UE::FPropertyBagRepository::IsPropertyBagPlaceholderObjectSupportEnabled())
 	{
-		TArray<const UObject*, TInlineAllocator<16>> PlaceholderExports;
-		const TSet<FTaggedExport>& AllExports = SaveContext.GetExports();
-		for (const FTaggedExport& Export : AllExports)
+		if (!UE::FPropertyBagRepository::IsPropertyBagPlaceholderObjectFeatureEnabled(UE::FPropertyBagRepository::EPlaceholderObjectFeature::AllowSavingPlaceholderExports))
 		{
-			if (UE::FPropertyBagRepository::IsPropertyBagPlaceholderObject(Export.Obj))
+			TArray<const UObject*, TInlineAllocator<16>> PlaceholderExports;
+			const TSet<FTaggedExport>& AllExports = SaveContext.GetExports();
+			for (const FTaggedExport& Export : AllExports)
 			{
-				PlaceholderExports.Add(Export.Obj);
-			}
-			else if (const UStruct* TypeExport = Cast<UStruct>(Export.Obj))
-			{
-				if (UE::FPropertyBagRepository::IsPropertyBagPlaceholderType(TypeExport))
+				if (UE::FPropertyBagRepository::IsPropertyBagPlaceholderObject(Export.Obj))
 				{
-					PlaceholderExports.Add(TypeExport);
+					PlaceholderExports.Add(Export.Obj);
+				}
+				else if (const UStruct* TypeExport = Cast<UStruct>(Export.Obj))
+				{
+					if (UE::FPropertyBagRepository::IsPropertyBagPlaceholderType(TypeExport))
+					{
+						PlaceholderExports.Add(TypeExport);
+					}
 				}
 			}
-		}
 
-		if (PlaceholderExports.Num() == AllExports.Num())
-		{
-			// If all exports are placeholders, allow this to pass as editor-only data (e.g. a placeholder type package).
-			return ESavePackageResult::ContainsEditorOnlyData;
-		}
-		else if (PlaceholderExports.Num() > 0)
-		{
-			for (const UObject* PlaceholderExport : PlaceholderExports)
+			if (PlaceholderExports.Num() == AllExports.Num())
 			{
-				check(PlaceholderExport);
-				UObject* Outer = PlaceholderExport->GetOuter();
-				FString ErrorMessage = FString::Printf(TEXT("Saving illegal placeholder object in package '%s'. Object Name '%s', Class '%s', Outer '%s', Outer Class '%s', Full Path '%s'. You might have unresolved classes or compilation failures to fix."), *SaveContext.GetPackage()->GetName(), *PlaceholderExport->GetName(), *PlaceholderExport->GetClass()->GetName(), Outer ? *Outer->GetName() : TEXT("None"), Outer ? *Outer->GetClass()->GetName() : TEXT("None"), *PlaceholderExport->GetFullName());
-				if (SaveContext.IsGenerateSaveError())
-				{
-					SaveContext.GetError()->Logf(ELogVerbosity::Warning, TEXT("%s"), *ErrorMessage);
-				}
-				else
+				// If all exports are placeholders, allow this to pass as editor-only data (e.g. a placeholder type package).
+				return ESavePackageResult::ContainsEditorOnlyData;
+			}
+			else if (PlaceholderExports.Num() > 0)
+			{
+				for (const UObject* PlaceholderExport : PlaceholderExports)
 				{
-					UE_LOG(LogSavePackage, Error, TEXT("%s"), *ErrorMessage);
+					check(PlaceholderExport);
+					UObject* Outer = PlaceholderExport->GetOuter();
+					FString ErrorMessage = FString::Printf(TEXT("Saving illegal placeholder object in package '%s'. Object Name '%s', Class '%s', Outer '%s', Outer Class '%s', Full Path '%s'. You might have unresolved classes or compilation failures to fix."), *SaveContext.GetPackage()->GetName(), *PlaceholderExport->GetName(), *PlaceholderExport->GetClass()->GetName(), Outer ? *Outer->GetName() : TEXT("None"), Outer ? *Outer->GetClass()->GetName() : TEXT("None"), *PlaceholderExport->GetFullName());
+					if (SaveContext.IsGenerateSaveError())
+					{
+						SaveContext.GetError()->Logf(ELogVerbosity::Warning, TEXT("%s"), *ErrorMessage);
+					}
+					else
+					{
+						UE_LOG(LogSavePackage, Error, TEXT("%s"), *ErrorMessage);
+					}
 				}
-			}
 
-			return ESavePackageResult::Error;
+				return ESavePackageResult::Error;
+			}
 		}
 	}
 #endif
diff --git a/Engine/Source/Runtime/CoreUObject/Public/UObject/OverriddenPropertySet.h b/Engine/Source/Runtime/CoreUObject/Public/UObject/OverriddenPropertySet.h
--- a/Engine/Source/Runtime/CoreUObject/Public/UObject/OverriddenPropertySet.h
+++ b/Engine/Source/Runtime/CoreUObject/Public/UObject/OverriddenPropertySet.h
@@ -254,6 +254,24 @@ private:
 	}
 };
 
+// helper struct that allocates space to temporarily store/own the value of a property and free the memory when it destructs
+struct FScopedPropertyValue
+{
+	explicit FScopedPropertyValue() = default;
+	explicit FScopedPropertyValue(FProperty* Property);
+	~FScopedPropertyValue();
+	FScopedPropertyValue(FScopedPropertyValue&& Other);
+	FScopedPropertyValue& operator=(FScopedPropertyValue&& Other);
+	FScopedPropertyValue(const FScopedPropertyValue& Other) = delete;
+	FScopedPropertyValue& operator=(const FScopedPropertyValue& Other) = delete;
+	operator bool() const;
+	void* GetData() const;
+
+private:
+	FProperty* Property = nullptr;
+	void* Data = nullptr;
+};
+
 USTRUCT()
 struct FOverriddenPropertyNodeID
 {
@@ -274,6 +292,7 @@ struct FOverriddenPropertyNodeID
 	// Overridden property node map helpers
 	static COREUOBJECT_API FOverriddenPropertyNodeID FromMapKey(const FProperty* KeyProperty, const void* KeyData);
 	COREUOBJECT_API int32 ToMapInternalIndex(FScriptMapHelper& MapHelper) const;
+	COREUOBJECT_API FScopedPropertyValue ToMapKey(FScriptMapHelper& MapHelper) const;
 
 	COREUOBJECT_API bool operator==(const FOverriddenPropertyNodeID& Other) const;
 
