[UE] Disable IDOs for standalone non-editor and non-cook commandlet runs.

#Synced-CL 50300730

diff --git a/Engine/Source/Runtime/CoreUObject/Private/UObject/InstanceDataObjectUtils.cpp b/Engine/Source/Runtime/CoreUObject/Private/UObject/InstanceDataObjectUtils.cpp
--- a/Engine/Source/Runtime/CoreUObject/Private/UObject/InstanceDataObjectUtils.cpp
+++ b/Engine/Source/Runtime/CoreUObject/Private/UObject/InstanceDataObjectUtils.cpp
@@ -638,6 +638,19 @@ namespace UE
 
 	bool IsInstanceDataObjectSupportEnabled()
 	{
+#if UE_EDITOR
+		// Keep IDOs disabled during non-editor runs
+		if (!GIsEditor)
+		{
+			return false;
+		}
+
+		// Keep IDOs disabled during non-cook commandlet runs
+		if (IsRunningCommandlet() && !IsRunningCookCommandlet())
+		{
+			return false;
+		}
+#endif
 		return bEnableIDOSupport;
 	}
 
