Sort serialized remote object headers by their outer to prevent out-of-order (inner before the outer) remote object creation during deserialization

#Synced-CL 50433781

diff --git a/Engine/Source/Runtime/CoreUObject/Private/UObject/RemoteObjectSerialization.cpp b/Engine/Source/Runtime/CoreUObject/Private/UObject/RemoteObjectSerialization.cpp
--- a/Engine/Source/Runtime/CoreUObject/Private/UObject/RemoteObjectSerialization.cpp
+++ b/Engine/Source/Runtime/CoreUObject/Private/UObject/RemoteObjectSerialization.cpp
@@ -453,14 +453,21 @@ public:
 		OutHeader.Archetype = Object->GetArchetype();
 		OutHeader.InternalFlags = (int32)(Object->GetInternalFlags() & EInternalObjectFlags::Garbage);
 
-		FPackedRemoteObjectHeader& PackedHeader = ObjectData.SerializedObjectHeaders.Emplace_GetRef();
-		PackedHeader.NameIndex = OutHeader.Name;
-		PackedHeader.IdIndex = OutHeader.RemoteId;
-		PackedHeader.ClassIndex = AddPathNameToNameMap(OutHeader.Class);
-
 		return true;
 	}
 
+	void CreatePackedObjectHeaders(const TArray<FRemoteObjectHeader>& InHeaders)
+	{
+		checkf(ObjectData.SerializedObjectHeaders.Num() == 0, TEXT("Packed object headers already created"));
+
+		for (const FRemoteObjectHeader& Header : InHeaders)
+		{
+			FPackedRemoteObjectHeader& PackedHeader = ObjectData.SerializedObjectHeaders.Emplace_GetRef();
+			PackedHeader.NameIndex = Header.Name;
+			PackedHeader.IdIndex = Header.RemoteId;
+			PackedHeader.ClassIndex = AddPathNameToNameMap(Header.Class);
+		}
+	}
 protected:
 
 	virtual FRemoteObjectReferenceInfo GetReferenceInfo(const FObjectPtr& ObjPtr) const
@@ -648,6 +655,11 @@ public:
 	{		
 		Object->Serialize(*this);
 	}
+
+	FRemoteObjectId GetRemoteObjectId(FNameIndexType IdIndex) const
+	{
+		return ObjectData.Tables.RemoteIds[IdIndex];
+	}
 };
 
 template <typename T>
@@ -763,6 +775,7 @@ FArchive& operator<<(FArchive& Ar, FRemoteObjectHeader& Header)
 	Ar << Header.Outer;
 	Ar << Header.Archetype;
 	Ar << Header.InternalFlags;
+	Ar << Header.StartOffset;
 
 	return Ar;
 }
@@ -1127,7 +1140,74 @@ void ResolvePathNames(const FRemoteObjectData& InObjectData, TArray<UObject*>& O
 	}
 }
 
-void SerializeObjectDataInternal(FArchiveRemoteObjectWriter& Ar, UObject* RequestedObject, FRemoteObjectId RequestedObjectId, TSet<UObject*>& OutObjects)
+FORCENOINLINE void SortRemoteObjectHeadersByOuter(FArchiveRemoteObjectWriter& Ar, TArray<FRemoteObjectHeader>& ObjectHeaders)
+{
+	using namespace UE::CoreUObject::Private;
+
+	struct FOuterSortPredicate
+	{
+		TMap<UObject*, int32> OuterLevelMap;
+
+		bool operator()(const FRemoteObjectHeader& A, const FRemoteObjectHeader& B) const
+		{
+			int32 OuterALevel = OuterLevelMap.FindChecked(A.Outer);
+			int32 OuterBLevel = OuterLevelMap.FindChecked(B.Outer);
+			return OuterALevel > OuterBLevel;
+		}
+	} Predicate;
+
+	// Assign a level index to each of the Serialized Objects' Outers. Indices grow from the Innermost to the Outermost Outer
+	// e.g.                     World (3)
+	//                             |
+	//                      PersistentLevel (2)
+	//                             |
+	//                          Actor (1)
+	//                         /         \
+	//                ComponentA (0)    ComponentB (0)
+	//                    |                  |
+	//                SubobjectA        SubobjectB (leaf = no index)
+
+	for (FRemoteObjectHeader& Header : ObjectHeaders)
+	{
+		int32 CurrentObjectOuterLevel = 0;
+		for (UObject* Outer = NoResolveObjectHandleNoRead(Header.Outer.GetHandle()); Outer; Outer = FObjectHandleUtils::GetNonAccessTrackedOuterNoResolve(Outer))
+		{
+			int32& OuterIndex = Predicate.OuterLevelMap.FindOrAdd(Outer, CurrentObjectOuterLevel);
+			if (OuterIndex < CurrentObjectOuterLevel)
+			{
+				OuterIndex = CurrentObjectOuterLevel;
+			}
+			CurrentObjectOuterLevel++;
+		}
+	}
+
+	// Sort object headers from the outermost to the innermost object
+	if (Predicate.OuterLevelMap.Num() > 1)
+	{
+		ObjectHeaders.Sort<FOuterSortPredicate>(Predicate);
+	}
+
+	// Make sure the root object header is still the first header in the array
+	checkf(FRemoteObjectId(Ar.GetRootObject()) == Ar.GetRemoteObjectId(ObjectHeaders[0].RemoteId), TEXT("Root object id (%s) does not match the first object header id (%s)."),
+		*FRemoteObjectId(Ar.GetRootObject()).ToString(), *Ar.GetRemoteObjectId(ObjectHeaders[0].RemoteId).ToString());
+
+#if DO_CHECK && UE_WITH_REMOTE_OBJECT_HANDLE
+	// Make sure the header array is sorted
+	for (int32 HeaderIndex = 0; HeaderIndex < ObjectHeaders.Num(); ++HeaderIndex)
+	{
+		FRemoteObjectHeader& Header = ObjectHeaders[HeaderIndex];
+		FRemoteObjectId OuterId = Header.Outer.GetRemoteId();
+		int32 OuterHeaderIndex = ObjectHeaders.IndexOfByPredicate([OuterId, &Ar](const FRemoteObjectHeader& MaybeOuterHeader)
+			{
+				return Ar.GetRemoteObjectId(MaybeOuterHeader.RemoteId) == OuterId;
+			});
+		checkf(OuterHeaderIndex < HeaderIndex, TEXT("Serialized object's[%d] %s Outer[%d] %s is serialized after the object"),
+			HeaderIndex, *Ar.GetRemoteObjectId(Header.RemoteId).ToString(ERemoteIdToStringVerbosity::PathName), OuterHeaderIndex, *OuterId.ToString(ERemoteIdToStringVerbosity::PathName));
+	}
+#endif // DO_CHECK
+}
+
+void SerializeObjectDataInternal(FArchiveRemoteObjectWriter& Ar, UObject* RequestedObject, FRemoteObjectId RequestedObjectId, TSet<UObject*>& OutObjects, ERemoteObjectSerializationFlags SerializationFlags)
 {
 	using namespace UE::RemoteObject::Handle;
 
@@ -1161,12 +1241,12 @@ void SerializeObjectDataInternal(FArchiveRemoteObjectWriter& Ar, UObject* Reques
 				FRemoteObjectHeader Header;
 				if (Ar.PopulateObjectHeader(ObjectToSerialize, Header))
 				{					
-					ObjectHeaders.Add(Header);
+					FRemoteObjectHeader& SerializedHeader = ObjectHeaders.Emplace_GetRef(Header);
 					OutObjects.Add(ObjectToSerialize);
 
-					Ar.SerializeRemoteObject(ObjectToSerialize, ObjectHeaders.Last());
+					SerializedHeader.StartOffset = Ar.Tell();
 
-					ObjectHeaders.Last().NextOffset = Ar.Tell();
+					Ar.SerializeRemoteObject(ObjectToSerialize, SerializedHeader);
 				}
 				else
 				{
@@ -1193,6 +1273,15 @@ void SerializeObjectDataInternal(FArchiveRemoteObjectWriter& Ar, UObject* Reques
 	Ar << HeaderOffset;
 	Ar.Seek(HeaderOffset);
 
+	// Sort objects from the outermost to the innermost. 
+	// Skip sorting if there's just one object or we're resetting a borrowed object because CDO delta serialization substitutes serialized objects (archetypes) with their instances and the outers don't match
+	if (OutObjects.Num() > 1 && !EnumHasAnyFlags(SerializationFlags, ERemoteObjectSerializationFlags::Resetting))
+	{
+		SortRemoteObjectHeadersByOuter(Ar, ObjectHeaders);
+	}
+
+	Ar.CreatePackedObjectHeaders(ObjectHeaders);
+
 	{
 		FRemoteObjectWriterScope Scope(Ar, TEXT("ObjectHeaders"));
 		Ar << ObjectHeaders;
@@ -1213,7 +1302,7 @@ FRemoteObjectData SerializeObjectData(UObject* InObject, TSet<UObject*>& OutObje
 	{
 		FArchiveRemoteObjectWriter Ar(Object, ObjectData, MigrationContext, nullptr, &OutReferencedObjects);
 		Ar.SetMigratingRemoteObjects(true);
-		SerializeObjectDataInternal(Ar, InObject, RequestedObjectId, OutObjects);
+		SerializeObjectDataInternal(Ar, InObject, RequestedObjectId, OutObjects, ERemoteObjectSerializationFlags::None);
 	}
 	checkf(ObjectData.SerializedObjectHeaders.Num(), TEXT("Serialized object data for object %s is empty"), *RequestedObjectId.ToString());
 
@@ -1267,7 +1356,7 @@ void ResetRemoteObject(UObject* InObject)
 		UObject* Archetype = const_cast<UObject*>(FindArchetype(Object));
 		FArchetypeDeltaWriter Ar(Archetype, *ArchetypeDelta, ReverseArchetypeToObjectMap);
 		TSet<UObject*> ArchetypeSerializedObjects;
-		SerializeObjectDataInternal(Ar, Archetype, RequestedObjectId, ArchetypeSerializedObjects);
+		SerializeObjectDataInternal(Ar, Archetype, RequestedObjectId, ArchetypeSerializedObjects, ERemoteObjectSerializationFlags::Resetting);
 	}
 
 	{
@@ -1421,16 +1510,13 @@ int32 DeserializeObjectData(FRemoteObjectData& ObjectData,
 				if (UObject* Object = OutObjects[ObjectIndex])
 				{
 					FRemoteObjectReaderScope Scope(Ar, Object);
+					Ar.Seek(ObjectHeaders[ObjectIndex].StartOffset);
 					Object->Serialize(Ar);
 					if (RequestedObjectIndex == -1 && RequestedObjectId == FRemoteObjectId(Object))
 					{
 						RequestedObjectIndex = ObjectIndex;
 					}
 				}
-				else
-				{
-					Ar.Seek(ObjectHeaders[ObjectIndex].NextOffset);
-				}
 			}
 		}
 
diff --git a/Engine/Source/Runtime/CoreUObject/Private/UObject/RemoteObjectSerialization.h b/Engine/Source/Runtime/CoreUObject/Private/UObject/RemoteObjectSerialization.h
--- a/Engine/Source/Runtime/CoreUObject/Private/UObject/RemoteObjectSerialization.h
+++ b/Engine/Source/Runtime/CoreUObject/Private/UObject/RemoteObjectSerialization.h
@@ -32,7 +32,7 @@ namespace UE::RemoteObject::Serialization
 		TObjectPtr<UObject> Outer;
 		TObjectPtr<UObject> Archetype;
 		int32 InternalFlags = 0;
-		int64 NextOffset = 0;
+		int64 StartOffset = 0;
 
 		/** Transient (SerialNumber is local-only but it's stored here for convenience when deserializing object data) */
 		int32 SerialNumber = 0;
