Added TemplateStructName metadata field to be able to obtain the template struct name from the IDO class

#Synced-CL 50485455
#Synced-CL 50485455

diff --git a/Engine/Source/Runtime/CoreUObject/Private/UObject/InstanceDataObjectUtils.cpp b/Engine/Source/Runtime/CoreUObject/Private/UObject/InstanceDataObjectUtils.cpp
--- a/Engine/Source/Runtime/CoreUObject/Private/UObject/InstanceDataObjectUtils.cpp
+++ b/Engine/Source/Runtime/CoreUObject/Private/UObject/InstanceDataObjectUtils.cpp
@@ -278,6 +278,7 @@ namespace UE
 	static const FName NAME_Component(ANSITEXTVIEW("component"));
 	static const FName NAME_IDOMapKey(ANSITEXTVIEW("Key"));
 	static const FName NAME_IDOMapValue(ANSITEXTVIEW("Value"));
+	static const FName NAME_TemplateStructName(ANSITEXTVIEW("TemplateStructName"));
 
 	template <typename T>
 	struct TCacheItem
@@ -1303,6 +1304,7 @@ namespace UE
 			UStruct* Result = NewObject<UStruct>(Outer, SuperStructClass, MakeUniqueObjectName(nullptr, SuperStructClass, SuperName));
 			Result->SetSuperStruct(Super);
 			Result->SetMetaData(NAME_IsInstanceDataObjectStruct, TEXT("True"));
+			Result->SetMetaData(NAME_TemplateStructName, *TemplateStruct->GetName());
 
 			// Gather properties
 			TArray<FProperty*> SuperProperties;
@@ -1359,6 +1361,10 @@ namespace UE
 		UStruct* Result = NewObject<UStruct>(Outer, StructClass, MakeUniqueObjectName(Outer, StructClass, InstanceDataObjectName));
 		Result->SetSuperStruct(Super);
 		Result->SetMetaData(NAME_IsInstanceDataObjectStruct, TEXT("True"));
+		if (TemplateStruct)
+		{
+			Result->SetMetaData(NAME_TemplateStructName, *TemplateStruct->GetName());
+		}
 
 		// inherit ContainsLooseProperties metadata
 		if (Super && Super->GetBoolMetaData(NAME_ContainsLoosePropertiesMetadata))
