Introduce a new object flag RF_GeneratePublicHash to explicitly control public hash generation during cooking, decoupling it from the RF_Public flag.
This allows objects to have public hashes generated during cooking without requiring the RF_Public flag, providing more granular control over the cooking process, as we allow referencing private objects when sharing an outer.

#Synced-CL 50504635

diff --git a/Engine/Source/Developer/IoStoreUtilities/Private/PackageStoreOptimizer.cpp b/Engine/Source/Developer/IoStoreUtilities/Private/PackageStoreOptimizer.cpp
--- a/Engine/Source/Developer/IoStoreUtilities/Private/PackageStoreOptimizer.cpp
+++ b/Engine/Source/Developer/IoStoreUtilities/Private/PackageStoreOptimizer.cpp
@@ -658,7 +658,7 @@ void FPackageStoreOptimizer::ProcessExports(const FCookedHeaderData& CookedHeade
 		Export.SerialSize = ObjectExport.SerialSize;
 		Export.bNotForClient = ObjectExport.bNotForClient;
 		Export.bNotForServer = ObjectExport.bNotForServer;
-		Export.bIsPublic = (Export.ObjectFlags & RF_Public) > 0 || ObjectExport.bGeneratePublicHash;
+		Export.bIsPublic = (Export.ObjectFlags & (RF_Public | RF_GeneratePublicHash)) > 0 || ObjectExport.bGeneratePublicHash;
 		ResolveExport(UnresolvedImports.GetData(), UnresolvedExports.GetData(), CookedHeaderData.ObjectImports.GetData(), CookedHeaderData.ObjectExports.GetData(), Package->Name, ExportIndex);
 		if (Export.bIsPublic)
 		{
diff --git a/Engine/Source/Runtime/CoreUObject/Public/UObject/ObjectMacros.h b/Engine/Source/Runtime/CoreUObject/Public/UObject/ObjectMacros.h
--- a/Engine/Source/Runtime/CoreUObject/Public/UObject/ObjectMacros.h
+++ b/Engine/Source/Runtime/CoreUObject/Public/UObject/ObjectMacros.h
@@ -602,6 +602,7 @@ enum EObjectFlags
 	RF_NeedInitialization		=0x00000200,	///< This object has not completed its initialization process. Cleared when ~FObjectInitializer completes
 	RF_NeedLoad					=0x00000400,	///< During load, indicates object needs loading.
 	RF_KeepForCooker UE_DEPRECATED(5.6, "This flag is now obsolete as the functionality is unnecessary.")			=0x00000800,	///< Keep this object during garbage collection because it's still being used by the cooker
+	RF_GeneratePublicHash		=0x00000800,	///< During cooking, generate a public hash for this object even is not marked RF_Public.
 	RF_NeedPostLoad				=0x00001000,	///< Object needs to be postloaded.
 	RF_NeedPostLoadSubobjects	=0x00002000,	///< During load, indicates that the object still needs to instance subobjects and fixup serialized component references
 	RF_NewerVersionExists		=0x00004000,	///< Object has been consigned to oblivion due to its owner package being reloaded, and a newer version currently exists
@@ -647,7 +648,7 @@ inline constexpr EObjectFlags RF_HasPlaceholderType = EObjectFlags(0x20000000);
 #define RF_AllFlags				(EObjectFlags)0xffffffff	///< All flags, used mainly for error checking
 
 /** Flags to load from unreal asset files */
-#define RF_Load						((EObjectFlags)(RF_Public | RF_Standalone | RF_Transactional | RF_ClassDefaultObject | RF_ArchetypeObject | RF_DefaultSubObject | RF_TextExportTransient | RF_InheritableComponentTemplate | RF_DuplicateTransient | RF_NonPIEDuplicateTransient | RF_MigratingAsset | RF_HasDynamicImports)) 
+#define RF_Load						((EObjectFlags)(RF_Public | RF_Standalone | RF_Transactional | RF_ClassDefaultObject | RF_ArchetypeObject | RF_GeneratePublicHash | RF_DefaultSubObject | RF_TextExportTransient | RF_InheritableComponentTemplate | RF_DuplicateTransient | RF_NonPIEDuplicateTransient | RF_MigratingAsset | RF_HasDynamicImports)) 
 
 /** Sub-objects will inherit these flags from their SuperObject */
 #define RF_PropagateToSubObjects	((EObjectFlags)(RF_Public | RF_ArchetypeObject | RF_Transactional | RF_Transient))
