[UE] Minor fix for import type placeholder logic - check for derivatives of UScriptStruct that may now be occupying the missing import class path, not just the base type.

#rnx

#Synced-CL 50300730

diff --git a/Engine/Source/Runtime/CoreUObject/Private/UObject/LinkerLoad.cpp b/Engine/Source/Runtime/CoreUObject/Private/UObject/LinkerLoad.cpp
--- a/Engine/Source/Runtime/CoreUObject/Private/UObject/LinkerLoad.cpp
+++ b/Engine/Source/Runtime/CoreUObject/Private/UObject/LinkerLoad.cpp
@@ -4701,7 +4701,7 @@ UClass* FLinkerLoad::TryCreatePlaceholderClassImport(int32 ImportIndex)
 					EObjectFlags ClassObjectFlags = RF_Public;
 
 					// If it's now a non-native struct type, use a placeholder class type with a unique name so exports can still be serialized.
-					const UScriptStruct* ExistingType = FindObjectFast<UScriptStruct>(ClassObjectPackage, Import.ObjectName, EFindObjectFlags::ExactClass);
+					const UScriptStruct* ExistingType = FindObjectFast<UScriptStruct>(ClassObjectPackage, Import.ObjectName, EFindObjectFlags::None);
 					if (ExistingType)
 					{
 						ClassObjectName = MakeUniqueObjectName(ClassObjectPackage, ImportClass, ClassObjectName);
