RemoteObjects:  DebugFindRemoteObjectStub version that takes in RemoteObjectIds (in the Debugger Watch Window)

#rnx

#Synced-CL 50444249

diff --git a/Engine/Source/Runtime/CoreUObject/Private/UObject/RemoteObject.cpp b/Engine/Source/Runtime/CoreUObject/Private/UObject/RemoteObject.cpp
--- a/Engine/Source/Runtime/CoreUObject/Private/UObject/RemoteObject.cpp
+++ b/Engine/Source/Runtime/CoreUObject/Private/UObject/RemoteObject.cpp
@@ -902,11 +902,6 @@ UClass* GetClass(FRemoteObjectId ObjectId, ERemoteObjectGetClassBehavior GetClas
 	return nullptr;
 }
 
-FRemoteObjectStub* DebugFindRemoteObjectStub(FRemoteObjectId ObjectId)
-{
-	return UE::RemoteObject::Private::FindRemoteObjectStub(ObjectId);
-}
-
 bool IsTransactionallyPostMigratingObjects()
 {
 	return UE::RemoteObject::Private::GIsTransactionallyPostMigratingObjects;
@@ -1519,6 +1514,20 @@ COREUOBJECT_API UE::RemoteObject::Handle::FRemoteObjectStub* DebugFindRemoteObje
 	return UE::RemoteObject::Private::ObjectMaps->FindStub(ObjId);
 }
 
+/**
+ * Put this in a Debug Watch Window to get the FRemoteObjectStub of a given ObjectId.
+ * e.g. DebugFindRemoteObjectStub(2,123)
+ */
+COREUOBJECT_API UE::RemoteObject::Handle::FRemoteObjectStub* DebugFindRemoteObjectStub(uint32 ServerId, uint64 SerialNumber)
+{
+	// Sanitize inputs
+	ServerId = (ServerId <= MAX_REMOTE_OBJECT_SERVER_ID) ? ServerId : 0;
+	SerialNumber = (SerialNumber <= MAX_REMOTE_OBJECT_SERIAL_NUMBER) ? SerialNumber : 0;
+
+	const FRemoteObjectId ObjId { FRemoteServerId::FromIdNumber(ServerId), SerialNumber };
+	return UE::RemoteObject::Private::ObjectMaps->FindStub(ObjId);
+}
+
 /**
  * Attempt to find a UObject in the currently debugged process by its FRemoteObjectId constituents.
  * Once you know a FRemoteObjectId, take its ServerId and SerialNumber and pass them into Debug Watch Window as arguments (in that order)
