Fixed missing include in StructuredArchiveTest.cpp

#rnx
#Synced-CL 50253651

diff --git a/Engine/Source/Runtime/CoreUObject/Private/Serialization/StructuredArchiveTest.cpp b/Engine/Source/Runtime/CoreUObject/Private/Serialization/StructuredArchiveTest.cpp
--- a/Engine/Source/Runtime/CoreUObject/Private/Serialization/StructuredArchiveTest.cpp
+++ b/Engine/Source/Runtime/CoreUObject/Private/Serialization/StructuredArchiveTest.cpp
@@ -1,5 +1,7 @@
 // Copyright Epic Games, Inc. All Rights Reserved.
 
+#include "CoreTypes.h"
+
 #if WITH_LOW_LEVEL_TESTS && WITH_TEXT_ARCHIVE_SUPPORT
 
 #include "Containers/StringView.h"
