Avoid stub lookups when checking remote object residence especially if we already have the UObject pointer (which is the majority of cases when we access TObjectPtrs)
Added SetResidence(Object, Residence, ResidentServerId) function and deprecated MarkAsLocal() and MarkAsRemote() functions

#Synced-CL 50532168

diff --git a/Engine/Source/Runtime/CoreUObject/Private/UObject/RemoteObject.cpp b/Engine/Source/Runtime/CoreUObject/Private/UObject/RemoteObject.cpp
--- a/Engine/Source/Runtime/CoreUObject/Private/UObject/RemoteObject.cpp
+++ b/Engine/Source/Runtime/CoreUObject/Private/UObject/RemoteObject.cpp
@@ -343,17 +343,62 @@ void RegisterSharedObject(UObject* Object)
 	MarkAsRemoteReference(Object);
 }
 
-void MarkAsRemote(UObject* Object, FRemoteServerId DestinationServerId)
+void SetResidence(UObject* Object, EResidence Residence, FRemoteServerId ResidentServerId)
 {
+	using namespace UE::RemoteObject::Handle;
+
 	static_assert(sizeof(FObjectHandle) == sizeof(UObject*));
-	checkf(!Object->IsTemplate(), TEXT("Attempted to Migrate Template Object '%s' which is considered an asset and never allowed to migrate"), *GetNameSafe(Object));
+
+	checkf(ObjectMaps, TEXT("Trying to set residence for %s to %s but the remote object system is not initialized!"), *GetPathNameSafe(Object), EnumToString(Residence));
+	checkf(Object, TEXT("Trying to set residence for a null object!"));
+	checkf(Residence == EResidence::Remote || ResidentServerId.IsLocal(), TEXT("Setting residence for %s to %s with a non-local server id (%s)"), *GetPathNameSafe(Object), EnumToString(Residence), *ResidentServerId.ToString());
+	ensureMsgf(!Object->HasAnyFlags(EObjectFlags::RF_ClassDefaultObject | EObjectFlags::RF_ArchetypeObject), TEXT("We're about to set an ArchetypeObject %s as remote reference"), *GetNameSafe(Object));
+	
+	EInternalObjectFlags ClearFlags = EInternalObjectFlags::None;
+	EInternalObjectFlags SetFlags = EInternalObjectFlags::None;
+
+	switch (Residence)
+	{
+	case EResidence::Local:
+		{
+			ClearFlags = (EInternalObjectFlags::Remote | EInternalObjectFlags::LocalNotReady);
+			SetFlags = EInternalObjectFlags::RemoteReference;
+		}
+		break;
+	case EResidence::LocalNotReady:
+		{
+			ClearFlags = EInternalObjectFlags::Remote;
+			SetFlags = (EInternalObjectFlags::RemoteReference | EInternalObjectFlags::LocalNotReady);
+		}
+		break;
+	case EResidence::Remote:
+		{
+			ClearFlags = (EInternalObjectFlags_RootFlags | EInternalObjectFlags::RemoteReference | EInternalObjectFlags::Borrowed | EInternalObjectFlags::LocalNotReady);
+			SetFlags = EInternalObjectFlags::Remote;			
+		}
+		break;
+	default:
+		checkf(false, TEXT("Attempted to set an unknown residence (%d) for %s"), (int32)Residence, *GetPathNameSafe(Object));
+		break;
+	}
 
 	FUObjectItem* ObjectItem = GUObjectArray.ObjectToObjectItem(Object);
-#if UE_WITH_REMOTE_OBJECT_HANDLE
-	ObjectItem->SetFlags(EInternalObjectFlags::Remote);
-	ObjectItem->ClearFlags(EInternalObjectFlags_RootFlags | EInternalObjectFlags::RemoteReference | EInternalObjectFlags::Borrowed);
-#endif
-	ObjectMaps->FindOrAddStub(Object, DestinationServerId);
+	ObjectItem->ClearFlags(ClearFlags);
+	ObjectItem->SetFlags(SetFlags);
+	if (FRemoteObjectStub* Stub = ObjectMaps->FindOrAddStub(Object, ResidentServerId)) // FindOrAddStub always needs to be called. It's not only for debugging purposes
+	{
+		UE_LOG(LogRemoteObject, VeryVerbose, TEXT("Setting residence for %s %s to %s"), *Stub->Id.ToString(), *GetPathNameSafe(Object), EnumToString(Residence));
+
+		checkf(Residence != EResidence::LocalNotReady || Stub->MigratedData, TEXT("Setting residence for %s to LocalNotReady is only possible if the stub references migrated data"), 
+			*Stub->Id.ToString(ERemoteIdToStringVerbosity::PathName));
+		checkf(Residence == EResidence::LocalNotReady || !Stub->MigratedData, TEXT("Setting residence for %s to %s is only possible if the stub does NOT reference migrated data"), 
+			*Stub->Id.ToString(ERemoteIdToStringVerbosity::PathName), EnumToString(Residence));
+	}
+}
+
+void MarkAsRemote(UObject* Object, FRemoteServerId ResidentServerId)
+{
+	SetResidence(Object, EResidence::Remote, ResidentServerId);
 }
 
 void MarkAsRemoteReference(const UObject* Object)
@@ -392,14 +437,7 @@ bool IsBorrowed(const UObject* Object)
 
 void MarkAsLocal(UObject* Object)
 {
-	ensureMsgf(!Object->HasAnyFlags(EObjectFlags::RF_ClassDefaultObject | EObjectFlags::RF_ArchetypeObject), TEXT("We're about to set an ArchetypeObject %s as remote reference"), *GetNameSafe(Object));
-
-#if UE_WITH_REMOTE_OBJECT_HANDLE
-	FUObjectItem* ObjectItem = GUObjectArray.ObjectToObjectItem(Object);
-	ObjectItem->ClearFlags(EInternalObjectFlags::Remote);
-	ObjectItem->SetFlags(EInternalObjectFlags::RemoteReference);
-#endif
-	ObjectMaps->FindOrAddStub(Object, FRemoteServerId::GetLocalServerId());
+	SetResidence(Object, EResidence::Local, FRemoteServerId::GetLocalServerId());
 }
 
 void StoreAssetPath(UObject* Object)
@@ -589,16 +627,18 @@ EResidence GetResidence(const UObject* Object)
 	{
 		const int32 InternalIndex = GUObjectArray.ObjectToIndex(Object);
 		const FUObjectItem* ObjectItem = InternalIndex >= 0 ? GUObjectArray.IndexToObject(InternalIndex) : nullptr;
-		if (!ObjectItem || ObjectItem->HasAnyFlags(EInternalObjectFlags::PendingConstruction))
+		if (ObjectItem && !ObjectItem->HasAnyFlags(EInternalObjectFlags::PendingConstruction))
 		{
-			return EResidence::Local;
-		}
-		else if (!ObjectItem->HasAnyFlags(EInternalObjectFlags::Remote))
-		{
-			UE::RemoteObject::Handle::FRemoteObjectStub* Stub = ObjectMaps->FindStub(ObjectItem->GetRemoteId());
-			return Stub && Stub->MigratedData ? EResidence::LocalNotReady : EResidence::Local;
+			if (ObjectItem->HasAnyFlags(EInternalObjectFlags::LocalNotReady))
+			{
+				return EResidence::LocalNotReady;
+			}
+			else if (ObjectItem->HasAnyFlags(EInternalObjectFlags::Remote))
+			{
+				return EResidence::Remote;
+			}
 		}
-		return EResidence::Remote;
+		return EResidence::Local;
 	}
 #endif
 	return EResidence::Local;
diff --git a/Engine/Source/Runtime/CoreUObject/Private/UObject/RemoteObjectPrivate.h b/Engine/Source/Runtime/CoreUObject/Private/UObject/RemoteObjectPrivate.h
--- a/Engine/Source/Runtime/CoreUObject/Private/UObject/RemoteObjectPrivate.h
+++ b/Engine/Source/Runtime/CoreUObject/Private/UObject/RemoteObjectPrivate.h
@@ -37,9 +37,18 @@ namespace UE::RemoteObject::Private
 	*/
 	void ShutdownRemoteObjects();
 
+	/**
+	* Sets residence of an object
+	* @param Object Object which residence will be set
+	* @param Residence Residence to set
+	* @param ResidentServerId Server where the object is currently resident (should be FRemoteServerId::GetLocalServerId() for EResidence::Local and LocalNotReady)
+	*/
+	void SetResidence(UObject* Object, EResidence Residence, FRemoteServerId ResidentServerId);
+
 	/**
 	* Marks object memory as remote and creates its stub
 	*/
+	UE_DEPRECATED(5.8, "Use SetResidence(Object, EResidence::Remote, ResidentServerId) instead.")
 	void MarkAsRemote(UObject* Object, FRemoteServerId DestinationServerId);
 
 	/**
@@ -55,6 +64,7 @@ namespace UE::RemoteObject::Private
 	/**
 	* Marks object memory as local
 	*/
+	UE_DEPRECATED(5.8, "Use SetResidence(Object, EResidence::Local, FRemoteServerId::GetLocalServerId()) instead.")
 	void MarkAsLocal(UObject* Object);
 
 	/**
diff --git a/Engine/Source/Runtime/CoreUObject/Private/UObject/RemoteObjectSerialization.cpp b/Engine/Source/Runtime/CoreUObject/Private/UObject/RemoteObjectSerialization.cpp
--- a/Engine/Source/Runtime/CoreUObject/Private/UObject/RemoteObjectSerialization.cpp
+++ b/Engine/Source/Runtime/CoreUObject/Private/UObject/RemoteObjectSerialization.cpp
@@ -1102,12 +1102,13 @@ UObject* ConstructRemoteObject(const FRemoteObjectData& ObjectData, const FRemot
 
 	checkf(FRemoteObjectId(Object) == SerializedObjectId, TEXT("Created an object with a different ID:%s than requested:%s"), *FRemoteObjectId(Object).ToString(), *SerializedObjectId.ToString());
 
-	MarkAsLocal(Object);
 	if (UE::RemoteObject::Handle::FRemoteObjectStub* Stub = UE::RemoteObject::Private::FindRemoteObjectStub(SerializedObjectId))
 	{
 		// If we don't reset the migrated data immediatetly we will enter an infinite recursion when trying to resolve this object's inner object's outer (this Object)
 		Stub->MigratedData = nullptr;
 	}
+	SetResidence(Object, EResidence::Local, FRemoteServerId::GetLocalServerId());
+
 	// Update internal flags on the migrated object. It's possible the object being migrated already existed in memory on this server and had the EInternalObjectFlags::Garbage flag set.
 	// Unless the migrated version also had this flag set we need to clear it.
 	// It's also possible that the local object didn't have this flag set but the migrated one has so we need to set it on this server too (it's not impossible to migrate objects marked as garbage)
diff --git a/Engine/Source/Runtime/CoreUObject/Private/UObject/RemoteObjectTransfer.cpp b/Engine/Source/Runtime/CoreUObject/Private/UObject/RemoteObjectTransfer.cpp
--- a/Engine/Source/Runtime/CoreUObject/Private/UObject/RemoteObjectTransfer.cpp
+++ b/Engine/Source/Runtime/CoreUObject/Private/UObject/RemoteObjectTransfer.cpp
@@ -56,7 +56,7 @@ TRefCountPtr<FTransactionalMigrationData> StoreMigratedObjectData(FRemoteObjectD
 /**
 * Removes the specified migrated data from the stubs of objects serialized inside of the migrated data
 */
-void ResetMigratedData(TRefCountPtr<FTransactionalMigrationData> MigratedData, FRemoteServerId);
+void ResetMigratedData(TRefCountPtr<FTransactionalMigrationData> MigratedData);
 
 /**
 * Creates, deserializes and PostMigrates an object(s) from object data received from another server
@@ -252,7 +252,7 @@ public:
 					}
 					else
 					{
-						ResetMigratedData(MigratedData, HighestPriorityRequest->DestinationServerId);
+						ResetMigratedData(MigratedData);
 						ObjectData = MoveTemp(MigratedData->Data);						
 						
 						for (FRemoteObjectId SerializedId : ObjectData)
@@ -312,7 +312,7 @@ public:
 #if UE_WITH_REMOTE_OBJECT_HANDLE
 								ObjectToSend->PostMigrate(MigrationContext);
 #endif								
-								MarkAsRemote(ObjectToSend, HighestPriorityRequest->DestinationServerId);
+								SetResidence(ObjectToSend, EResidence::Remote, HighestPriorityRequest->DestinationServerId);
 								UE::RemoteObject::Handle::ChangeOwnerServerId(ObjectToSend, OwnerServerId);
 							}
 							else
@@ -727,7 +727,7 @@ public:
 #if UE_WITH_REMOTE_OBJECT_HANDLE
 			SentObject->PostMigrate(MigrationContext);
 #endif
-			MarkAsRemote(SentObject, OwnerServerId);
+			SetResidence(SentObject, EResidence::Remote, OwnerServerId);
 		}
 
 		UE_AUTORTFM_OPEN
@@ -969,7 +969,7 @@ TRefCountPtr<FTransactionalMigrationData> StoreMigratedObjectData(FRemoteObjectD
 		}
 		if (UObject* Object = StaticFindObjectFastInternal(SerializedId))
 		{
-			MarkAsLocal(Object);
+			SetResidence(Object, EResidence::LocalNotReady, FRemoteServerId::GetLocalServerId());
 		}
 	};
 
@@ -1144,7 +1144,7 @@ UObject* MigrateObjectFromObjectData(FRemoteObjectData& ObjectData, const FUObje
 	return Result;
 } 
 
-void ResetMigratedData(TRefCountPtr<FTransactionalMigrationData> MigratedData, FRemoteServerId ResidentServerId)
+void ResetMigratedData(TRefCountPtr<FTransactionalMigrationData> MigratedData)
 {
 	using namespace UE::RemoteObject::Handle;
 	using namespace UE::RemoteObject::Private;
@@ -1155,8 +1155,8 @@ void ResetMigratedData(TRefCountPtr<FTransactionalMigrationData> MigratedData, F
 	for (FRemoteObjectId SerializedId : MigratedData->Data)
 	{
 		FRemoteObjectStub* Stub = FindRemoteObjectStub(SerializedId);
-		checkf(Stub->MigratedData == MigratedData || !Stub->MigratedData, TEXT("Stub for serialized object id %s does not have the expected migrated object data"), *SerializedId.ToString());
-		Stub->ResidentServerId = ResidentServerId;
+		checkf(Stub, TEXT("Resetting migrated data for %s but its stub does not exist"), *SerializedId.ToString());
+		checkf(Stub->MigratedData == MigratedData || !Stub->MigratedData, TEXT("Stub for serialized object id %s does not have the expected migrated object data"), *SerializedId.ToString());		
 		Stub->MigratedData = nullptr;
 	}
 	MigratedData->bCompleted = true;
@@ -1177,7 +1177,20 @@ void TransactionallyMigrateObjects(TRefCountPtr<FTransactionalMigrationData> Mig
 		MigrateObjectFromObjectData(MigratedData->Data, MigratedData->Context);
 		SetTransactionallyPostMigratingObjects(false);
 
-		ResetMigratedData(MigratedData, FRemoteServerId::GetLocalServerId());
+		ResetMigratedData(MigratedData);
+
+#if DO_CHECK
+		for (FRemoteObjectId SerializedId : MigratedData->Data)
+		{
+			FRemoteObjectStub* Stub = FindRemoteObjectStub(SerializedId);
+			checkf(Stub, TEXT("Object %s has been fully migrated but its stub does not exist"), *SerializedId.ToString());
+			checkf(Stub->ResidentServerId == FRemoteServerId::GetLocalServerId(), TEXT("Object %s has been fully migrated but its resident server is not local (%s)"), *SerializedId.ToString(), *Stub->ResidentServerId.ToString());
+			checkf(Stub->MigratedData == nullptr, TEXT("Stub for serialized object %s still has its migrated data stored even though the data has been fully deserialized"), *SerializedId.ToString());
+			UObject* Object = StaticFindObjectFastInternal(SerializedId);
+			checkf(Object != nullptr, TEXT("Object %s has been migrated but its UObject does not exist locally"), *SerializedId.ToString());
+			checkf(GetResidence(Object) == EResidence::Local, TEXT("Object %s %s has been fully migrated but its residence is not Local (%s)"), *SerializedId.ToString(), *GetPathNameSafe(Object), EnumToString(GetResidence(Object)));
+		}
+#endif
 	}
 #if DO_CHECK
 	else
@@ -1385,7 +1398,7 @@ void TransferObjectOwnershipToRemoteServer(UObject* Object, FRemoteServerId Dest
 #if UE_WITH_REMOTE_OBJECT_HANDLE
 			SentObject->PostMigrate(MigrationContext);
 #endif
-			MarkAsRemote(SentObject, DestinationServerId);
+			SetResidence(SentObject, EResidence::Remote, DestinationServerId);
 			UE::RemoteObject::Handle::ChangeOwnerServerId(SentObject, DestinationServerId);
 		}
 	}
@@ -1533,7 +1546,7 @@ void StoreObjectToDatabase(UObject* Object, FRemoteObjectId ObjectId)
 		// This ensure should never fire.  We shouldn't be able to store a non-locally-owned Object into the Database (that is the job of the Owning Server).  A scenario where I've seen this:
 		// We borrowed an Object that was Marked As Garbage and didn't return it to the Owning Server before a GC.
 		ensureMsgf(false, TEXT("%hs was about to store non-locally-owned Object %s (residence: %d) to Server %s's Database; making it remote instead"), __func__, *GetNameSafe(Object), Residence, *LocalServerId.ToString());
-		MarkAsRemote(Object, OwnerServerId);
+		SetResidence(Object, EResidence::Remote, OwnerServerId);
 		return;
 	}
 
@@ -1581,7 +1594,7 @@ void StoreObjectToDatabase(UObject* Object, FRemoteObjectId ObjectId)
 
 		// MarkAsRemote clears EInternalObjectFlags::RemoteReference flag which will prevent re-enetring this loop with the same outermost object but 
 		// we still want to separately process its subobjects (so that they're stored as separate entries in the database) so only mark the root object as remote
-		MarkAsRemote(OutermostUnreachableRemoteReference, DatabaseId);
+		SetResidence(OutermostUnreachableRemoteReference, EResidence::Remote, DatabaseId);
 
 		// It's possible that OutermostUnreachableRemoteReference no longer referenced this Object and Object hasn't been stored to disk
 		// in which case we need to repeat this loop until Object is marked as remote (and is no longer marked as RemoteReference).
diff --git a/Engine/Source/Runtime/CoreUObject/Public/UObject/ObjectMacros.h b/Engine/Source/Runtime/CoreUObject/Public/UObject/ObjectMacros.h
--- a/Engine/Source/Runtime/CoreUObject/Public/UObject/ObjectMacros.h
+++ b/Engine/Source/Runtime/CoreUObject/Public/UObject/ObjectMacros.h
@@ -668,13 +668,15 @@ enum class EInternalObjectFlags : int32
 	None = 0,
 
 #if UE_WITH_REMOTE_OBJECT_HANDLE
-	Remote = 1 << 0, //< Object is no longer owned by this process
-	RemoteReference = 1 << 1, //< Object referenced by remote process
-	Borrowed = 1 << 2, //< Object that is not owned by this process 
+	Remote = 1 << 0,			//< Object is no longer owned by this process
+	RemoteReference = 1 << 1,	//< Object referenced by remote process
+	Borrowed = 1 << 2,			//< Object that is not owned by this process 
+	LocalNotReady = 1 << 3,		//< Object's memory is local but requires additional work before it can be accessed
 #else
 	Remote = 0,
 	RemoteReference = 0,
 	Borrowed = 0,
+	LocalNotReady = 0,
 #endif
 
 
@@ -711,7 +713,7 @@ ENUM_CLASS_FLAGS(EInternalObjectFlags);
 #define EInternalObjectFlags_MinFlagBitIndex 14
 #define EInternalObjectFlags_RefCountMask 0xFFFFFFFF
 #define EInternalObjectFlags_GarbageCollectionKeepFlags (EInternalObjectFlags::Borrowed | EInternalObjectFlags::Native | EInternalObjectFlags::Async | EInternalObjectFlags::AsyncLoadingPhase1 | EInternalObjectFlags::AsyncLoadingPhase2 | EInternalObjectFlags::LoaderImport)
-#define EInternalObjectFlags_AllFlags (EInternalObjectFlags::Remote | EInternalObjectFlags::RemoteReference | EInternalObjectFlags::Borrowed | EInternalObjectFlags::ReachabilityFlag0 | EInternalObjectFlags::ReachabilityFlag1 | EInternalObjectFlags::ReachabilityFlag2 | EInternalObjectFlags::AutoRTFMConstructionAborted | EInternalObjectFlags::LoaderImport | EInternalObjectFlags::Garbage | EInternalObjectFlags::ReachableInCluster | EInternalObjectFlags::ClusterRoot | EInternalObjectFlags::Native | EInternalObjectFlags::Async | EInternalObjectFlags::AsyncLoadingPhase1 | EInternalObjectFlags::AsyncLoadingPhase2 | EInternalObjectFlags::RootSet | EInternalObjectFlags::PendingConstruction | EInternalObjectFlags::Unreachable)
+#define EInternalObjectFlags_AllFlags (EInternalObjectFlags::Remote | EInternalObjectFlags::RemoteReference | EInternalObjectFlags::Borrowed | EInternalObjectFlags::LocalNotReady | EInternalObjectFlags::ReachabilityFlag0 | EInternalObjectFlags::ReachabilityFlag1 | EInternalObjectFlags::ReachabilityFlag2 | EInternalObjectFlags::AutoRTFMConstructionAborted | EInternalObjectFlags::LoaderImport | EInternalObjectFlags::Garbage | EInternalObjectFlags::ReachableInCluster | EInternalObjectFlags::ClusterRoot | EInternalObjectFlags::Native | EInternalObjectFlags::Async | EInternalObjectFlags::AsyncLoadingPhase1 | EInternalObjectFlags::AsyncLoadingPhase2 | EInternalObjectFlags::RootSet | EInternalObjectFlags::PendingConstruction | EInternalObjectFlags::Unreachable)
 #define EInternalObjectFlags_RootFlags (EInternalObjectFlags::RootSet | EInternalObjectFlags_GarbageCollectionKeepFlags)
 #define EInternalObjectFlags_ReachabilityFlags (EInternalObjectFlags::ReachabilityFlag0 | EInternalObjectFlags::ReachabilityFlag1 | EInternalObjectFlags::ReachabilityFlag2 | EInternalObjectFlags::Unreachable)
 #define EInternalObjectFlags_AsyncLoading (EInternalObjectFlags::AsyncLoadingPhase1 | EInternalObjectFlags::AsyncLoadingPhase2)
