Guard against constructing two objects with the same remote id

#Synced-CL 50532168

diff --git a/Engine/Source/Runtime/CoreUObject/Private/UObject/UObjectHash.cpp b/Engine/Source/Runtime/CoreUObject/Private/UObject/UObjectHash.cpp
--- a/Engine/Source/Runtime/CoreUObject/Private/UObject/UObjectHash.cpp
+++ b/Engine/Source/Runtime/CoreUObject/Private/UObject/UObjectHash.cpp
@@ -1393,8 +1393,17 @@ FORCEINLINE static UPackage* AssignExternalPackageToObject(FUObjectHashTables& T
 // Assumes that ThreadHash's critical is already locked
 FORCEINLINE static void AddToRemoteIdMap(FUObjectHashTables& ThreadHash, UObjectBase* Object)
 {
-	FHashBucket& Bucket = ThreadHash.HashId.FindOrAdd(GetTypeHash(FRemoteObjectId(Object)));
-	checkSlow(!Bucket.Contains(Object)); // if it already exists, something is wrong with the external code
+	FRemoteObjectId RemoteId(Object);
+	FHashBucket& Bucket = ThreadHash.HashId.FindOrAdd(GetTypeHash(RemoteId));
+#if DO_CHECK
+	for (FHashBucketIterator It(Bucket); It; ++It)
+	{
+		UObject* ExistingObject = (UObject*)*It;
+		checkf(ExistingObject != Object, TEXT("%s (%s) is being added to the remote id map for the second time"), *GetPathNameSafe(ExistingObject), *RemoteId.ToString());
+		FRemoteObjectId ExistingId(ExistingObject);
+		checkf(ExistingId != RemoteId, TEXT("Trying to add %s to the remote id map but it has the same remote id (%s) as %s"), *GetPathNameSafe((UObject*)Object), *RemoteId.ToString(), *GetPathNameSafe(ExistingObject));
+	}
+#endif
 	Bucket.Add(Object);
 }
 #endif // UE_WITH_REMOTE_OBJECT_HANDLE
