[Backout] - CL50592960
#rnx
Original CL Desc
-----------------------------------------------------------------
Support for external entities in Prefabs and Non-Root world entities

Added menu entries to externalize/internalize owned entities for prefabs and non-root entities, enabled according to relevant cvars.
External owned entities is enabled/disabled per-Entity by the bSaveOwnedEntitiesAsExternalFiles property which sets the RF_HasDynamicImports object flag and adds an Externalized override operation in the OverriddenPropertySet for that that object/property.
Those externally packaged objects have additional metadata to allow tracking of the override operations.

When using external entities in an instance (either in a derived prefab or an world instance) the overrides operations are distributed in the external entities
* Added entities:
*Detected in entity::Rename(), and in a post instantiation delegate
*Saves the entity in its own file with specific metadata
* Removed entities:
* Detected through a FPropertyArrayExternalObjectHandler
* Removes the entity and adds a UOverrideRemovePlaceholder instead, with specific metadata to know which entity is being removed
* Modified entities:
* Saved as its own file

Added a new flag in FPropertyTag to know if that property was saved using external objects
The serialization of FPropertyArray has been modified to invoke a registered handler (per-property/class it can be reused for other cases of external assets) upon load/save
* On Load of a package with external entities the dynamic imports callbacks adds the necessary packages to the imports (discovered through the AssetRegistry)
* On Deserialization of the property, the necessary override operations are discovered from the assets/metadata in the AssetRegistry and the overrides are restored
* On Serialization of the property, if we'll be saving in external assets we save no operation/reference it gets rebuilt on load
Overrides management code has been modified to invoke the same handler as the serialization to notify objects of array override being added

Added FPropertyArrayExternalObjectHandler
* Registered per Property/Class so that invocation can occur on the native classes and on IDOs
* Invoked during serialization and overrides creation
* The handler must implement IClassPropertyArrayExternalOverridesHandler, and support being called for IDOs

Added FExternalObjectsOverrideHelper
* A utility class to deal with various operations pertaining to external overrides, like override discovery, external package management, remove placeholders,...

Added a delegate to FCoreUObjectDelegates::OnObjectsReinstanced to apply all missing externalizations on instances after reinstantiation (since they can override the bSaveOwnedEntitiesAsExternalFiles new entities from the template would not be external yet)

Added 2 cvars to control this feature availability
* Entity.OneFilePerPrefabEntityEnabled : Controls whether we allow entities to be externalized in prefabs
* Entity.OneFilePerEntityOnAllInstancesEnabled : Controls whether we allow all owned entities instances to be externalized (ie non-root owned entities in a world)

diff --git a/Engine/Config/BaseGame.ini b/Engine/Config/BaseGame.ini
--- a/Engine/Config/BaseGame.ini
+++ b/Engine/Config/BaseGame.ini
@@ -264,8 +264,6 @@ SubtitleWidgetToUseDefault=/Script/UMG.WidgetBlueprintGeneratedClass'/SubtitlesA
 +ProjectAssetPaths="/Game"
 +DisallowedObjectClassesToLoad="/Script/Engine.Blueprint"
 +DisallowedObjectClassesToLoad="/Script/Engine.AssetImportData"
-+DisallowedObjectClassesToLoad="/Script/Entity.BaseEntity"
-+DisallowedObjectClassesToLoad="/Script/Engine.OverrideRemovePlaceholder"
 +DisallowedAssetClassesToGather="/Script/Engine.World"
 +DisallowedPathsToGather="/Game/Developers"
 ;+EngineSpecificAssetsToCook="/Engine/Maps/Templates/Template_Default.umap"
diff --git a/Engine/Source/Editor/UnrealEd/Private/Kismet2/KismetReinstanceUtilities.cpp b/Engine/Source/Editor/UnrealEd/Private/Kismet2/KismetReinstanceUtilities.cpp
--- a/Engine/Source/Editor/UnrealEd/Private/Kismet2/KismetReinstanceUtilities.cpp
+++ b/Engine/Source/Editor/UnrealEd/Private/Kismet2/KismetReinstanceUtilities.cpp
@@ -3463,7 +3463,6 @@ void FBlueprintCompileReinstancer::PreCreateSubObjectsForReinstantiation_Inner(
 		UObject* OldSubObjectOuter = OldSubObject->GetOuter();
 		FName SubObjectName = OldSubObject->GetFName();
 		EObjectFlags SubObjectFlags = OldSubObject->GetFlags() & RF_Reinst;
-		UPackage* OldExternalPackage = OldSubObject->GetExternalPackage();
 
 		// Only re-create objects that are created from a default object,
 		// skip any default subobjects are they will be created by their parent object CDO.
@@ -3539,14 +3538,6 @@ void FBlueprintCompileReinstancer::PreCreateSubObjectsForReinstantiation_Inner(
 
 						NewSubObject = NewObject<UObject>(NewUObject, SubObjectClass, SubObjectName, SubObjectFlags, Template);
 					}
-
-					// Transfer external package state if necessary 
-					if (OldExternalPackage)
-					{
-						NewSubObject->SetExternalPackage(OldExternalPackage); // place new subobject in the pre-existing external package 
-						OldSubObject->SetExternalPackage(nullptr); // move old subobject outside of external package (it'll be in the transient package) so that it's doesn't cause the package to stick around
-					}
-
 					CreatedInstanceMap.Add(OldSubObject, NewSubObject);
 					PreCreateSubObjectsForReinstantiation_Inner(OldInstancedSubObjects, OldToNewClassMap, OldSubObject, NewSubObject, CreatedInstanceMap, OldToNewInstanceMap, OrderedListOfObjectToCopy);
 					if (OrderedListOfObjectToCopy)
@@ -3560,13 +3551,6 @@ void FBlueprintCompileReinstancer::PreCreateSubObjectsForReinstantiation_Inner(
 		// There might be new subobjects attached to the sub object that are particular to this instance, let's traverse it to find them out.
 		else if (UObject** NewSubObject = ContainedNewSubObjects.FindByPredicate([SubObjectName](UObject* SubObject) { return SubObject && SubObject->GetFName() == SubObjectName; }))
 		{
-			// Transfer external package state if necessary 
-			if (OldExternalPackage)
-			{
-				(*NewSubObject)->SetExternalPackage(OldExternalPackage); // place new suboject in the pre-existing external package 
-				OldSubObject->SetExternalPackage(nullptr); // move old subobject outside of external package (it'll be in the transient package) so that it's doesn't cause the package to stick around
-			}
-
 			PreCreateSubObjectsForReinstantiation_Inner(OldInstancedSubObjects, OldToNewClassMap, OldSubObject, *NewSubObject, CreatedInstanceMap, OldToNewInstanceMap, OrderedListOfObjectToCopy);
 		}
 	}
diff --git a/Engine/Source/Runtime/CoreUObject/Private/UObject/OverriddenObjectsExternalHandler.cpp b/Engine/Source/Runtime/CoreUObject/Private/UObject/OverriddenObjectsExternalHandler.cpp
deleted file mode 100644
diff --git a/Engine/Source/Runtime/CoreUObject/Private/UObject/OverriddenPropertySet.cpp b/Engine/Source/Runtime/CoreUObject/Private/UObject/OverriddenPropertySet.cpp
--- a/Engine/Source/Runtime/CoreUObject/Private/UObject/OverriddenPropertySet.cpp
+++ b/Engine/Source/Runtime/CoreUObject/Private/UObject/OverriddenPropertySet.cpp
@@ -11,7 +11,6 @@
 #include "UObject/PropertyBagRepository.h"
 #include "UObject/UObjectArchetypeHelper.h"
 #include "UObject/UObjectThreadContext.h"
-#include "UObject/OverriddenObjectsExternalHandler.h"
 
 #include UE_INLINE_GENERATED_CPP_BY_NAME(OverriddenPropertySet)
 
@@ -1186,9 +1185,6 @@ void FOverriddenPropertySet::NotifyPropertyChange(FOverriddenPropertyNode* Paren
 							{
 								AddedSubObjectOverriddenProperties->bWasAdded = true;
 							}
-#if WITH_EDITOR
-							FPropertyArrayExternalObjectHandler::Get().NotifyAdd(Owner, Property, AddedSubObject);
-#endif
 						}
 					}
 				};
@@ -1218,10 +1214,6 @@ void FOverriddenPropertySet::NotifyPropertyChange(FOverriddenPropertyNode* Paren
 								{
 									RemovedSubObjectNode.SetOperation(EOverriddenPropertyOperation::Remove);
 								}
-
-#if WITH_EDITOR
-								FPropertyArrayExternalObjectHandler::Get().NotifyRemove(Owner, Property, RemovedSubObject, RemovedSubObjectArchetype);
-#endif
 							}
 							else
 							{
@@ -1239,10 +1231,6 @@ void FOverriddenPropertySet::NotifyPropertyChange(FOverriddenPropertyNode* Paren
 									{
 										bNeedsCleanup = true;
 									}
-
-#if WITH_EDITOR
-									FPropertyArrayExternalObjectHandler::Get().NotifyRemove(Owner, Property, RemovedSubObject, nullptr);
-#endif
 								}
 								else
 								{
@@ -1758,26 +1746,6 @@ void FOverriddenPropertySet::OverrideProperty(FPropertyVisitorPath::Iterator Pro
 	NotifyPropertyChange(&RootNode, EPropertyNotificationType::PostEdit, PropertyIterator, EPropertyChangeType::Unspecified, Data, bNeedsCleanup);
 }
 
-void FOverriddenPropertySet::SetExternalizedOverriddenProperty(const FOverriddenPropertyNodeID& NodeID, bool bExternal)
-{	
-	FOverriddenPropertyNode& Node = RootNode.FindOrAddNode(NodeID);
-	FOverriddenPropertyNodeID ExternalizedID;
-
-	if (bExternal)
-	{
-		FOverriddenPropertyNode& ExternalizeNode = Node.FindOrAddNode(ExternalizedID);
-		ExternalizeNode.SetOperation(EOverriddenPropertyOperation::Externalized);
-	}
-	else
-	{
-		Node.RemoveSubPropertyNode(ExternalizedID);
-		if (Node.GetSubPropertyNodes().Num() == 0)
-		{
-			Node.SetOperation(EOverriddenPropertyOperation::None);
-		}
-	}
-}
-
 void FOverriddenPropertySet::NotifyPropertyChange(const EPropertyNotificationType Notification, FPropertyVisitorPath::Iterator PropertyIterator, const EPropertyChangeType::Type ChangeType, const void* Data)
 {
 	bool bNeedsCleanup = false;
@@ -2019,11 +1987,6 @@ const FOverriddenPropertyNode* FOverriddenPropertySet::GetOverriddenPropertyNode
 	return OverriddenPropertyNode;
 }
 
-FOverridableObjectArrayOverrides::FOverridableObjectArrayOverrides() :
-	ObjectsOperation(EOverriddenPropertyOperation::None)
-{
-}
-
 FScopedPropertyValue::FScopedPropertyValue(FProperty* InProperty)
 	: Property(InProperty)
 	, Data(FMemory::Malloc(InProperty->GetSize(), InProperty->GetMinAlignment()))
diff --git a/Engine/Source/Runtime/CoreUObject/Private/UObject/Property.cpp b/Engine/Source/Runtime/CoreUObject/Private/UObject/Property.cpp
--- a/Engine/Source/Runtime/CoreUObject/Private/UObject/Property.cpp
+++ b/Engine/Source/Runtime/CoreUObject/Private/UObject/Property.cpp
@@ -2029,22 +2029,7 @@ const TCHAR* FProperty::ImportSingleProperty( const TCHAR* Str, void* DestData,
 							ArrayHelper.ExpandForIndex(Index);
 							break;
 						}
-						case EOverriddenPropertyOperation::Externalized:
-						{
-							checkf(ArrayProperty->Inner->HasAnyPropertyFlags(CPF_PersistentInstance), TEXT("Only instanced sub object is supporting externalized operation"));
-							const FObjectPropertyBase* InnerObjectProperty = CastFieldChecked<FObjectPropertyBase>(ArrayProperty->Inner);
 
-							if (Path)
-							{
-								FArchiveSerializedPropertyChain Chain = Path->ToSerializedPropertyChain();
-								if (FOverriddenPropertyNode* ArrayOverriddenPropertyNode = OverriddenProperties ? OverriddenProperties->RestoreOverriddenPropertyOperation(EOverriddenPropertyOperation::Modified, &Chain, /*Property*/nullptr) : nullptr)
-								{
-									OverriddenProperties->RestoreSubPropertyOperation(EOverriddenPropertyOperation::Externalized, *ArrayOverriddenPropertyNode, FOverriddenPropertyNodeID());
-								}
-							}
-
-							break;
-						}
 						default:
 							checkf(false, TEXT("Unsupported array operation while importing text"));
 							return nullptr;
diff --git a/Engine/Source/Runtime/CoreUObject/Private/UObject/PropertyArray.cpp b/Engine/Source/Runtime/CoreUObject/Private/UObject/PropertyArray.cpp
--- a/Engine/Source/Runtime/CoreUObject/Private/UObject/PropertyArray.cpp
+++ b/Engine/Source/Runtime/CoreUObject/Private/UObject/PropertyArray.cpp
@@ -7,7 +7,6 @@
 #include "UObject/LinkerLoad.h"
 #include "UObject/ObjectInstancingGraph.h"
 #include "UObject/OverriddenPropertySet.h"
-#include "UObject/OverriddenObjectsExternalHandler.h"
 #include "UObject/PropertyHelper.h"
 #include "UObject/UnrealTypePrivate.h"
 #include "UObject/UObjectThreadContext.h"
@@ -240,12 +239,10 @@ void FArrayProperty::SerializeItem(FStructuredArchive::FSlot Slot, TNotNull<void
 		Inner->SerializeItem(Slot, Item);
 	};
 
-	const FPropertyTag* PropertyTag = FPropertyTagScope::GetCurrentPropertyTag();
-
 	// Make sure the container is reloading accordingly to the value set in the property tag if any
-	if (UnderlyingArchive.IsLoading() && PropertyTag)
+	if (UnderlyingArchive.IsLoading() && FPropertyTagScope::GetCurrentPropertyTag())
 	{
-		bExperimentalOverridableLogic = PropertyTag->bExperimentalOverridableLogic;
+		bExperimentalOverridableLogic = FPropertyTagScope::GetCurrentPropertyTag()->bExperimentalOverridableLogic;
 	}
 
 	// *** Experimental *** Special serialization path for array with overridable serialization
@@ -253,79 +250,25 @@ void FArrayProperty::SerializeItem(FStructuredArchive::FSlot Slot, TNotNull<void
 	{
 		checkf(!UnderlyingArchive.ArUseCustomPropertyList, TEXT("Using custom property list is not supported by overridable serialization"));
 
-		// Container for temporarily tracking some indices & communicating with external objects handler
-		FOverridableObjectArrayOverrides	ContainerOverrides;
-
-		const FObjectPropertyBase* InnerObjectProperty = CastField<FObjectPropertyBase>(Inner);
-		
-		// this indicates whether or not we can use external objects overrides in this particular serializing situation
-		const bool bCouldUseExternalObjects = InnerObjectProperty && HasAnyPropertyFlags(CPF_ExperimentalExternalObjects) && UnderlyingArchive.IsPersistent() && !UnderlyingArchive.HasAnyPortFlags(PPF_DuplicateForPIE);
-
-		// this indicates whether or not this object is to use external objects for overrides (can be overriden in certain deprecation cases)
-		// this flag will either loaded from the serialized tag, or calculated during save
-		const bool bTaggedUsingExternalsObjects = PropertyTag ? PropertyTag->bExperimentalExternalObjects.Get(false) : false;
-		
 		FStructuredArchive::FRecord Record = Slot.EnterRecord();
 		if (UnderlyingArchive.IsLoading())
 		{
-#if WITH_EDITOR
-			bool bUsingExternalObjects = bTaggedUsingExternalsObjects && bCouldUseExternalObjects;
-
-			if (bCouldUseExternalObjects)
-			{
-				// Final determination for bUsingExternalObjects is delegated to handler if the property
-				// was saved before started to include this state in the tag
-				bool bUseHandlerForUsingExternalObjects = false;
-				if (PropertyTag && !PropertyTag->bExperimentalExternalObjects.IsSet())
-				{
-					bUseHandlerForUsingExternalObjects = true;
-				}
-
-				// Invoke the external object handler to obtain the overrides for this property and the 'usingExternalObjects' state if we get it from there
-				check(Context->SerializedObject);
-				if ((bTaggedUsingExternalsObjects || bUseHandlerForUsingExternalObjects))
-				{					
-					bool bHandlerUsingExternalObjects = FPropertyArrayExternalObjectHandler::Get().GetExternalOverrides(Context->SerializedObject, this, &ContainerOverrides);
-					if (bUseHandlerForUsingExternalObjects)
-					{
-						bUsingExternalObjects = bHandlerUsingExternalObjects;
-					}
-					else
-					{
-						// in case we're keeping this value from the tag validate all everything is as expected
-						check(bHandlerUsingExternalObjects == bUsingExternalObjects);
-					}
-				}
-			}
-#else			 
-			check(!bTaggedUsingExternalsObjects);
-			const bool bUsingExternalObjects = false;
-#endif // WITH_EDITOR
-
 			int32 NumReplaced = 0;
 			FStructuredArchive::FArray ReplacedArray = Record.EnterArray(TEXT("Replaced"), NumReplaced);
 			if (NumReplaced != INDEX_NONE)
 			{
-				if (bUsingExternalObjects)
-				{
-					check(ContainerOverrides.ObjectsOperation == EOverriddenPropertyOperation::Replace);
-					check(NumReplaced == 0);
-				}
-				else
-				{
-					FSerializedPropertyScope SerializedProperty(UnderlyingArchive, Inner, this);
+				FSerializedPropertyScope SerializedProperty(UnderlyingArchive, Inner, this);
 
-					ArrayHelper.EmptyAndAddValues(NumReplaced);
-				
-					for (int32 i = 0; i < NumReplaced; i++)
-					{
-						SerializeContainerItem(ReplacedArray.EnterElement(), ArrayHelper.GetRawPtr(i));
-					}
+				ArrayHelper.EmptyAndAddValues(NumReplaced);
+				for (int32 i = 0; i < NumReplaced; i++)
+				{
+					SerializeContainerItem(ReplacedArray.EnterElement(), ArrayHelper.GetRawPtr(i));
 				}
 			}
 			else
 			{
 				// Only Array of Instanced subobject are handled here as sort of a set where the matching key is done using the archetype
+				const FObjectPropertyBase* InnerObjectProperty = CastFieldChecked<FObjectPropertyBase>(Inner);
 				checkf(InnerObjectProperty->HasAnyPropertyFlags(CPF_PersistentInstance), TEXT("Only supported code path here is the instanced subobjects"));
 
 				FOverriddenPropertySet* OverriddenProperties = FOverridableSerializationLogic::GetOverriddenProperties();
@@ -336,7 +279,7 @@ void FArrayProperty::SerializeItem(FStructuredArchive::FSlot Slot, TNotNull<void
 
 				// If the arrays are not the same size, the code will not be able to remove items correctly
 				// as it searches for the items to remove in the default array but removes them from this array.
-				checkf (bUsingExternalObjects || (DefaultsArrayHelper.Num() == ArrayHelper.Num()), TEXT("The array \'%s\' is expected to be the same size as it default for removal logic."), *GetName());
+				checkf(DefaultsArrayHelper.Num() == ArrayHelper.Num(), TEXT("The array \'%s\' is expected to be the same size as it default for removal logic."), *GetName());
 
 				auto FindObject = [InnerObjectProperty](UObject* Object, UObject* Object2, FScriptArrayHelper& ArrayHelper) -> int32
 				{
@@ -392,7 +335,12 @@ void FArrayProperty::SerializeItem(FStructuredArchive::FSlot Slot, TNotNull<void
 								UE_LOG(LogOverridableObject, VeryVerbose, TEXT("Unable to load removed item %s(0x%p)"), *GetNameSafe(RemovedSubObject), RemovedSubObject);
 							}
 
-							ContainerOverrides.RemovedArchetypes.Add(RemovedSubObject);
+							// Need to fetch the ArrayOverriddenPropertyNode every loop as the previous iteration might have reallocated the node.
+							if (FOverriddenPropertyNode* ArrayOverriddenPropertyNode = OverriddenProperties ? OverriddenProperties->RestoreOverriddenPropertyOperation(EOverriddenPropertyOperation::Modified, UnderlyingArchive.GetSerializedPropertyChain(), /*Property*/nullptr) : nullptr)
+							{
+								// Rebuild the overridden info
+								OverriddenProperties->RestoreSubObjectOperation(EOverriddenPropertyOperation::Remove, *ArrayOverriddenPropertyNode, RemovedSubObject);
+							}
 						}
 					}
 
@@ -468,7 +416,7 @@ void FArrayProperty::SerializeItem(FStructuredArchive::FSlot Slot, TNotNull<void
 									InnerObjectProperty->SetObjectPropertyValue(ArrayHelper.GetRawPtr(Index), ShadowedObject);
 								}
 							}
-#endif // WITH_EDITORONLY_DATA
+#endif //  WITH_EDITORONLY_DATA
 
 						}
 					}
@@ -503,62 +451,38 @@ void FArrayProperty::SerializeItem(FStructuredArchive::FSlot Slot, TNotNull<void
 							}
 							InnerObjectProperty->SetObjectPropertyValue(ArrayHelper.GetRawPtr(Index), AddedSubObject);
 
-							ContainerOverrides.Added.Add(AddedSubObject);
-						}
-					}
-				}
-
-				// If we were tagged as using external objects restore the externalize override state even if we didn't
-				// use it for serialization. 
-				// 
-				// For example, objects being reinstanced will pass through a non persistent archive during CopyPropertiesForUnrelatedObjects
-				// and this will lose the override state if we don't restore it here. 
-				if (bTaggedUsingExternalsObjects)
-				{
-					if (OverriddenProperties)
-					{
-						if (FOverriddenPropertyNode* ArrayOverriddenPropertyNode = OverriddenProperties ? OverriddenProperties->RestoreOverriddenPropertyOperation(EOverriddenPropertyOperation::Modified, UnderlyingArchive.GetSerializedPropertyChain(), /*Property*/nullptr) : nullptr)
-						{
-							OverriddenProperties->RestoreSubPropertyOperation(EOverriddenPropertyOperation::Externalized, *ArrayOverriddenPropertyNode, FOverriddenPropertyNodeID());
+							// Need to fetch the ArrayOverriddenPropertyNode every loop as the previous iteration might have reallocated the node.
+							if (FOverriddenPropertyNode* ArrayOverriddenPropertyNode = OverriddenProperties ? OverriddenProperties->RestoreOverriddenPropertyOperation(EOverriddenPropertyOperation::Modified, UnderlyingArchive.GetSerializedPropertyChain(), /*Property*/nullptr) : nullptr)
+							{
+								// Rebuild the overridden info
+								OverriddenProperties->RestoreSubObjectOperation(EOverriddenPropertyOperation::Add, *ArrayOverriddenPropertyNode, AddedSubObject);
+							}
 						}
 					}
 				}
-			
-				// Restore overrides
-				for (int32 i = 0; i < ContainerOverrides.RemovedArchetypes.Num(); i++)
-				{
-					// Need to fetch the ArrayOverriddenPropertyNode every loop as the previous iteration might have reallocated the node.
-					if (FOverriddenPropertyNode* ArrayOverriddenPropertyNode = OverriddenProperties ? OverriddenProperties->RestoreOverriddenPropertyOperation(EOverriddenPropertyOperation::Modified, UnderlyingArchive.GetSerializedPropertyChain(), /*Property*/nullptr) : nullptr)
-					{
-						// Rebuild the overridden info
-						OverriddenProperties->RestoreSubObjectOperation(EOverriddenPropertyOperation::Remove, *ArrayOverriddenPropertyNode, ContainerOverrides.RemovedArchetypes[i]);
-					}
-				}
-				for (int32 i = 0; i < ContainerOverrides.Added.Num(); i++)
-				{
-					// Need to fetch the ArrayOverriddenPropertyNode every loop as the previous iteration might have reallocated the node.
-					if (FOverriddenPropertyNode* ArrayOverriddenPropertyNode = OverriddenProperties ? OverriddenProperties->RestoreOverriddenPropertyOperation(EOverriddenPropertyOperation::Modified, UnderlyingArchive.GetSerializedPropertyChain(), /*Property*/nullptr) : nullptr)
-					{
-						// Rebuild the overridden info
-						OverriddenProperties->RestoreSubObjectOperation(EOverriddenPropertyOperation::Add, *ArrayOverriddenPropertyNode, ContainerOverrides.Added[i]);
-					}
-				}
-
 			}
 		}
 		else
 		{
-			bool bUsingExternalObjects = bTaggedUsingExternalsObjects && bCouldUseExternalObjects;
-
 			// Container for temporarily tracking some indices
 			TArray<int32> RemovedIndices;
 			TArray<int32> AddedIndices;
 			TArray<int32> ShadowIndices;
 
+			// We filter out external objects if:
+			//	- We're an array of objects.
+			//	- We have the ExperimentalExternalObjects flag.
+			//	- We're writing to a persistent archive.
+			//	- We have at least one external entity object.
+			// For now, this will be considered as a full array replacement.
+			FObjectProperty* ObjectProperty = CastField<FObjectProperty>(Inner);
+			const bool bFilterExternalObjects = ObjectProperty && HasAnyPropertyFlags(CPF_ExperimentalExternalObjects) && UnderlyingArchive.IsPersistent() && !UnderlyingArchive.HasAnyPortFlags(PPF_DuplicateForPIE);
+
 			bool bReplaceArray = !Defaults || !UnderlyingArchive.DoDelta() || UnderlyingArchive.IsTransacting();
 
 			if (!bReplaceArray)
 			{
+				const FObjectPropertyBase* InnerObjectProperty = CastField<FObjectPropertyBase>(Inner);
 				EOverriddenPropertyOperation ArrayOverrideOp = EOverriddenPropertyOperation::None;
 				FOverriddenPropertySet* OverriddenProperties = FOverridableSerializationLogic::GetOverriddenProperties();
 				if (OverriddenProperties)
@@ -574,8 +498,6 @@ void FArrayProperty::SerializeItem(FStructuredArchive::FSlot Slot, TNotNull<void
 
 				if (!bReplaceArray)
 				{
-					FScriptArrayHelper DefaultsArrayHelper(this, Defaults);
-
 					if (FOverridableSerializationLogic::ShouldPropertyShadowSerializeSubObject(this))
 					{
 						const int32 ArrayNum = ArrayHelper.Num();
@@ -590,7 +512,7 @@ void FArrayProperty::SerializeItem(FStructuredArchive::FSlot Slot, TNotNull<void
 
 					if (OverriddenProperties && ArrayOverrideOp != EOverriddenPropertyOperation::None)
 					{
-						auto FindObject = [InnerObjectProperty](const FOverriddenPropertyNodeID& ObjectToFind, FScriptArrayHelper& ArrayHelper) -> int32
+						auto FindObject = [InnerObjectProperty, bFilterExternalObjects](const FOverriddenPropertyNodeID& ObjectToFind, FScriptArrayHelper& ArrayHelper) -> int32
 						{
 							const int32 ArrayNum = ArrayHelper.Num();
 							for (int i = 0; i < ArrayNum; ++i)
@@ -607,6 +529,8 @@ void FArrayProperty::SerializeItem(FStructuredArchive::FSlot Slot, TNotNull<void
 							return INDEX_NONE;
 						};
 
+						FScriptArrayHelper DefaultsArrayHelper(this, Defaults);
+
 						if (const FOverriddenPropertyNode* ArrayOverriddenPropertyNode = OverriddenProperties->GetOverriddenPropertyNode(UnderlyingArchive.GetSerializedPropertyChain()))
 						{
 							for (const FOverriddenPropertyNode& SubNode : ArrayOverriddenPropertyNode->GetSubPropertyNodes())
@@ -632,7 +556,11 @@ void FArrayProperty::SerializeItem(FStructuredArchive::FSlot Slot, TNotNull<void
  										if (Index != INDEX_NONE)
 										{
 											ShadowIndices.Remove(Index);
-											AddedIndices.Add(Index);
+
+											if (!bFilterExternalObjects || !InnerObjectProperty->GetObjectPropertyValue(ArrayHelper.GetElementPtr(Index))->IsPackageExternal())
+											{
+												AddedIndices.Add(Index);												
+											}
 										}
 										else
 										{
@@ -640,11 +568,6 @@ void FArrayProperty::SerializeItem(FStructuredArchive::FSlot Slot, TNotNull<void
 										}
 										break;
 									}
-								case EOverriddenPropertyOperation::Externalized:
-									{
-										// nothing to do here
-										break;
-									}
 								default:
 									checkf(false, TEXT("Unsupported operation type"));
 									break;
@@ -653,78 +576,51 @@ void FArrayProperty::SerializeItem(FStructuredArchive::FSlot Slot, TNotNull<void
 						}
 					}
 				}
-
-				ContainerOverrides.ObjectsOperation = ArrayOverrideOp;
 			}
 
 			if (bReplaceArray)
 			{
-				if (bUsingExternalObjects)
-				{
-					// External objects always acquire their size on reload
-					int32 NumReplaced = 0;
-					FStructuredArchive::FArray ReplacedArray = Record.EnterArray(TEXT("Replaced"), NumReplaced);
-				}
-				else
-				{
-					int32 NumReplaced = ArrayHelper.Num();
-					FStructuredArchive::FArray ReplacedArray = Record.EnterArray(TEXT("Replaced"), NumReplaced);
+				int32 NumReplaced = ArrayHelper.Num();
+				FStructuredArchive::FArray ReplacedArray = Record.EnterArray(TEXT("Replaced"), NumReplaced);
 
-					for (int32 i = 0; i < NumReplaced; ++i)
-					{
-						SerializeContainerItem(ReplacedArray.EnterElement(), ArrayHelper.GetRawPtr(i));
-					}
+				for (int32 i = 0; i < NumReplaced; ++i)
+				{
+					SerializeContainerItem(ReplacedArray.EnterElement(), ArrayHelper.GetRawPtr(i));
 				}
 			}
 			else
 			{
+				int32 NumReplaced = INDEX_NONE;
+				Record.EnterArray(TEXT("Replaced"), NumReplaced);
+
+				checkf(Defaults, TEXT("Expecting overridable serialization to have defaults to compare to"));
 				FScriptArrayHelper DefaultsArrayHelper(this, Defaults);
 
-#if WITH_EDITOR		
-				if (bUsingExternalObjects)
+				int32 NumRemoved = RemovedIndices.Num();
+				FStructuredArchive::FArray RemovedArray = Record.EnterArray(TEXT("Removed"), NumRemoved);
+				for (int32 i : RemovedIndices)
 				{
-					// These are stored as external objects so we let the serialized object handle the storage 
-					// of the override info.
-					int32 NoOpReplaced = INDEX_NONE;
-					int32 NoOp = 0;
-					Record.EnterArray(TEXT("Replaced"), NoOpReplaced);
-					Record.EnterArray(TEXT("Removed"), NoOp);
-					Record.EnterArray(TEXT("Modified"), NoOp);
-					Record.EnterArray(TEXT("Shadowed"), NoOp);
-					Record.EnterArray(TEXT("Added"), NoOp);
+					SerializeContainerItem(RemovedArray.EnterElement(), DefaultsArrayHelper.GetRawPtr(i));
 				}
-				else
-#endif // WITH_EDITOR
-				{					
-					int32 NumReplaced = INDEX_NONE;
-					Record.EnterArray(TEXT("Replaced"), NumReplaced);
 
-					int32 NumRemoved = RemovedIndices.Num();
-					FStructuredArchive::FArray RemovedArray = Record.EnterArray(TEXT("Removed"), NumRemoved);
-					for (int32 i : RemovedIndices)
-					{
-						SerializeContainerItem(RemovedArray.EnterElement(), DefaultsArrayHelper.GetRawPtr(i));
-					}
-
-					// We should not have modified elements anymore but keeping for eventual support of modified values like struct
-					int32 NumModified = 0;
-					FStructuredArchive::FArray ModifiedArray = Record.EnterArray(TEXT("Modified"), NumModified);
+				// We should not have modified elements anymore but keeping for eventual support of modified values like struct
+				int32 NumModified = 0;
+				FStructuredArchive::FArray ModifiedArray = Record.EnterArray(TEXT("Modified"), NumModified);
 
-					// Support of subobject shadowed serialization
-					// Introduced from EUnrealEngineObjectUE5Version::OS_SUB_OBJECT_SHADOW_SERIALIZATION
-					int32 NumShadowed = ShadowIndices.Num();
-					FStructuredArchive::FArray ShadowedArray = Record.EnterArray(TEXT("Shadowed"), NumShadowed);
-					for (int32 i : ShadowIndices)
-					{
-						SerializeContainerItem(ShadowedArray.EnterElement(), ArrayHelper.GetRawPtr(i));
-					}
+				// Support of subobject shadowed serialization
+				// Introduced from EUnrealEngineObjectUE5Version::OS_SUB_OBJECT_SHADOW_SERIALIZATION
+				int32 NumShadowed = ShadowIndices.Num();
+				FStructuredArchive::FArray ShadowedArray = Record.EnterArray(TEXT("Shadowed"), NumShadowed);
+				for (int32 i : ShadowIndices)
+				{
+					SerializeContainerItem(ShadowedArray.EnterElement(), ArrayHelper.GetRawPtr(i));
+				}
 
-					int32 NumAdded = AddedIndices.Num();
-					FStructuredArchive::FArray AddedArray = Record.EnterArray(TEXT("Added"), NumAdded);
-					for (int32 i : AddedIndices)
-					{
-						SerializeContainerItem(AddedArray.EnterElement(), ArrayHelper.GetRawPtr(i));
-					}
+				int32 NumAdded = AddedIndices.Num();
+				FStructuredArchive::FArray AddedArray = Record.EnterArray(TEXT("Added"), NumAdded);
+				for (int32 i : AddedIndices)
+				{
+					SerializeContainerItem(AddedArray.EnterElement(), ArrayHelper.GetRawPtr(i));
 				}
 			}
 		}
diff --git a/Engine/Source/Runtime/CoreUObject/Private/UObject/PropertyTag.cpp b/Engine/Source/Runtime/CoreUObject/Private/UObject/PropertyTag.cpp
--- a/Engine/Source/Runtime/CoreUObject/Private/UObject/PropertyTag.cpp
+++ b/Engine/Source/Runtime/CoreUObject/Private/UObject/PropertyTag.cpp
@@ -10,10 +10,8 @@
 #include "UObject/BlueprintsObjectVersion.h"
 #include "UObject/DebugSerializationFlags.h"
 #include "UObject/OverriddenPropertySet.h"
-#include "UObject/OverriddenObjectsExternalHandler.h"
 #include "UObject/Package.h"
 #include "UObject/UnrealType.h"
-#include "UObject/UObjectThreadContext.h"
 
 #ifndef UE_ENABLE_SERIALIZE_TAGGED_PROPERTY_TRACING
 #define UE_ENABLE_SERIALIZE_TAGGED_PROPERTY_TRACING 0
@@ -49,7 +47,6 @@ enum class EPropertyTagExtension : uint8
 	////////////////////////////////////////////////
 	// First extension group
 	OverridableInformation		= 0x02,
-	HasExternalsObjects			= 0x04,
 
 	//
 	// Add more extensions for the first group here
@@ -158,20 +155,6 @@ static EPropertyTagExtension CalculatePropertyExtensionFlags(FArchive& Underlyin
 		{
 			Tag.bExperimentalOverridableLogic = false;
 		}
-		
-		Tag.bExperimentalExternalObjects = false;
-#if WITH_EDITOR
-		if (Tag.GetProperty()->HasAllPropertyFlags(CPF_ExperimentalOverridableLogic|CPF_ExperimentalExternalObjects) && !UnderlyingArchive.IsCooking())
-		{
-			FUObjectSerializeContext* Context = FUObjectThreadContext::Get().GetSerializeContext();
-			if (Context->SerializedObject)
-			{
-				Tag.bExperimentalExternalObjects = FPropertyArrayExternalObjectHandler::Get().IsSavingExternal(Context->SerializedObject, Tag.GetProperty());
-			}
-			
-			PropertyTagExtensions |= EPropertyTagExtension::HasExternalsObjects;
-		}		
-#endif // WITH_EDITOR
 	}
 
 	return PropertyTagExtensions;
@@ -195,13 +178,6 @@ static void SerializePropertyExtensions(FStructuredArchive::FSlot Slot, EPropert
 		Slot << SA_ATTRIBUTE(TEXT("OverriddenPropertyOperation"), Tag.OverrideOperation);
 		Slot << SA_ATTRIBUTE(TEXT("ExperimentalOverridableLogic"), Tag.bExperimentalOverridableLogic);
 	}
-
-	if (EnumHasAnyFlags(PropertyTagExtensions, EPropertyTagExtension::HasExternalsObjects))
-	{
-		bool bValue = Tag.bExperimentalExternalObjects.Get(false);
-		Slot << SA_ATTRIBUTE(TEXT("ExperimentalExternalObjects"), bValue);
-		Tag.bExperimentalExternalObjects = bValue;
-	}
 }
 
 static void ParsePathName(FStringView Path, UE::FPropertyTypeNameBuilder& Builder)
diff --git a/Engine/Source/Runtime/CoreUObject/Public/UObject/OverriddenObjectsExternalHandler.h b/Engine/Source/Runtime/CoreUObject/Public/UObject/OverriddenObjectsExternalHandler.h
deleted file mode 100644
diff --git a/Engine/Source/Runtime/CoreUObject/Public/UObject/OverriddenPropertySet.h b/Engine/Source/Runtime/CoreUObject/Public/UObject/OverriddenPropertySet.h
--- a/Engine/Source/Runtime/CoreUObject/Public/UObject/OverriddenPropertySet.h
+++ b/Engine/Source/Runtime/CoreUObject/Public/UObject/OverriddenPropertySet.h
@@ -12,22 +12,6 @@
 
 COREUOBJECT_API DECLARE_LOG_CATEGORY_EXTERN(LogOverridableObject, Warning, All);
 
-struct FOverridableObjectArrayOverrides
-{	
-	COREUOBJECT_API FOverridableObjectArrayOverrides();
-
-	// Returned during loading
-	TArray<UObject*> Added;
-	TArray<UObject*> RemovedArchetypes; // Ptr to the Archetype of the removed object  
-	
-	// Objects that are present in the array (inherited from the archetype or wholly replaced) 
-	// and that are not the result of an Add/Remove  (see next enum for a specific interpretation)
-	TArray<UObject*> Objects; 
-	
-	// Indicates the override operation for the Objects in the previous array 
-	EOverriddenPropertyOperation ObjectsOperation;
-};
-
 struct FArchiveSerializedPropertyChain;
 
 struct FOverridableSerializationLogicInternalAdapter;
@@ -199,7 +183,6 @@ enum class EOverriddenPropertyOperation : uint8
 	Add,		/* this element was added in the container */
 	Remove,		/* this element was removed from the container */
 	SubObjectsShadowing, /* only used to serialize/collect the subobjects, this saved data can be skipped at loading unless loading a loose property or in a placeholder which will keep reference to subobjects */
-	Externalized, /* this indicates the elements are stored externally */
 };
 
 inline TOptional<EOverriddenPropertyOperation> GetOverriddenOperationFromString(const FString& OverriddenOperationString)
@@ -483,12 +466,6 @@ public:
 	 * @param Data memory of the current property */
 	UE_API void OverrideProperty(FPropertyVisitorPath::Iterator PropertyIterator, const void* Data);
 
-	/**
-	 * Utility method that allows to indicate if an array property has its content externalized or not
-	 * @param Property The property in the set that is being externalized/internalized 
-	 * @param bExternal Indicate if the content will be stored externally */
-	UE_API void SetExternalizedOverriddenProperty(const FOverriddenPropertyNodeID& NodeID, bool bExternal);
-
 	/**
 	 * Handling and storing modification on a property of an object
 	 * @param Notification type either pre/post property overridden
diff --git a/Engine/Source/Runtime/CoreUObject/Public/UObject/PropertyTag.h b/Engine/Source/Runtime/CoreUObject/Public/UObject/PropertyTag.h
--- a/Engine/Source/Runtime/CoreUObject/Public/UObject/PropertyTag.h
+++ b/Engine/Source/Runtime/CoreUObject/Public/UObject/PropertyTag.h
@@ -66,7 +66,6 @@ public:
 	EPropertyTagSerializeType SerializeType = EPropertyTagSerializeType::Unknown;
 	EOverriddenPropertyOperation OverrideOperation; // Overridable serialization state reconstruction 
 	bool	bExperimentalOverridableLogic = false; // Remember if property had CPF_ExperimentalOverridableLogic when saved
-	TOptional<bool>	bExperimentalExternalObjects; // Remember if property had CPF_ExperimentalOverridableLogic|CPF_ExperimentalExternalObjects when saved and the object enabled serialization of external objects for that property
 
 	// Constructors.
 	FPropertyTag();
diff --git a/Engine/Source/Runtime/Engine/Internal/OverriddenObjectsExternalPackageHelpers.h b/Engine/Source/Runtime/Engine/Internal/OverriddenObjectsExternalPackageHelpers.h
deleted file mode 100644
diff --git a/Engine/Source/Runtime/Engine/Private/OverriddenObjectsExternalPackageHelpers.cpp b/Engine/Source/Runtime/Engine/Private/OverriddenObjectsExternalPackageHelpers.cpp
deleted file mode 100644
diff --git a/Engine/Source/Runtime/Engine/Private/UnrealExporter.cpp b/Engine/Source/Runtime/Engine/Private/UnrealExporter.cpp
--- a/Engine/Source/Runtime/Engine/Private/UnrealExporter.cpp
+++ b/Engine/Source/Runtime/Engine/Private/UnrealExporter.cpp
@@ -853,7 +853,6 @@ void ExportProperties_Overrides
 
 						TArray<int32> RemovedIndices;
 						TArray<int32> AddedIndices;
-						bool bExternalized = false;
 
 						for (const FOverriddenPropertyNode& SubNode: ArrayOverriddenPropertyNode->GetSubPropertyNodes())
 						{
@@ -877,21 +876,12 @@ void ExportProperties_Overrides
 									}
 									break;
 								}
-							case EOverriddenPropertyOperation::Externalized:
-								// nothing to do 
-								bExternalized = true;
-								break;
 							default:
 								checkf(false, TEXT("Unsupported operation type"));
 								break;
 							}
 						}
 
-						if(bExternalized)
-						{
-							Out.Logf(TEXT("%s%s<%s>=%s\r\n"), FCString::Spc(Indent), *SanitizedPropertyName, *GetOverriddenOperationString(EOverriddenPropertyOperation::Externalized), TEXT("\"\""));
-						}
-
 						for (int32 i : RemovedIndices)
 						{
 							FString Value;
