PackageName and AssetRegistry:
Add accessor functions for IMountPoint to FPackageName, including new OnContentPathMounted delegates that pass through the IMountPoint instead of strings.
Use the new IMountPoint accessors in AssetRegistry, for better performance of the upcoming ActiveMountsLoader.
Add FSuspendPathsMounting to match FSuspendPathsDismounting, to get batching performance improvement when mounting MountPoints.
Pass more stack local variable arguments into AssetRegistry helper functions, for use by the hooks for the upcoming ActiveMountsLoader.
#rnx

#Synced-CL 50477166

diff --git a/Engine/Source/Runtime/AssetRegistry/Private/ARLoader.cpp b/Engine/Source/Runtime/AssetRegistry/Private/ARLoader.cpp
--- a/Engine/Source/Runtime/AssetRegistry/Private/ARLoader.cpp
+++ b/Engine/Source/Runtime/AssetRegistry/Private/ARLoader.cpp
@@ -220,6 +220,7 @@ void FPreloader::TryLoadAsync()
 	// Trigger outside the lock so that a locked Consume function that is waiting on PreloadReady can wait inside the lock.
 	PreloadReady->Trigger();
 
+	FString StateFilePath;
 	FConsumeFunction LocalConsumeCallback;
 	{
 		FScopeLock Lock(&StateLock);
@@ -229,6 +230,7 @@ void FPreloader::TryLoadAsync()
 			LoadState = EState::Loaded;
 			if (ConsumeCallback)
 			{
+				StateFilePath = ARPath;
 				LocalConsumeCallback = MoveTemp(ConsumeCallback);
 				ConsumeCallback.Reset();
 				LoadState = EState::Consumed;
@@ -240,7 +242,7 @@ void FPreloader::TryLoadAsync()
 	{
 		// No further threads will read/write payload at this point (until destructor, which is called after all async threads are complete
 		// so we can use it outside the lock
-		LocalConsumeCallback(LocalResult, MoveTemp(Payload));
+		LocalConsumeCallback(LocalResult, MoveTemp(Payload), StateFilePath);
 		Shutdown();
 	}
 }
@@ -256,7 +258,7 @@ FPreloader::EConsumeResult FPreloader::ConsumeInternal(FConsumeFunction&& Consum
 	{
 		Lock.Unlock(); // Unlock before calling external code in Consume callback
 		ELoadResult LocalResult = (LoadState == EState::Consumed || ConsumeCallback) ? ELoadResult::AlreadyConsumed : ELoadResult::Inactive;
-		ConsumeSynchronous(LocalResult, FAssetRegistryState());
+		ConsumeSynchronous(LocalResult, FAssetRegistryState(), FString());
 		return EConsumeResult::Failed;
 	}
 
@@ -264,8 +266,9 @@ FPreloader::EConsumeResult FPreloader::ConsumeInternal(FConsumeFunction&& Consum
 	{
 		ELoadResult LocalResult = TrySetPath() ? TryLoad() : ELoadResult::NotFound;
 		LoadState = EState::Consumed;
+		FString StateFilePath = ARPath;
 		Lock.Unlock(); // Unlock before calling external code in Consume callback
-		ConsumeSynchronous(LocalResult, MoveTemp(Payload));
+		ConsumeSynchronous(LocalResult, MoveTemp(Payload), StateFilePath);
 		Shutdown(); // Shutdown can be called outside the lock since AsyncThread doesn't exist
 		return LocalResult == ELoadResult::Succeeded ? EConsumeResult::Succeeded : EConsumeResult::Failed;
 	}
@@ -293,8 +296,9 @@ FPreloader::EConsumeResult FPreloader::ConsumeInternal(FConsumeFunction&& Consum
 	LoadState = EState::Consumed;
 
 	// No further async threads exist that will read/write payload at this point so we can use it outside the lock
+	FString StateFilePath = ARPath;
 	Lock.Unlock(); // Unlock before calling external code in Consume callback
-	ConsumeSynchronous(LocalResult, MoveTemp(Payload));
+	ConsumeSynchronous(LocalResult, MoveTemp(Payload), StateFilePath);
 	Shutdown(); // Shutdown can be called outside the lock since we have set state to Consumed and the Async thread will notice and exit
 	return LocalResult == ELoadResult::Succeeded ? EConsumeResult::Succeeded : EConsumeResult::Failed;
 }
@@ -361,10 +365,13 @@ void FAsyncConsumer::Wait(UE::AssetRegistry::FInterfaceWriteScopeLock& ScopeLock
 	}
 }
 
-void FAsyncConsumer::Consume(UAssetRegistryImpl& UARI, UE::AssetRegistry::Impl::FEventContext& EventContext, ELoadResult LoadResult, FAssetRegistryState&& ARState)
+void FAsyncConsumer::Consume(UAssetRegistryImpl& UARI, UE::AssetRegistry::Impl::FEventContext& EventContext,
+	UE::AssetRegistry::FInterfaceWriteScopeLock& ProofOfWriteLock, ELoadResult LoadResult,
+	FAssetRegistryState&& ARState, FStringView StateFilePath)
 {
 	// Called within the lock
-	UARI.GuardedData.LoadPremadeAssetRegistry(EventContext, LoadResult, MoveTemp(ARState));
+	UARI.GuardedData.LoadPremadeAssetRegistry(UARI, EventContext, ProofOfWriteLock, LoadResult,
+		MoveTemp(ARState), StateFilePath);
 	check(ReferenceCount >= 1);
 	check(Consumed != nullptr);
 	Consumed->Trigger();
diff --git a/Engine/Source/Runtime/AssetRegistry/Private/ARLoader.h b/Engine/Source/Runtime/AssetRegistry/Private/ARLoader.h
--- a/Engine/Source/Runtime/AssetRegistry/Private/ARLoader.h
+++ b/Engine/Source/Runtime/AssetRegistry/Private/ARLoader.h
@@ -4,6 +4,7 @@
 
 #include "ARTypes.h"
 #include "AssetRegistry/AssetRegistryState.h"
+#include "Containers/StringView.h"
 #include "Containers/UnrealString.h"
 #include "Delegates/Delegate.h"
 #include "HAL/CriticalSection.h"
@@ -44,7 +45,8 @@ public:
 		Failed,
 		Deferred
 	};
-	using FConsumeFunction = TFunction<void(ELoadResult LoadResult, FAssetRegistryState&& ARState)>;
+	using FConsumeFunction = TFunction<void(ELoadResult LoadResult, FAssetRegistryState&& ARState,
+		FStringView StateFilePath)>;
 
 	FPreloader();
 	~FPreloader();
@@ -144,7 +146,8 @@ struct FAsyncConsumer
 	void Wait(FInterfaceWriteScopeLock& ScopeLock);
 	/** Callback from the async thread to consume the Premade ARState */
 	void Consume(UAssetRegistryImpl& UARI, UE::AssetRegistry::Impl::FEventContext& EventContext,
-		ELoadResult LoadResult, FAssetRegistryState&& ARState);
+		UE::AssetRegistry::FInterfaceWriteScopeLock& ProofOfWriteLock,
+		ELoadResult LoadResult, FAssetRegistryState&& ARState, FStringView StateFilePath);
 
 private:
 	/**
diff --git a/Engine/Source/Runtime/AssetRegistry/Private/AssetRegistry.cpp b/Engine/Source/Runtime/AssetRegistry/Private/AssetRegistry.cpp
--- a/Engine/Source/Runtime/AssetRegistry/Private/AssetRegistry.cpp
+++ b/Engine/Source/Runtime/AssetRegistry/Private/AssetRegistry.cpp
@@ -373,7 +373,8 @@ void FAssetRegistryImpl::ConditionalLoadPremadeAssetRegistry(UE::AssetRegistry::
 	AsyncConsumer.Wait(ScopeLock);
 }
 
-void FAssetRegistryImpl::ConsumeOrDeferPreloadedPremade(UAssetRegistryImpl& UARI, Impl::FEventContext& EventContext)
+void FAssetRegistryImpl::ConsumeOrDeferPreloadedPremade(UAssetRegistryImpl& UARI, Impl::FEventContext& EventContext,
+	FInterfaceWriteScopeLock& ProofOfWriteLock)
 {
 	// Called from inside WriteLock on InterfaceLock
 	using namespace UE::AssetRegistry::Premade;
@@ -387,20 +388,24 @@ void FAssetRegistryImpl::ConsumeOrDeferPreloadedPremade(UAssetRegistryImpl& UARI
 
 	if (Premade::CanLoadAsync())
 	{
-		FPreloader::FConsumeFunction ConsumeFromAsyncThread = [this, &UARI](Premade::ELoadResult LoadResult, FAssetRegistryState&& ARState)
+		FPreloader::FConsumeFunction ConsumeFromAsyncThread = [this, &UARI](Premade::ELoadResult LoadResult,
+			FAssetRegistryState&& ARState, FStringView StateFilePath)
 		{
 			Impl::FEventContext EventContext;
 			{
 				UE::AssetRegistry::FInterfaceWriteScopeLock InterfaceScopeLock(UARI.InterfaceLock);
-				AsyncConsumer.Consume(UARI, EventContext, LoadResult, MoveTemp(ARState));
+				AsyncConsumer.Consume(UARI, EventContext, InterfaceScopeLock, LoadResult,
+					MoveTemp(ARState), StateFilePath);
 			}
 			UARI.Broadcast(EventContext);
 		};
-		auto ConsumeOnCurrentThread = [ConsumeFromAsyncThread](Premade::ELoadResult LoadResult, FAssetRegistryState&& ARState) mutable
+		auto ConsumeOnCurrentThread = [ConsumeFromAsyncThread](Premade::ELoadResult LoadResult,
+			FAssetRegistryState&& ARState, FStringView InStateFilePath) mutable
 		{
-			Async(EAsyncExecution::TaskGraph, [LoadResult, ARState=MoveTemp(ARState), ConsumeFromAsyncThread=MoveTemp(ConsumeFromAsyncThread)]() mutable
+			Async(EAsyncExecution::TaskGraph, [LoadResult, ARState=MoveTemp(ARState),
+				ConsumeFromAsyncThread=MoveTemp(ConsumeFromAsyncThread), StateFilePath=FString(InStateFilePath)]() mutable
 			{
-				ConsumeFromAsyncThread(LoadResult, MoveTemp(ARState));
+				ConsumeFromAsyncThread(LoadResult, MoveTemp(ARState), StateFilePath);
 			});
 		};
 
@@ -409,9 +414,12 @@ void FAssetRegistryImpl::ConsumeOrDeferPreloadedPremade(UAssetRegistryImpl& UARI
 	}
 	else
 	{
-		GPreloader.Consume([this, &EventContext](Premade::ELoadResult LoadResult, FAssetRegistryState&& ARState)
+		GPreloader.Consume(
+			[this, &UARI, &EventContext, &ProofOfWriteLock](Premade::ELoadResult LoadResult,
+			FAssetRegistryState&& ARState, FStringView StateFilePath)
 			{
-				LoadPremadeAssetRegistry(EventContext, LoadResult, MoveTemp(ARState));
+				LoadPremadeAssetRegistry(UARI, EventContext, ProofOfWriteLock, LoadResult,
+					MoveTemp(ARState), StateFilePath);
 			});
 	}
 }
@@ -537,8 +545,9 @@ FAssetRegistryImpl::FAssetRegistryImpl()
 {
 }
 
-void FAssetRegistryImpl::LoadPremadeAssetRegistry(Impl::FEventContext& EventContext,
-	Premade::ELoadResult LoadResult, FAssetRegistryState&& ARState)
+void FAssetRegistryImpl::LoadPremadeAssetRegistry(UAssetRegistryImpl& UARI, Impl::FEventContext& EventContext,
+	FInterfaceWriteScopeLock& ProofOfWriteLock,
+	Premade::ELoadResult LoadResult, FAssetRegistryState&& ARState, FStringView StateFilePath)
 {
 	SCOPED_BOOT_TIMING("LoadPremadeAssetRegistry");
 	UE_SCOPED_ENGINE_ACTIVITY("Loading premade asset registry");
@@ -552,6 +561,7 @@ void FAssetRegistryImpl::LoadPremadeAssetRegistry(Impl::FEventContext& EventCont
 		{
 			if (State.GetNumAssets() == 0)
 			{
+				// Add the premade state into the global AssetRegistry
 				State = MoveTemp(ARState);
 				CachePathsFromState(EventContext, State);
 				if (bEmitAssetEvents)
@@ -559,28 +569,38 @@ void FAssetRegistryImpl::LoadPremadeAssetRegistry(Impl::FEventContext& EventCont
 					EventContext.AssetEvents.Reserve(State.GetNumAssets());
 					State.EnumerateAllAssets([&EventContext](const FAssetData& AssetData)
 					{
-						EventContext.AssetEvents.Emplace(AssetData, UE::AssetRegistry::Impl::FEventContext::EEvent::Added);
+						EventContext.AssetEvents.Emplace(AssetData,
+							UE::AssetRegistry::Impl::FEventContext::EEvent::Added);
 					});
 				}
 			}
-			else if (State.GetNumAssets() < ARState.GetNumAssets())
+			else
 			{
-				FAssetRegistryState ExistingState = MoveTemp(State);
-				State = MoveTemp(ARState);
-				CachePathsFromState(EventContext, State);
-				if (bEmitAssetEvents)
+				// This is unexpected, we should not have called AppendState or in any other way added Assets to the
+				// AssetRegistry before loading the premade AssetRegistry. However, we can handle it with only minor
+				// performance issues.
+
+				// Add the premade state into the global AssetRegistry
+				if (State.GetNumAssets() < ARState.GetNumAssets())
 				{
-					EventContext.AssetEvents.Reserve(State.GetNumAssets());
-					State.EnumerateAllAssets([&EventContext](const FAssetData& AssetData)
+					FAssetRegistryState ExistingState = MoveTemp(State);
+					State = MoveTemp(ARState);
+					CachePathsFromState(EventContext, State);
+					if (bEmitAssetEvents)
 					{
-						EventContext.AssetEvents.Emplace(AssetData, UE::AssetRegistry::Impl::FEventContext::EEvent::Added);
-					});
+						EventContext.AssetEvents.Reserve(State.GetNumAssets());
+						State.EnumerateAllAssets([&EventContext](const FAssetData& AssetData)
+							{
+								EventContext.AssetEvents.Emplace(AssetData,
+									UE::AssetRegistry::Impl::FEventContext::EEvent::Added);
+							});
+					}
+					AppendState(EventContext, ExistingState);
+				}
+				else
+				{
+					AppendState(EventContext, ARState, FAssetRegistryState::EInitializationMode::OnlyUpdateNew, bEmitAssetEvents);
 				}
-				AppendState(EventContext, ExistingState);
-			}
-			else
-			{
-				AppendState(EventContext, ARState, FAssetRegistryState::EInitializationMode::OnlyUpdateNew, bEmitAssetEvents);
 			}
 			UpdatePersistentMountPoints();
 			State.bCookedGlobalAssetRegistryState = true;
@@ -729,7 +749,7 @@ void FAssetRegistryImpl::Initialize(UE::AssetRegistry::FInterfaceWriteScopeLock&
 		}
 	}
 
-	ConsumeOrDeferPreloadedPremade(Context.UARI, Context.Events);
+	ConsumeOrDeferPreloadedPremade(Context.UARI, Context.Events, ScopeLock);
 
 	// Report startup time. This does not include DirectoryWatcher startup time.
 	double StartupDuration = FPlatformTime::Seconds() - StartupStartTime;
@@ -916,8 +936,8 @@ void UAssetRegistryImpl::InitializeEvents(UE::AssetRegistry::Impl::FInitializeCo
 
 	// Listen for new content paths being added or removed at runtime.  These are usually plugin-specific asset paths that
 	// will be loaded a bit later on.
-	FPackageName::OnContentPathMounted().AddUObject(this, &UAssetRegistryImpl::OnContentPathMounted);
-	FPackageName::OnContentPathDismounted().AddUObject(this, &UAssetRegistryImpl::OnContentPathDismounted);
+	FPackageName::OnMountPointMounted().AddUObject(this, &UAssetRegistryImpl::OnMountPointMounted);
+	FPackageName::OnMountPointDismounted().AddUObject(this, &UAssetRegistryImpl::OnMountPointDismounted);
 
 	// If we were called before engine has fully initialized, refresh classes on initialize. If not this won't do anything as it already happened
 	FCoreDelegates::OnPostEngineInit.AddUObject(this, &UAssetRegistryImpl::OnPostEngineInit);
@@ -1698,8 +1718,8 @@ void UAssetRegistryImpl::FinishDestroy()
 		UE::AssetRegistry::FInterfaceWriteScopeLock InterfaceScopeLock(InterfaceLock);
 
 		// Stop listening for content mount point events
-		FPackageName::OnContentPathMounted().RemoveAll(this);
-		FPackageName::OnContentPathDismounted().RemoveAll(this);
+		FPackageName::OnMountPointMounted().RemoveAll(this);
+		FPackageName::OnMountPointDismounted().RemoveAll(this);
 		FCoreDelegates::OnPostEngineInit.RemoveAll(this);
 		FCoreDelegates::OnPreExit.RemoveAll(this);
 		IPluginManager::Get().OnLoadingPhaseComplete().RemoveAll(this);
@@ -5625,7 +5645,7 @@ void FAssetRegistryImpl::OnAdditionalMountSearchCompleted(Impl::FEventContext& E
 	SetPerformanceMode(Impl::EPerformanceMode::MostlyStatic);
 
 	LogSearchDiagnostics(AdditionalMountSearchStartTime, /* bInitialScan */ false);
-	CSV_TRACE_REGION_END(TEXT("Additional Mount Search"), TEXT("AssetRegistry")); // Matching CSV_TRACE_REGION_BEGIN in OnContentPathMounted
+	CSV_TRACE_REGION_END(TEXT("Additional Mount Search"), TEXT("AssetRegistry")); // Matching CSV_TRACE_REGION_BEGIN in OnMountPointMounted
 
 	GlobalGatherer->UpdateCacheForSaving();
 	GlobalGatherer->OnAdditionalMountSearchCompleted();
@@ -8317,7 +8337,7 @@ void FAssetRegistryImpl::OnDirectoryRescanRequired(FInterfaceWriteScopeLock& Sco
 		{
 			if (FPackageName::TryGetMountPointForPath(MountedLongPackageName, UnusedPackageName, MountedFilePath, UnusedRelPath))
 			{
-				FString NormalizeMountedFilePath = UE::AssetRegistry::CreateStandardFilename(FString(MountedFilePath));
+				FString NormalizeMountedFilePath = UE::AssetRegistry::CreateStandardFilename(MountedFilePath);
 				if (FPaths::IsUnderDirectory(NormalizeMountedFilePath, NormalizedDirPath))
 				{
 					DirPathsAndPackageNames.Emplace(MoveTemp(NormalizeMountedFilePath), MoveTemp(MountedLongPackageName));
@@ -8366,7 +8386,7 @@ void FAssetRegistryImpl::OnDirectoryRescanRequired(FInterfaceWriteScopeLock& Sco
 					// Convert Filename to a PackagePath. We know the base dir so its faster to use that than FPackageName
 					// which has to scan all mount dirs
 					FStringView RelPath;
-					FString NormalizedFilename = UE::AssetRegistry::CreateStandardFilename(Filename);
+					FString NormalizedFilename = UE::AssetRegistry::CreateStandardFilename(FStringView(Filename));
 					if (!FPathViews::TryMakeChildPathRelativeTo(NormalizedFilename, LocalPath, RelPath))
 					{
 						return true;
@@ -8752,35 +8772,38 @@ void FAssetRegistryImpl::ScanModifiedAssetFiles(FInterfaceWriteScopeLock& ScopeL
 
 }
 
-void UAssetRegistryImpl::OnContentPathMounted(const FString& InAssetPath, const FString& FileSystemPath)
+void UAssetRegistryImpl::OnMountPointMounted(const TRefCountPtr<UE::PackageName::IMountPoint>& MountPoint)
 {
-	TRACE_CPUPROFILER_EVENT_SCOPE(AssetRegistry_OnContentPathMounted);
-	// Sanitize
-	FString AssetPathWithTrailingSlash;
-	if (!InAssetPath.EndsWith(TEXT("/"), ESearchCase::CaseSensitive))
-	{
-		// We actually want a trailing slash here so the path can be properly converted while searching for assets
-		AssetPathWithTrailingSlash = InAssetPath + TEXT("/");
-	}
-	else
+	if (SuspendMountingCounter.load(std::memory_order_acquire) > 0)
 	{
-		AssetPathWithTrailingSlash = InAssetPath;
+		FScopeLock Guard(&MountingCriticalSection);
+		ListOfMountedMountPoints.Emplace(MountPoint);
+		return;
 	}
 
+	OnMultipleMountPointsMounted(TConstArrayView<TRefCountPtr<UE::PackageName::IMountPoint>>(&MountPoint, 1));
+}
+
+void UAssetRegistryImpl::OnMultipleMountPointsMounted(TConstArrayView<TRefCountPtr<UE::PackageName::IMountPoint>> MountedMountPoints)
+{
+	TRACE_CPUPROFILER_EVENT_SCOPE(AssetRegistry_OnMultipleMountPointsMounted);
+	LLM_SCOPE(ELLMTag::AssetRegistry);
+
 #if WITH_EDITOR
 	FDirectoryWatcherModule& DirectoryWatcherModule = FModuleManager::LoadModuleChecked<FDirectoryWatcherModule>(TEXT("DirectoryWatcher"));
 	IDirectoryWatcher* DirectoryWatcher = nullptr;
 	if (UE::AssetRegistry::Impl::IsDirectoryWatcherEnabled())
 	{
-		LLM_SCOPE(ELLMTag::AssetRegistry);
 		DirectoryWatcher = DirectoryWatcherModule.Get();
-		if (DirectoryWatcher)
+		for (const TRefCountPtr<UE::PackageName::IMountPoint>& MountPoint : MountedMountPoints)
 		{
-			// Make sure the directory exits on disk so that the OS level DirectoryWatcher can be used to monitor it.
-			IPlatformFile::GetPlatformPhysical().CreateDirectoryTree(*FileSystemPath);
+			if (DirectoryWatcher)
+			{
+				// Make sure the directory exits on disk so that the OS level DirectoryWatcher can be used to monitor it.
+				IPlatformFile::GetPlatformPhysical().CreateDirectoryTree(*WriteToString<512>(MountPoint->GetLocalPathStandard()));
+			}
 		}
 	}
-		
 #endif
 
 	UE::AssetRegistry::Impl::FEventContext EventContext;
@@ -8790,25 +8813,29 @@ void UAssetRegistryImpl::OnContentPathMounted(const FString& InAssetPath, const
 		UE::AssetRegistry::Impl::FClassInheritanceContext InheritanceContext;
 		UE::AssetRegistry::Impl::FClassInheritanceBuffer InheritanceBuffer;
 		GetInheritanceContextWithRequiredLock(InterfaceScopeLock, InheritanceContext, InheritanceBuffer);
-		GuardedData.OnContentPathMounted(InterfaceScopeLock, EventContext, InheritanceContext, InAssetPath, AssetPathWithTrailingSlash,
-			FileSystemPath);
+		GuardedData.OnMultipleMountPointsMounted(*this, EventContext,
+			InheritanceContext, InterfaceScopeLock, MountedMountPoints);
 
 		// Listen for directory changes in this content path
 #if WITH_EDITOR
-		const FString StandardFileSystemPath = UE::AssetRegistry::CreateStandardFilename(FileSystemPath);
-		// In-game doesn't listen for directory changes
-		if (DirectoryWatcher && !IsDirAlreadyWatchedByRootWatchers(StandardFileSystemPath))
+		for (const TRefCountPtr<UE::PackageName::IMountPoint>& MountPoint : MountedMountPoints)
 		{
-			if (!OnDirectoryChangedDelegateHandles.Contains(AssetPathWithTrailingSlash))
+			const FString StandardFileSystemPath = UE::AssetRegistry::CreateStandardFilename(MountPoint->GetLocalPathStandard());
+			// In-game doesn't listen for directory changes
+			if (DirectoryWatcher && !IsDirAlreadyWatchedByRootWatchers(StandardFileSystemPath))
 			{
-				FDelegateHandle NewHandle;
-				DirectoryWatcher->RegisterDirectoryChangedCallback_Handle(
-					StandardFileSystemPath,
-					IDirectoryWatcher::FDirectoryChanged::CreateUObject(this, &UAssetRegistryImpl::OnDirectoryChanged),
-					NewHandle, 
-					IDirectoryWatcher::WatchOptions::IncludeDirectoryChanges);
+				const FString AssetPathWithTrailingSlash = UE::AssetRegistry::CreateStandardLongPackageNameDirectory(MountPoint->GetLongPackageName());
+				if (!OnDirectoryChangedDelegateHandles.Contains(AssetPathWithTrailingSlash))
+				{
+					FDelegateHandle NewHandle;
+					DirectoryWatcher->RegisterDirectoryChangedCallback_Handle(
+						StandardFileSystemPath,
+						IDirectoryWatcher::FDirectoryChanged::CreateUObject(this, &UAssetRegistryImpl::OnDirectoryChanged),
+						NewHandle,
+						IDirectoryWatcher::WatchOptions::IncludeDirectoryChanges);
 
-				OnDirectoryChangedDelegateHandles.Add(AssetPathWithTrailingSlash, NewHandle);
+					OnDirectoryChangedDelegateHandles.Add(AssetPathWithTrailingSlash, NewHandle);
+				}
 			}
 		}
 #endif // WITH_EDITOR
@@ -8820,56 +8847,63 @@ void UAssetRegistryImpl::OnContentPathMounted(const FString& InAssetPath, const
 namespace UE::AssetRegistry
 {
 
-void FAssetRegistryImpl::OnContentPathMounted(UE::AssetRegistry::FInterfaceWriteScopeLock& ScopeLock, Impl::FEventContext& EventContext,
-	Impl::FClassInheritanceContext& InheritanceContext, const FString& InAssetPath,
-	const FString& AssetPathWithTrailingSlash, const FString& FileSystemPath)
+void FAssetRegistryImpl::OnMultipleMountPointsMounted(UAssetRegistryImpl& UARI, Impl::FEventContext& EventContext,
+	Impl::FClassInheritanceContext& InheritanceContext,
+	UE::AssetRegistry::FInterfaceWriteScopeLock& ScopeLock, 
+	TConstArrayView<TRefCountPtr<UE::PackageName::IMountPoint>> MountedMountPoints)
 {
-	// Content roots always exist
-	AddPath(EventContext, UE::String::RemoveFromEnd(FStringView(AssetPathWithTrailingSlash), TEXTVIEW("/")));
-
-	if (GlobalGatherer.IsValid() && bSearchAllAssets)
+	for (const TRefCountPtr<UE::PackageName::IMountPoint>& MountPoint : MountedMountPoints)
 	{
-		if (GlobalGatherer->IsSynchronous())
-		{
-			Impl::FScanPathContext Context(EventContext, InheritanceContext, { FileSystemPath }, TArray<FString>());
-			ScanPathsSynchronous(&ScopeLock, Context);
-		}
-		else
+		FString FileSystemPath(MountPoint->GetLocalPathStandard());
+		const FString AssetPathWithTrailingSlash = UE::AssetRegistry::CreateStandardLongPackageNameDirectory(MountPoint->GetLongPackageName());
+
+		// Content roots always exist
+		AddPath(EventContext, UE::String::RemoveFromEnd(FStringView(AssetPathWithTrailingSlash), TEXTVIEW("/")));
+
+		if (GlobalGatherer.IsValid() && bSearchAllAssets)
 		{
-			if (!IsGathering())
+			if (GlobalGatherer->IsSynchronous())
+			{
+				Impl::FScanPathContext Context(EventContext, InheritanceContext, { FileSystemPath }, TArray<FString>());
+				ScanPathsSynchronous(&ScopeLock, Context);
+			}
+			else
 			{
-				CSV_TRACE_REGION_BEGIN(TEXT("Additional Mount Search"), TEXT("AssetRegistry")); // Matching CSV_TRACE_REGION_END in OnAdditionalMountSearchCompleted
-				AdditionalMountSearchStartTime = FPlatformTime::Seconds();
-				bAdditionalMountSearchInProgress.store(true, std::memory_order_relaxed);
-				GlobalGatherer->SetIsAdditionalMountSearchInProgress(true);
+				if (!IsGathering())
+				{
+					CSV_TRACE_REGION_BEGIN(TEXT("Additional Mount Search"), TEXT("AssetRegistry")); // Matching CSV_TRACE_REGION_END in OnAdditionalMountSearchCompleted
+					AdditionalMountSearchStartTime = FPlatformTime::Seconds();
+					bAdditionalMountSearchInProgress.store(true, std::memory_order_relaxed);
+					GlobalGatherer->SetIsAdditionalMountSearchInProgress(true);
 #if WITH_EDITOR
-				bGatherNeedsApplyObjectRedirects = true;
+					bGatherNeedsApplyObjectRedirects = true;
 #endif
-			}
+				}
 
-			GlobalGatherer->AddMountPoint(FileSystemPath, InAssetPath);
-			GlobalGatherer->SetIsOnAllowList(FileSystemPath, true);
+				GlobalGatherer->AddMountPoint(MountPoint->GetLocalPathStandard(), MountPoint->GetLongPackageName());
+				GlobalGatherer->SetIsOnAllowList(MountPoint->GetLocalPathStandard(), true);
+			}
 		}
 	}
 }
 
 }
 
-void UAssetRegistryImpl::OnContentPathDismounted(const FString& InAssetPath, const FString& FileSystemPath)
+void UAssetRegistryImpl::OnMountPointDismounted(const TRefCountPtr<UE::PackageName::IMountPoint>& MountPoint)
 {
 	if (SuspendDismountingCounter.load(std::memory_order_acquire) > 0)
 	{
-		FScopeLock Guard(&DismountingCriticalSection);
-		ListOfUnmountedPaths.Emplace(InAssetPath, FileSystemPath);
+		FScopeLock Guard(&MountingCriticalSection);
+		ListOfDismountedMountPoints.Emplace(MountPoint);
 		return;
 	}
 
-	OnMultipleContentPathsDismounted({MakeTuple(InAssetPath, FileSystemPath)});
+	OnMultipleMountPointsDismounted(TConstArrayView<TRefCountPtr<UE::PackageName::IMountPoint>>(&MountPoint, 1));
 }
 
-void UAssetRegistryImpl::OnMultipleContentPathsDismounted(TConstArrayView<TTuple<FString, FString>> UnmountedAssetFilesystemPairs)
+void UAssetRegistryImpl::OnMultipleMountPointsDismounted(TConstArrayView<TRefCountPtr<UE::PackageName::IMountPoint>> DismountedMountPoints)
 {
-	TRACE_CPUPROFILER_EVENT_SCOPE(AssetRegistry_OnMultipleContentPathsDismounted);
+	TRACE_CPUPROFILER_EVENT_SCOPE(AssetRegistry_OnMultipleMountPointsDismounted);
 #if WITH_EDITOR
 	FDirectoryWatcherModule& DirectoryWatcherModule = FModuleManager::LoadModuleChecked<FDirectoryWatcherModule>(TEXT("DirectoryWatcher"));
 	IDirectoryWatcher* DirectoryWatcher = nullptr;
@@ -8883,22 +8917,18 @@ void UAssetRegistryImpl::OnMultipleContentPathsDismounted(TConstArrayView<TTuple
 	{
 		LLM_SCOPE(ELLMTag::AssetRegistry);
 		UE::AssetRegistry::FInterfaceWriteScopeLock InterfaceScopeLock(InterfaceLock);
-		GuardedData.OnMultipleContentPathsDismounted(EventContext, UnmountedAssetFilesystemPairs);
+		GuardedData.OnMultipleMountPointsDismounted(*this, EventContext,
+			DismountedMountPoints);
 
 		// Stop listening for directory changes in this content path
 #if WITH_EDITOR
-		for (const TTuple<FString, FString>& Pair : UnmountedAssetFilesystemPairs)
+		for (const TRefCountPtr<UE::PackageName::IMountPoint>& MountPoint : DismountedMountPoints)
 		{
-			const FString StandardFileSystemPath = UE::AssetRegistry::CreateStandardFilename(Pair.Value);
+			const FString StandardFileSystemPath = UE::AssetRegistry::CreateStandardFilename(MountPoint->GetLocalPathStandard());
 			if (DirectoryWatcher && !IsDirAlreadyWatchedByRootWatchers(StandardFileSystemPath))
 			{
-				// Make sure OnDirectoryChangedDelegateHandles key is symmetrical with the one used in OnContentPathMounted
-				FString AssetPathWithTrailingSlash = Pair.Key;
-				if (!AssetPathWithTrailingSlash.EndsWith(TEXT("/"), ESearchCase::CaseSensitive))
-				{
-					AssetPathWithTrailingSlash += TEXT("/");
-				}
-
+				// Make sure OnDirectoryChangedDelegateHandles key is symmetrical with the one used in OnMountPointMounted
+				FString AssetPathWithTrailingSlash = UE::AssetRegistry::CreateStandardLongPackageNameDirectory(MountPoint->GetLongPackageName());
 				FDelegateHandle DirectoryChangedHandle;
 				if (ensure(OnDirectoryChangedDelegateHandles.RemoveAndCopyValue(AssetPathWithTrailingSlash, DirectoryChangedHandle)))
 				{
@@ -8914,51 +8944,57 @@ void UAssetRegistryImpl::OnMultipleContentPathsDismounted(TConstArrayView<TTuple
 namespace UE::AssetRegistry
 {
 
-void FAssetRegistryImpl::OnMultipleContentPathsDismounted(
+void FAssetRegistryImpl::OnMultipleMountPointsDismounted(UAssetRegistryImpl& UARI,
 	Impl::FEventContext& EventContext,
-	TConstArrayView<TTuple<FString, FString>> UnmountedAssetFilesystemPairs)
+	TConstArrayView<TRefCountPtr<UE::PackageName::IMountPoint>> DismountedMountPoints)
 {
 	// For the GlobalGatherer we remove all input MountPoints; we don't need to keep information about their gather.
 	// But for the FAssetData stored inside our AssetRegistryState, for some mount points
 	// we are configured to keep those assets in memory even when dismounted.
-	TSet<FName> MountPointsToRemoveAssetsFrom;
-	for (const TTuple<FString, FString>& Pair : UnmountedAssetFilesystemPairs)
+	TArray<TRefCountPtr<UE::PackageName::IMountPoint>> MountPointsToRemoveAssetsFrom;
+	MountPointsToRemoveAssetsFrom.Reserve(DismountedMountPoints.Num());
+	for (const TRefCountPtr<UE::PackageName::IMountPoint>& MountPoint : DismountedMountPoints)
 	{
 		if (GlobalGatherer.IsValid())
 		{
-			GlobalGatherer->RemoveMountPoint(Pair.Value);
+			GlobalGatherer->RemoveMountPoint(MountPoint->GetLongPackageName());
 		}
+
 		// The input AssetPaths of the unmounted MountPoints are unnormalized and might have a trailing slash,
 		// but we are comparing them by exact string inside our containers 
 		// and need to normalize them to not have a trailing slash.
-		FStringView AssetPathNoTrailingSlash = Pair.Key;
-		if (AssetPathNoTrailingSlash.EndsWith(TEXT("/"), ESearchCase::CaseSensitive))
-		{
-			// We don't want a trailing slash here as it could interfere with RemoveAssetPath
-			AssetPathNoTrailingSlash.LeftChopInline(1);
-		}
-		FName AssetPathNoTrailingSlashFName(AssetPathNoTrailingSlash);
+		FName AssetPathNoTrailingSlashFName = UE::AssetRegistry::CreateStandardLongPackageFNameNoTrailingSlash(
+			MountPoint->GetLongPackageName());
 		if (PersistentMountPoints.Contains(AssetPathNoTrailingSlashFName))
 		{
-			// This path is not marked to never remove its AssetDatas.
+			// This path is marked to never remove its AssetDatas.
 			continue;
 		}
 
-		MountPointsToRemoveAssetsFrom.Add(AssetPathNoTrailingSlashFName);
+		MountPointsToRemoveAssetsFrom.Add(MountPoint);
 	}
+	UnloadMountPoints(EventContext, MountPointsToRemoveAssetsFrom);
+}
 
+void FAssetRegistryImpl::UnloadMountPoints(Impl::FEventContext& EventContext,
+	TConstArrayView<TRefCountPtr<UE::PackageName::IMountPoint>> MountPoints)
+{
 	// Collect the list of paths from given list of roots
 	TSet<FName> AssetPathsToRemoveAssetsFrom;
+	TArray<FName> MountPointRootNoTrailingSlashFNames;
+
+	const bool bRecurse = true;
+	MountPointRootNoTrailingSlashFNames.Reserve(MountPoints.Num());
+	for (const TRefCountPtr<UE::PackageName::IMountPoint>& MountPoint : MountPoints)
 	{
-		const bool bRecurse = true;
-		for (FName AssetPathNoTrailingSlashFName : MountPointsToRemoveAssetsFrom)
-		{
-			CachedPathTree.GetSubPaths(AssetPathNoTrailingSlashFName, AssetPathsToRemoveAssetsFrom, bRecurse);
-			AssetPathsToRemoveAssetsFrom.Add(AssetPathNoTrailingSlashFName);
-		}
+		FName AssetPathNoTrailingSlashFName = UE::AssetRegistry::CreateStandardLongPackageFNameNoTrailingSlash(
+			MountPoint->GetLongPackageName());
+		CachedPathTree.GetSubPaths(AssetPathNoTrailingSlashFName, AssetPathsToRemoveAssetsFrom, bRecurse);
+		AssetPathsToRemoveAssetsFrom.Add(AssetPathNoTrailingSlashFName);
+		MountPointRootNoTrailingSlashFNames.Add(AssetPathNoTrailingSlashFName);
 	}
 
-	// Remove all cached assets and Verse files found at this location
+	// Remove all cached assets and Verse files found in the mountpoints.
 	{
 		TArray<FAssetData*> AllAssetDataToRemove;
 		for (FName PathName : AssetPathsToRemoveAssetsFrom)
@@ -8985,8 +9021,8 @@ void FAssetRegistryImpl::OnMultipleContentPathsDismounted(
 		RemoveAssetDatas(EventContext, AllAssetDataToRemove);
 	}
 
-	// Remove the root path
-	for (FName AssetPathNoTrailingSlashFName : MountPointsToRemoveAssetsFrom)
+	// Remove the root paths of the mountpoints.
+	for (FName AssetPathNoTrailingSlashFName : MountPointRootNoTrailingSlashFNames)
 	{
 		const bool bEvenIfAssetsStillExist = true;
 		RemoveAssetPath(EventContext, AssetPathNoTrailingSlashFName, bEvenIfAssetsStillExist);
@@ -9247,6 +9283,24 @@ bool UAssetRegistryImpl::IsDirAlreadyWatchedByRootWatchers(const FString& Direct
 
 #endif
 
+void UAssetRegistryImpl::SuspendPathsMounting()
+{
+	SuspendMountingCounter.fetch_add(1, std::memory_order_acq_rel);
+}
+
+void UAssetRegistryImpl::ResumePathsMounting()
+{
+	if (SuspendMountingCounter.fetch_sub(1, std::memory_order_acq_rel) == 1)
+	{
+		TArray<TRefCountPtr<UE::PackageName::IMountPoint>> LocalList;
+		{
+			FScopeLock Guard(&MountingCriticalSection);
+			LocalList = MoveTemp(ListOfMountedMountPoints);
+		}
+		OnMultipleMountPointsMounted(LocalList);
+	}
+}
+
 void UAssetRegistryImpl::SuspendPathsDismounting()
 {
 	SuspendDismountingCounter.fetch_add(1, std::memory_order_acq_rel);
@@ -9254,13 +9308,15 @@ void UAssetRegistryImpl::SuspendPathsDismounting()
 
 void UAssetRegistryImpl::ResumePathsDismounting()
 {
-	TArray<TTuple<FString, FString>> LocalList;
 	if (SuspendDismountingCounter.fetch_sub(1, std::memory_order_acq_rel) == 1)
 	{
-		FScopeLock Guard(&DismountingCriticalSection);
-		LocalList = MoveTemp(ListOfUnmountedPaths);
+		TArray<TRefCountPtr<UE::PackageName::IMountPoint>> LocalList;
+		{
+			FScopeLock Guard(&MountingCriticalSection);
+			LocalList = MoveTemp(ListOfDismountedMountPoints);
+		}
+		OnMultipleMountPointsDismounted(LocalList);
 	}
-	OnMultipleContentPathsDismounted(LocalList);
 }
 
 void UAssetRegistryImpl::RequestPauseBackgroundProcessing()
@@ -10851,6 +10907,11 @@ bool FInterruptionContext::ShouldExitEarly()
 
 }
 
+FString CreateStandardFilename(FStringView InPath)
+{
+	return CreateStandardFilename(FString(InPath));
+}
+
 FString CreateStandardFilename(const FString& InPath)
 {
 	FString Result = FPaths::CreateStandardFilename(InPath);
@@ -10869,4 +10930,26 @@ FString CreateStandardFilename(const FString& InPath)
 	return Result;
 }
 
+FString CreateStandardLongPackageNameDirectory(FStringView LongPackageName)
+{
+	FString Result(LongPackageName);
+	if (!Result.EndsWith(TEXT("/"), ESearchCase::CaseSensitive))
+	{
+		Result = Result + TEXT("/");
+	}
+	return Result;
+}
+
+FName CreateStandardLongPackageFNameNoTrailingSlash(FStringView LongPackageName)
+{
+	FStringView ResultNoTrailingSlash = LongPackageName;
+	if (ResultNoTrailingSlash.EndsWith(TEXT("/"), ESearchCase::CaseSensitive))
+	{
+		// We don't want a trailing slash here as it could interfere with RemoveAssetPath
+		ResultNoTrailingSlash.LeftChopInline(1);
+	}
+	return FName(ResultNoTrailingSlash);
+}
+
+
 } // namespace AssetRegistry
diff --git a/Engine/Source/Runtime/AssetRegistry/Private/AssetRegistry.h b/Engine/Source/Runtime/AssetRegistry/Private/AssetRegistry.h
--- a/Engine/Source/Runtime/AssetRegistry/Private/AssetRegistry.h
+++ b/Engine/Source/Runtime/AssetRegistry/Private/AssetRegistry.h
@@ -4,10 +4,12 @@
 
 #include "AssetRegistryImpl.h"
 #include "Containers/Ticker.h"
+#include "Misc/PackageName.h"
 #include "UObject/Object.h"
 
 #include "AssetRegistry.generated.h"
 
+class FAssetDataGatherer;
 namespace UE::AssetRegistry::Premade { struct FAsyncConsumer; }
 
 /**
@@ -110,7 +112,8 @@ public:
 	virtual void ScanModifiedAssetFiles(const TArray<FString>& InFilePaths) override;
 	virtual void ScanModifiedAssetFiles(const TArray<FString>& InFilePaths, UE::AssetRegistry::EScanFlags ScanFlags) override;
 	virtual void Serialize(FArchive& Ar) override;
-	virtual void AppendState(const FAssetRegistryState& InState, UE::AssetRegistry::EAppendMode AppendMode = UE::AssetRegistry::EAppendMode::Append) override;
+	virtual void AppendState(const FAssetRegistryState& InState,
+		UE::AssetRegistry::EAppendMode AppendMode = UE::AssetRegistry::EAppendMode::Append) override;
 	virtual SIZE_T GetAllocatedSize(bool bLogDetailed = false) const override;
 	virtual void LoadPackageRegistryData(FArchive& Ar, FLoadPackageRegistryData& InOutData) const override;
 	virtual void LoadPackageRegistryData(const FString& PackageFilename, FLoadPackageRegistryData& InOutData) const override;
@@ -250,29 +253,22 @@ private:
 #endif
 
 	/**
-	 * Called by the engine core when a new content path is added dynamically at runtime.  This is wired to
+	 * Called by the engine core when a new MountPoint is added dynamically at runtime.  This is wired to
 	 * FPackageName's static delegate.
-	 *
-	 * @param	AssetPath		The new content root asset path that was added (e.g. "/MyPlugin/")
-	 * @param	FileSystemPath	The filesystem path that the AssetPath is mapped to
 	 */
-	void OnContentPathMounted(const FString& AssetPath, const FString& FileSystemPath);
+	void OnMountPointMounted(const TRefCountPtr<UE::PackageName::IMountPoint>& MountPoint);
 
 	/**
-	 * Called by the engine core when a content path is removed dynamically at runtime.  This is wired to
+	 * Called by the engine core when a MountPoint is removed dynamically at runtime.  This is wired to
 	 * FPackageName's static delegate.
-	 *
-	 * @param	AssetPath		The new content root asset path that was added (e.g. "/MyPlugin/")
-	 * @param	FileSystemPath	The filesystem path that the AssetPath is mapped to
 	 */
-	void OnContentPathDismounted(const FString& AssetPath, const FString& FileSystemPath);
+	void OnMountPointDismounted(const TRefCountPtr<UE::PackageName::IMountPoint>& MountPoint);
 
-	/**
-	 * Batched processing version of the OnContentPathDismounted
-	 *
-	 * @param	UnmountedAssetFilesystemPairs	the list of pairs AssetPath - FilesystemPath
-	 */
-	void OnMultipleContentPathsDismounted(TConstArrayView<TTuple<FString, FString>> UnmountedAssetFilesystemPairs);
+	/** Batched processing version of OnMountPointMounted. */
+	void OnMultipleMountPointsMounted(TConstArrayView<TRefCountPtr<UE::PackageName::IMountPoint>> MountedMountPoints);
+
+	/** Batched processing version of the OnMountPointDismounted. */
+	void OnMultipleMountPointsDismounted(TConstArrayView<TRefCountPtr<UE::PackageName::IMountPoint>> DismountedMountPoints);
 
 	/** Called to refresh the native classes list, called at end of engine initialization. */
 	void OnPostEngineInit();
@@ -320,7 +316,9 @@ private:
 	bool IsDirAlreadyWatchedByRootWatchers(const FString& Directory) const;
 #endif
 
-	// Implementation methods for FSuspendPathsDismounting
+	// Implementation methods for FSuspendPathsMounting and FSuspendPathsDismounting
+	virtual void SuspendPathsMounting() override;
+	virtual void ResumePathsMounting() override;
 	virtual void SuspendPathsDismounting() override;
 	virtual void ResumePathsDismounting() override;
 
@@ -447,14 +445,17 @@ private:
 	FTSTicker::FDelegateHandle TickRequestHandle;
 	FCriticalSection DeferredEventsCriticalSection;
 
-	/** Handles suspension and resume of paths dismounting Used for batching many Dismount calls together. */
+	// Suspension and resuming of paths mounting and dismounting. Used for batching many
+	// OnMountPointMounted/OnMountPointDismounted calls together.
+	std::atomic<int32> SuspendMountingCounter = 0;
 	std::atomic<int32> SuspendDismountingCounter = 0;
-	FCriticalSection DismountingCriticalSection;
-	TArray<TTuple<FString, FString>> ListOfUnmountedPaths;
+	FCriticalSection MountingCriticalSection;
+	TArray<TRefCountPtr<UE::PackageName::IMountPoint>> ListOfMountedMountPoints;
+	TArray<TRefCountPtr<UE::PackageName::IMountPoint>> ListOfDismountedMountPoints;
 
 
-	friend class FAssetDataGatherer;
-	friend class UE::AssetRegistry::FAssetRegistryImpl;
-	friend struct UE::AssetRegistry::Premade::FAsyncConsumer;
+	friend FAssetDataGatherer;
+	friend UE::AssetRegistry::FAssetRegistryImpl;
+	friend UE::AssetRegistry::Premade::FAsyncConsumer;
 };
 PRAGMA_ENABLE_DEPRECATION_WARNINGS
diff --git a/Engine/Source/Runtime/AssetRegistry/Private/AssetRegistryImpl.h b/Engine/Source/Runtime/AssetRegistry/Private/AssetRegistryImpl.h
--- a/Engine/Source/Runtime/AssetRegistry/Private/AssetRegistryImpl.h
+++ b/Engine/Source/Runtime/AssetRegistry/Private/AssetRegistryImpl.h
@@ -31,14 +31,13 @@ struct FFileChangeData;
 class FPackageReader;
 class UAssetRegistryImpl;
 namespace UE::AssetRegistry::DependsNode { struct FGatherDependencyDataScratchpad; }
-namespace UE::AssetRegistry::Impl { struct FEventContext; }
 namespace UE::AssetRegistry::Premade { enum class ELoadResult : uint8; }
 namespace UE::AssetRegistry::Premade { struct FAsyncConsumer; }
+namespace UE::AssetRegistry::Impl { struct FEventContext; }
 namespace UE::AssetRegistry::Impl { struct FInitializeContext; }
 namespace UE::AssetRegistry::Impl { struct FScanPathContext; }
 namespace UE::AssetRegistry::Impl { struct FTickContext; }
 namespace UE::AssetRegistry::Impl { struct FClassInheritanceContext; }
-
 #if WITH_EDITOR
 namespace UE::AssetDependencyGatherer::Private { class FRegisteredAssetDependencyGatherer; }
 #endif
@@ -257,9 +256,12 @@ public:
 		Impl::FClassInheritanceContext& InheritanceContext, FString& DirPath, int64 BeforeTimeStamp);
 	void AddLoadedAssetToProcess(const UObject& AssetLoaded);
 #endif
-	void OnContentPathMounted(UE::AssetRegistry::FInterfaceWriteScopeLock& ScopeLock, Impl::FEventContext& EventContext, Impl::FClassInheritanceContext& InheritanceContext,
-		const FString& InAssetPath, const FString& AssetPathWithTrailingSlash, const FString& FileSystemPath);
-	void OnMultipleContentPathsDismounted(Impl::FEventContext& EventContext, TConstArrayView<TTuple<FString, FString>> UnmountedAssetFilesystemPairs);
+	void OnMultipleMountPointsMounted(UAssetRegistryImpl& UARI, Impl::FEventContext& EventContext,
+		Impl::FClassInheritanceContext& InheritanceContext,
+		UE::AssetRegistry::FInterfaceWriteScopeLock& ScopeLock,
+		TConstArrayView<TRefCountPtr<UE::PackageName::IMountPoint>> MountedMountPoints);
+	void OnMultipleMountPointsDismounted(UAssetRegistryImpl& UARI, Impl::FEventContext& EventContext,
+		TConstArrayView<TRefCountPtr<UE::PackageName::IMountPoint>> DismountedMountPoints);
 
 	// Other helper functions called by UAssetRegistryImpl
 
@@ -405,7 +407,12 @@ public:
 	 * For performance reasons, this must be called within the WriteScopeLock, and the caller 
 	 * must handle the possibility that it leaves and reenters the lock.
 	 */
-	void ConditionalLoadPremadeAssetRegistry(UE::AssetRegistry::FInterfaceWriteScopeLock& ScopeLock, UE::AssetRegistry::Impl::FEventContext& EventContext);
+	void ConditionalLoadPremadeAssetRegistry(UE::AssetRegistry::FInterfaceWriteScopeLock& ScopeLock,
+		UE::AssetRegistry::Impl::FEventContext& EventContext);
+
+	/** Unload specified MountPoints, helper for OnMultipleMountPointsDismounted and others. */
+	void UnloadMountPoints(Impl::FEventContext& EventContext,
+		TConstArrayView<TRefCountPtr<UE::PackageName::IMountPoint>> MountPoints);
 
 #if WITH_EDITOR
 	/** Request to pause or resume background processing of scan results.
@@ -528,10 +535,12 @@ private:
 	void UpdateInheritanceBuffer(Impl::FClassInheritanceBuffer& OutBuffer) const;
 
 	/** Conditionally loads the Premade AssetRegistry from disk, or queues it for asynchronous loading if not yet loaded. */
-	void ConsumeOrDeferPreloadedPremade(UAssetRegistryImpl& UARI, Impl::FEventContext& EventContext);
+	void ConsumeOrDeferPreloadedPremade(UAssetRegistryImpl& UARI, Impl::FEventContext& EventContext,
+		FInterfaceWriteScopeLock& ProofOfWriteLock);
 	/** Moves a premade asset registry state into this AR */
-	void LoadPremadeAssetRegistry(Impl::FEventContext& Context,
-		Premade::ELoadResult LoadResult, FAssetRegistryState&& ARState);
+	void LoadPremadeAssetRegistry(UAssetRegistryImpl& UARI, Impl::FEventContext& Context,
+		FInterfaceWriteScopeLock& ProofOfWriteLock,
+		Premade::ELoadResult LoadResult, FAssetRegistryState&& ARState, FStringView StateFilepath);
 	/** Add MountPoints of all AssetDatas currently registered in this->State to the list of PersistentMountPoints. */
 	void UpdatePersistentMountPoints();
 	void OnInitialSearchCompleted(Impl::FEventContext& EventContext);
@@ -1012,7 +1021,12 @@ void EnumerateMemoryAssets(const FARCompiledFilter& InFilter, TSet<FName>& OutPa
 // Normalization functions for paths. We require that normalized paths are case-insensitively string-identical for the same input
 // path specification. e.g. all slashes must be replaced with forward-slashes. Special case: a leading \\ is left unreplaced,
 // so //share will be seen as different than \\share.
+FString CreateStandardFilename(FStringView InPath);
 FString CreateStandardFilename(const FString& InFilename);
+// Normalization for LongPackageNames that identify a directory. Ensures it ends with a trailing forward-slash.
+FString CreateStandardLongPackageNameDirectory(FStringView LongPackageName);
+// Normalization for LongPackageNames that should not end in a trailing slash.
+FName CreateStandardLongPackageFNameNoTrailingSlash(FStringView LongPackageName);
 
 } // namespace UE::AssetRegistry
 
diff --git a/Engine/Source/Runtime/AssetRegistry/Public/AssetRegistry/IAssetRegistry.h b/Engine/Source/Runtime/AssetRegistry/Public/AssetRegistry/IAssetRegistry.h
--- a/Engine/Source/Runtime/AssetRegistry/Public/AssetRegistry/IAssetRegistry.h
+++ b/Engine/Source/Runtime/AssetRegistry/Public/AssetRegistry/IAssetRegistry.h
@@ -1155,7 +1155,20 @@ public:
 	DECLARE_TS_MULTICAST_DELEGATE_TwoParams(FEnumerateAssetsEvent, const FARCompiledFilter& /*Filter*/, UE::AssetRegistry::EEnumerateAssetsFlags /*Flag*/)
 	virtual FEnumerateAssetsEvent& OnEnumerateAssetsEvent() = 0;
 
-	/** Request to pause paths dismounting.
+	/**
+	 * Request to pause paths mounting.
+	 * Can be used by the calling code to temporarily pause loading of data inside the AssetRegistry in response
+	 * to mounted MountPoints (e.g. ActiveMountLoading). All mounting requests will be executed in one batch at the
+	 * end of that scope.
+	 */
+	struct FSuspendPathsMounting
+	{
+		UE_NODISCARD_CTOR FSuspendPathsMounting() { IAssetRegistry::GetChecked().SuspendPathsMounting(); }
+		~FSuspendPathsMounting() { IAssetRegistry::GetChecked().ResumePathsMounting(); }
+	};
+
+	/**
+	 * Request to pause paths dismounting.
 	 * Can be used by the calling code to temporarily pause clearing of data inside the AssetRegistry.
 	 * All unmounting requests will be executed in one batch at the end of that scope.
 	 */
@@ -1167,7 +1180,9 @@ public:
 
 protected:
 
-	// Implementation methods for FSuspendPathsDismounting
+	// Implementation methods for FSuspendPathsMounting and FSuspendPathsDismounting
+	virtual void SuspendPathsMounting() {};
+	virtual void ResumePathsMounting() {};
 	virtual void SuspendPathsDismounting() {};
 	virtual void ResumePathsDismounting() {};
 
diff --git a/Engine/Source/Runtime/CoreUObject/Private/Misc/PackageName.cpp b/Engine/Source/Runtime/CoreUObject/Private/Misc/PackageName.cpp
--- a/Engine/Source/Runtime/CoreUObject/Private/Misc/PackageName.cpp
+++ b/Engine/Source/Runtime/CoreUObject/Private/Misc/PackageName.cpp
@@ -51,10 +51,10 @@ DEFINE_LOG_CATEGORY(LogPackageName);
 
 #define LOCTEXT_NAMESPACE "PackageNames"
 
-/** Event that is triggered when a new content path is mounted */
+FPackageName::FMountPointEvent FPackageName::OnMountPointMountedEvent;
 FPackageName::FOnContentPathMountedEvent FPackageName::OnContentPathMountedEvent;
 
-/** Event that is triggered when a content path is dismounted */
+FPackageName::FMountPointEvent FPackageName::OnMountPointDismountedEvent;
 FPackageName::FOnContentPathDismountedEvent FPackageName::OnContentPathDismountedEvent;
 
 /** Delegate used to check whether a package exist without using the filesystem. */
@@ -175,17 +175,26 @@ FLongPackagePathsSingleton& FLongPackagePathsSingleton::Get()
 	return Singleton;
 }
 
+/**
+ * Function to bind to IPluginManager::FRegisterMountPointDelegate, which expects a function with the same
+ * arguments as RegisterMountPoint but with a void return value.
+ */
+static void RegisterMountPointForPluginManager(const FString& RootPath, const FString& ContentPath)
+{
+	(void)FPackageName::RegisterMountPoint(RootPath, ContentPath);
+}
+
 void FLongPackagePathsSingleton::OnCoreUObjectInitialized()
 {
 	// Allow the plugin manager to mount new content paths by exposing access through a delegate.  PluginManager is 
 	// a Core class, but content path functionality is added at the CoreUObject level.
 	IPluginManager& PluginManager = IPluginManager::Get();
-	PluginManager.SetRegisterMountPointDelegate(IPluginManager::FRegisterMountPointDelegate::CreateStatic(&FPackageName::RegisterMountPoint));
+	PluginManager.SetRegisterMountPointDelegate(IPluginManager::FRegisterMountPointDelegate::CreateStatic(&RegisterMountPointForPluginManager));
 	PluginManager.SetUnRegisterMountPointDelegate(IPluginManager::FRegisterMountPointDelegate::CreateStatic(&FPackageName::UnRegisterMountPoint));
 	UE::CoreUObject::Private::PluginHandler::Install();
 }
 
-bool FLongPackagePathsSingleton::TryNormalizePaths(const FString& InLongPackageName, const FString& InLocalPath,
+bool FLongPackagePathsSingleton::TryNormalizePaths(FStringView InLongPackageName, FStringView InLocalPath,
 	FString& OutLongPackageName, FString& OutLocalPathStandard, FString& OutLocalPathAbsolute, bool bSilenceLogs)
 {
 	if (!TryNormalizeLongPackageNamePath(InLongPackageName, OutLongPackageName, bSilenceLogs))
@@ -199,25 +208,26 @@ bool FLongPackagePathsSingleton::TryNormalizePaths(const FString& InLongPackageN
 	return true;
 }
 
-bool FLongPackagePathsSingleton::TryNormalizeLongPackageNamePath(const FString& InLongPackageName,
+bool FLongPackagePathsSingleton::TryNormalizeLongPackageNamePath(FStringView InLongPackageName,
 	FString& OutLongPackageName, bool bSilenceLogs)
 {
 	int32 NumBackSlashes = 0;
 	int32 LenLongPackageName = InLongPackageName.Len();
 	for (int32 N = 0; N < LenLongPackageName; ++N)
 	{
-		TCHAR C = (*InLongPackageName)[N];
+		TCHAR C = InLongPackageName.GetData()[N];
 		if (C == '\\')
 		{
 			++NumBackSlashes;
 		}
 	}
-	if (!InLongPackageName.StartsWith("/") || NumBackSlashes > 0)
+	if (LenLongPackageName == 0 || InLongPackageName[0] != '/' || NumBackSlashes > 0)
 	{
 		if (!bSilenceLogs)
 		{
 			UE_LOG(LogPackageName, Error,
-				TEXT("Invalid LongPackageName for mount point: \"%s\". Callstack:"), *InLongPackageName);
+				TEXT("Invalid LongPackageName for mount point: \"%.*s\". Callstack:"),
+				InLongPackageName.Len(), InLongPackageName.GetData());
 			FDebug::DumpStackTraceToLog(ELogVerbosity::Display);
 		}
 		return false;
@@ -231,14 +241,14 @@ bool FLongPackagePathsSingleton::TryNormalizeLongPackageNamePath(const FString&
 	return true;
 }
 
-bool FLongPackagePathsSingleton::TryNormalizeLocalPath(const FString& InLocalPath,
+bool FLongPackagePathsSingleton::TryNormalizeLocalPath(FStringView InLocalPath,
 	FString& OutLocalPathStandard, FString& OutLocalPathAbsolute, bool bSilenceLogs)
 {
 	// If a relative path is passed, convert to an absolute path 
 	OutLocalPathAbsolute = InLocalPath;
 	if (FPaths::IsRelative(OutLocalPathAbsolute))
 	{
-		OutLocalPathAbsolute = FPaths::ConvertRelativePathToFull(InLocalPath);
+		OutLocalPathAbsolute = FPaths::ConvertRelativePathToFull(OutLocalPathAbsolute);
 
 		// Revert to original path if unable to convert to full path
 		if (OutLocalPathAbsolute.Len() <= 1)
@@ -246,7 +256,8 @@ bool FLongPackagePathsSingleton::TryNormalizeLocalPath(const FString& InLocalPat
 			OutLocalPathAbsolute = InLocalPath;
 			if (!bSilenceLogs)
 			{
-				UE_LOG(LogPackageName, Warning, TEXT("Unable to convert mount point relative path: %s"), *InLocalPath);
+				UE_LOG(LogPackageName, Warning, TEXT("Unable to convert mount point relative path: %.*s"),
+					InLocalPath.Len(), InLocalPath.GetData());
 			}
 		}
 	}
@@ -269,7 +280,7 @@ bool FLongPackagePathsSingleton::TryNormalizeLocalPath(const FString& InLocalPat
 	return true;
 }
 
-void FLongPackagePathsSingleton::InsertMountPoint(const FString& RootPath, const FString& ContentPath)
+TRefCountPtr<UE::PackageName::IMountPoint> FLongPackagePathsSingleton::InsertMountPoint(const FString& RootPath, const FString& ContentPath)
 {
 	using namespace UE::PackageName;
 
@@ -278,11 +289,12 @@ void FLongPackagePathsSingleton::InsertMountPoint(const FString& RootPath, const
 	if (!TryNormalizePaths(RootPath, ContentPath, SearchKey.LongPackageName,
 		SearchKey.LocalPathStandard, LocalPathAbsolute, false /* bSilenceLogs */))
 	{
-		return;
+		return TRefCountPtr<UE::PackageName::IMountPoint>();
 	}
 	uint32 SearchHash = GetTypeHash(SearchKey);
 
-	TRefCountPtr<FMountPoint> MountPoint;
+	TRefCountPtr<IMountPoint> PublicMountPoint;
+	FMountPoint* MountPoint = nullptr;
 	{
 		UE::TWriteScopeLock ScopeLock(MountLock);
 
@@ -298,12 +310,13 @@ void FLongPackagePathsSingleton::InsertMountPoint(const FString& RootPath, const
 				FString(LocalPathAbsolute), EMountFlags::Removable);
 			MountPoints.AddByHash(SearchHash, FMountPointKeyByPtr(*MountPoint));
 		}
-		RootPathTree.FindOrAdd(MountPoint->GetRootPath()) = MountPoint.GetReference();
-		ContentPathTree.FindOrAdd(MountPoint->GetContentPathRelative()) = MountPoint.GetReference();
-		ContentPathTree.FindOrAdd(MountPoint->GetContentPathAbsolute()) = MountPoint.GetReference();
+		PublicMountPoint = MountPoint;
+		RootPathTree.FindOrAdd(MountPoint->GetRootPath()) = MountPoint;
+		ContentPathTree.FindOrAdd(MountPoint->GetContentPathRelative()) = MountPoint;
+		ContentPathTree.FindOrAdd(MountPoint->GetContentPathAbsolute()) = MountPoint;
 
 		bool bRemoveFromContainer;
-		MountPoint->SetMounted(true, ScopeLock, MountPoint, bRemoveFromContainer);
+		MountPoint->SetMounted(true, ScopeLock, PublicMountPoint, bRemoveFromContainer);
 		// Set mount order to highest priority, replacing earlier mount order if it already existed.
 		MountPoint->SetMountOrder(NextMountOrder++);
 		if (NextMountOrder == 0)
@@ -314,10 +327,21 @@ void FLongPackagePathsSingleton::InsertMountPoint(const FString& RootPath, const
 		check(!bRemoveFromContainer); // Should only be settable to true when calling SetMounted(false).
 	}
 
+	BroadcastMountPointMounted(PublicMountPoint, MountPoint);
+	return PublicMountPoint;
+}
+
+void FLongPackagePathsSingleton::BroadcastMountPointMounted(
+	const TRefCountPtr<UE::PackageName::IMountPoint>& PublicMountPoint, UE::PackageName::FMountPoint* MountPoint)
+{
+	using namespace UE::PackageName;
+
+	check(PublicMountPoint && PublicMountPoint.GetReference() == MountPoint);
+
 	// Downgrade the log messages on startup to Verbose to reduce startup spam
 #if !NO_LOGGING
 	FString Message = FString::Printf(TEXT("FPackageName: Mount point added: '%s' mounted to '%s'"),
-		*SearchKey.LocalPathStandard, *SearchKey.LongPackageName);
+		*MountPoint->GetContentPathRelative(), *MountPoint->GetRootPath());
 	if (IsEngineStartupModuleLoadingComplete())
 	{
 		UE_LOG(LogPackageName, Log, TEXT("%s"), *Message);
@@ -329,7 +353,50 @@ void FLongPackagePathsSingleton::InsertMountPoint(const FString& RootPath, const
 #endif
 
 	// Let subscribers know that a new content path was mounted
-	FPackageName::OnContentPathMounted().Broadcast( SearchKey.LongPackageName, SearchKey.LocalPathStandard);
+	FPackageName::OnMountPointMounted().Broadcast(PublicMountPoint);
+	FPackageName::OnContentPathMounted().Broadcast(MountPoint->GetRootPath(), MountPoint->GetContentPathRelative());
+}
+
+void FLongPackagePathsSingleton::InsertMountPoint(const TRefCountPtr<UE::PackageName::IMountPoint>& PublicMountPoint)
+{
+	using namespace UE::PackageName;
+
+	if (!PublicMountPoint)
+	{
+		return;
+	}
+	FMountPoint* MountPoint = PublicMountPoint->AsFMountPoint();
+	if (!MountPoint)
+	{
+		UE_LOG(LogPackageName, Error,
+			TEXT("IMountPoint passed to FPackageName that is not an FMountPoint subclass. This is unexpected, all IMountPoints should be FMountPoint. ")
+			TEXT("LongPackageName: %.*s, LocalPathStandard: %.*s"),
+			PublicMountPoint->GetLongPackageName().Len(), PublicMountPoint->GetLongPackageName().GetData(),
+			PublicMountPoint->GetLocalPathStandard().Len(), PublicMountPoint->GetLocalPathStandard().GetData());
+		FDebug::DumpStackTraceToLog(ELogVerbosity::Display);
+		return;
+	}
+
+	{
+		UE::TWriteScopeLock ScopeLock(MountLock);
+
+		RootPathTree.FindOrAdd(MountPoint->GetRootPath()) = MountPoint;
+		ContentPathTree.FindOrAdd(MountPoint->GetContentPathRelative()) = MountPoint;
+		ContentPathTree.FindOrAdd(MountPoint->GetContentPathAbsolute()) = MountPoint;
+
+		bool bRemoveFromContainer;
+		MountPoint->SetMounted(true, ScopeLock, PublicMountPoint, bRemoveFromContainer);
+		// Set mount order to highest priority, replacing earlier mount order if it already existed.
+		MountPoint->SetMountOrder(NextMountOrder++);
+		if (NextMountOrder == 0)
+		{
+			// Overflow, due to e.g. removing and adding a mountpoint too many times.
+			ReassignMountOrders();
+		}
+		check(!bRemoveFromContainer); // Should only be settable to true when calling SetMounted(false).
+	}
+
+	BroadcastMountPointMounted(PublicMountPoint, MountPoint);
 }
 
 void FLongPackagePathsSingleton::ReassignMountOrders()
@@ -373,8 +440,8 @@ void FLongPackagePathsSingleton::RemoveMountPoint(const FString& RootPath, const
 	}
 	uint32 SearchHash = GetTypeHash(SearchKey);
 
-	TRefCountPtr<FMountPoint> ExistingMount;
-	bool bUnmounted = false;
+	TRefCountPtr<IMountPoint> PublicMountPoint;
+	FMountPoint* MountPoint = nullptr;
 	{
 		UE::TWriteScopeLock ScopeLock(MountLock);
 		FMountPointKeyByPtr* ExistingPtr = MountPoints.FindByHash(SearchHash, SearchKey);
@@ -385,93 +452,328 @@ void FLongPackagePathsSingleton::RemoveMountPoint(const FString& RootPath, const
 				*SearchKey.LocalPathStandard, *SearchKey.LongPackageName);
 			return;
 		}
-		ExistingMount = ExistingPtr->GetMountPoint();
-		if (!ExistingMount->IsRemovable())
+		MountPoint = ExistingPtr->GetMountPoint();
+		PublicMountPoint = MountPoint;
+		if (!TryRemoveMountPointInternal(PublicMountPoint, ScopeLock, MountPoint))
 		{
-			UE_LOG(LogPackageName, Display, TEXT("FPackageName: Mount point remove failed: MountPoint mapping '%s' to '%s' is a default mapping and is not removable."),
-				*SearchKey.LocalPathStandard, *SearchKey.LongPackageName);
 			return;
 		}
+	}
+	BroadcastMountPointDismounted(PublicMountPoint, MountPoint);
+}
 
-		bool bRootPathExisted;
-		bool bContentPathRelativeExisted;
-		bool bContentPathAbsoluteExisted;
-		RootPathTree.Remove(ExistingMount->GetRootPath(), &bRootPathExisted);
-		ContentPathTree.Remove(ExistingMount->GetContentPathRelative(), &bContentPathRelativeExisted);
-		ContentPathTree.Remove(ExistingMount->GetContentPathAbsolute(), &bContentPathAbsoluteExisted);
-		if (bRootPathExisted || bContentPathRelativeExisted || bContentPathAbsoluteExisted)
-		{
-			// See if there is an existing mounted MountPoint that was overridden by the removed mount point and now
-			// needs to be published as the MountPoint for one of its paths.
-			FMountPoint* NewRootPath = nullptr;
-			FMountPoint* NewContentPathRelative = nullptr;
-			FMountPoint* NewContentPathAbsolute = nullptr;
-			for (const FMountPointKeyByPtr& OtherMountPtr : MountPoints)
+bool FLongPackagePathsSingleton::TryRemoveMountPointInternal(
+	const TRefCountPtr<UE::PackageName::IMountPoint>& InternallyHeldMountPoint,
+	UE::TWriteScopeLock<FTransactionallySafeRWLock>& ProofOfWriteLock, UE::PackageName::FMountPoint* MountPoint)
+{
+	using namespace UE::PackageName;
+
+	// Called within the WriteLock
+	// Caller must hold a single internal refcount in addition to any externally held refcount; that single internal
+	// refcount is required by the SetMounted call within this function to decide whether the MountPoint is externally referenced
+	if (!MountPoint->IsRemovable())
+	{
+		UE_LOG(LogPackageName, Display, TEXT("FPackageName: Mount point remove failed: MountPoint mapping '%s' to '%s' is a default mapping and is not removable."),
+			*MountPoint->GetContentPathRelative(), *MountPoint->GetRootPath());
+		return false;
+	}
+
+	bool bRootPathExisted;
+	bool bContentPathRelativeExisted;
+	bool bContentPathAbsoluteExisted;
+	RootPathTree.Remove(MountPoint->GetRootPath(), &bRootPathExisted);
+	ContentPathTree.Remove(MountPoint->GetContentPathRelative(), &bContentPathRelativeExisted);
+	ContentPathTree.Remove(MountPoint->GetContentPathAbsolute(), &bContentPathAbsoluteExisted);
+	if (bRootPathExisted || bContentPathRelativeExisted || bContentPathAbsoluteExisted)
+	{
+		// See if there is an existing mounted MountPoint that was overridden by the removed mount point and now
+		// needs to be published as the MountPoint for one of its paths.
+		FMountPoint* NewRootPath = nullptr;
+		FMountPoint* NewContentPathRelative = nullptr;
+		FMountPoint* NewContentPathAbsolute = nullptr;
+		for (const FMountPointKeyByPtr& OtherMountPtr : MountPoints)
+		{
+			FMountPoint* OtherMount = OtherMountPtr.GetMountPoint();
+			CA_ASSUME(OtherMount != nullptr);
+			if (OtherMount == MountPoint || !OtherMount->IsMounted())
 			{
-				FMountPoint* OtherMount = OtherMountPtr.GetMountPoint();
-				CA_ASSUME(OtherMount != nullptr);
-				if (OtherMount == ExistingMount.GetReference()
-					|| !OtherMount->IsMounted())
-				{
-					continue;
-				}
-				if (bRootPathExisted && OtherMount->GetRootPath().Equals(
-					SearchKey.LongPackageName, ESearchCase::IgnoreCase))
-				{
-					if (NewRootPath == nullptr
-						|| OtherMount->GetMountOrder() > NewRootPath->GetMountOrder())
-					{
-						NewRootPath = OtherMount;
-					}
-				}
-				if (bContentPathRelativeExisted && OtherMount->GetContentPathRelative().Equals(
-					SearchKey.LocalPathStandard, ESearchCase::IgnoreCase))
+				continue;
+			}
+			if (bRootPathExisted && OtherMount->GetRootPath().Equals(
+				MountPoint->GetRootPath(), ESearchCase::IgnoreCase))
+			{
+				if (NewRootPath == nullptr
+					|| OtherMount->GetMountOrder() > NewRootPath->GetMountOrder())
 				{
-					if (NewContentPathRelative == nullptr
-						|| OtherMount->GetMountOrder() > NewContentPathRelative->GetMountOrder())
-					{
-						NewContentPathRelative = OtherMount;
-					}
+					NewRootPath = OtherMount;
 				}
-				if (bContentPathAbsoluteExisted && OtherMount->GetContentPathAbsolute().Equals(
-					LocalPathAbsolute, ESearchCase::IgnoreCase))
+			}
+			if (bContentPathRelativeExisted && OtherMount->GetContentPathRelative().Equals(
+				MountPoint->GetContentPathRelative(), ESearchCase::IgnoreCase))
+			{
+				if (NewContentPathRelative == nullptr
+					|| OtherMount->GetMountOrder() > NewContentPathRelative->GetMountOrder())
 				{
-					if (NewContentPathAbsolute == nullptr
-						|| OtherMount->GetMountOrder() > NewContentPathAbsolute->GetMountOrder())
-					{
-						NewContentPathAbsolute = OtherMount;
-					}
+					NewContentPathRelative = OtherMount;
 				}
 			}
-			if (NewRootPath)
+			if (bContentPathAbsoluteExisted && OtherMount->GetContentPathAbsolute().Equals(
+				MountPoint->GetContentPathAbsolute(), ESearchCase::IgnoreCase))
 			{
-				RootPathTree.FindOrAdd(NewRootPath->GetRootPath()) = NewRootPath;
+				if (NewContentPathAbsolute == nullptr
+					|| OtherMount->GetMountOrder() > NewContentPathAbsolute->GetMountOrder())
+				{
+					NewContentPathAbsolute = OtherMount;
+				}
 			}
-			if (NewContentPathRelative)
+		}
+		if (NewRootPath)
+		{
+			RootPathTree.FindOrAdd(NewRootPath->GetRootPath()) = NewRootPath;
+		}
+		if (NewContentPathRelative)
+		{
+			ContentPathTree.FindOrAdd(NewContentPathRelative->GetContentPathRelative()) = NewContentPathRelative;
+		}
+		if (NewContentPathAbsolute)
+		{
+			ContentPathTree.FindOrAdd(NewContentPathAbsolute->GetContentPathAbsolute()) = NewContentPathAbsolute;
+		}
+	}
+
+	bool bRemoveFromContainer;
+	MountPoint->SetMounted(false, ProofOfWriteLock, InternallyHeldMountPoint, bRemoveFromContainer);
+	if (bRemoveFromContainer)
+	{
+		MountPoints.Remove(FMountPointKeyByPtr(*MountPoint));
+	}
+
+	return true;
+}
+
+void FLongPackagePathsSingleton::BroadcastMountPointDismounted(
+	const TRefCountPtr<UE::PackageName::IMountPoint>& PublicMountPoint, UE::PackageName::FMountPoint* MountPoint)
+{
+	using namespace UE::PackageName;
+
+	check(PublicMountPoint && PublicMountPoint.GetReference() == MountPoint);
+	UE_LOG(LogPackageName, Display, TEXT("FPackageName: Mount point removed: '%s' unmounted from '%s'"),
+		*MountPoint->GetContentPathRelative(), *MountPoint->GetRootPath());
+	// Let subscribers know that a content path was unmounted
+	FPackageName::OnMountPointDismounted().Broadcast(PublicMountPoint);
+	FPackageName::OnContentPathDismounted().Broadcast(MountPoint->GetRootPath(), MountPoint->GetContentPathRelative());
+}
+
+void FLongPackagePathsSingleton::RemoveMountPoint(const TRefCountPtr<UE::PackageName::IMountPoint>& PublicMountPoint)
+{
+	using namespace UE::PackageName;
+
+	if (!PublicMountPoint)
+	{
+		return;
+	}
+	FMountPoint* MountPoint = PublicMountPoint->AsFMountPoint();
+	if (!MountPoint)
+	{
+		UE_LOG(LogPackageName, Error,
+			TEXT("IMountPoint passed to FPackageName that is not an FMountPoint subclass. This is unexpected, all IMountPoints should be FMountPoint. ")
+			TEXT("LongPackageName: %.*s, LocalPathStandard: %.*s"),
+			PublicMountPoint->GetLongPackageName().Len(), PublicMountPoint->GetLongPackageName().GetData(),
+			PublicMountPoint->GetLocalPathStandard().Len(), PublicMountPoint->GetLocalPathStandard().GetData());
+		FDebug::DumpStackTraceToLog(ELogVerbosity::Display);
+		return;
+	}
+
+	// The SetMounted function assumes we take an additional refcount in this function, and uses that to decide
+	// that the MountPoint is no longer externally referenced. It is externally referenced (our caller passed in
+	// a refcountptr to it that they still have), so we need to increment refcount to let SetMounted know that
+	// it is externally referenced.
+	TRefCountPtr<IMountPoint> InternallyHeldMountPoint = PublicMountPoint;
+	{
+		UE::TWriteScopeLock ScopeLock(MountLock);
+		if (!TryRemoveMountPointInternal(InternallyHeldMountPoint, ScopeLock, MountPoint))
+		{
+			return;
+		}
+	}
+	BroadcastMountPointDismounted(InternallyHeldMountPoint, MountPoint);
+}
+
+TRefCountPtr<UE::PackageName::IMountPoint>
+FLongPackagePathsSingleton::FindMountPointByRootPackageName(FStringView RootLongPackageName)
+{
+	using namespace UE::PackageName;
+
+	FString LongPackageName;
+	if (!TryNormalizeLongPackageNamePath(RootLongPackageName, LongPackageName, true /* bSilenceLogs */))
+	{
+		return TRefCountPtr<UE::PackageName::IMountPoint>();
+	}
+
+	{
+		UE::TReadScopeLock ScopeLock(MountLock);
+
+		// First look for a Mounted MountPoint because its a common case that we can do fast
+		FMountPoint** ExistingPtr = RootPathTree.Find(LongPackageName);
+		if (ExistingPtr)
+		{
+			return TRefCountPtr<IMountPoint>(*ExistingPtr);
+		}
+
+		// We don't have a lookup map by RootName for overridden or non-mounted MountPoints, so search our entire list.
+		FMountPoint* HighestPriorityExisting = nullptr;
+		for (const FMountPointKeyByPtr& ExistingKey : MountPoints)
+		{
+			FMountPoint* Existing = ExistingKey.GetMountPoint();
+			if (Existing->GetLongPackageName().Equals(LongPackageName, ESearchCase::IgnoreCase))
 			{
-				ContentPathTree.FindOrAdd(NewContentPathRelative->GetContentPathRelative()) = NewContentPathRelative;
+				if (FMountPoint::IsHigherOverridePriority(Existing, HighestPriorityExisting))
+				{
+					HighestPriorityExisting = Existing;
+				}
 			}
-			if (NewContentPathAbsolute)
+		}
+		if (HighestPriorityExisting)
+		{
+			return TRefCountPtr<IMountPoint>(HighestPriorityExisting);
+		}
+	}
+	return TRefCountPtr<UE::PackageName::IMountPoint>();
+}
+
+TRefCountPtr<UE::PackageName::IMountPoint>
+FLongPackagePathsSingleton::FindMountPointByRootLocalPath(FStringView RootLocalPath)
+{
+	using namespace UE::PackageName;
+
+	FString LocalPathStandard;
+	FString LocalPathAbsolute;
+	if (!TryNormalizeLocalPath(RootLocalPath, LocalPathStandard, LocalPathAbsolute, true /* bSilenceLogs */))
+	{
+		return TRefCountPtr<UE::PackageName::IMountPoint>();
+	}
+
+	{
+		UE::TReadScopeLock ScopeLock(MountLock);
+
+		// First look for a Mounted MountPoint because its a common case that we can do fast
+		FMountPoint** ExistingPtr = ContentPathTree.Find(LocalPathStandard);
+		if (ExistingPtr)
+		{
+			return TRefCountPtr<IMountPoint>(*ExistingPtr);
+		}
+
+		// We don't have a lookup map by LocalPath for overridden or non-mounted MountPoints, so search our entire list.
+		FMountPoint* HighestPriorityExisting = nullptr;
+		for (const FMountPointKeyByPtr& ExistingKey : MountPoints)
+		{
+			FMountPoint* Existing = ExistingKey.GetMountPoint();
+			if (Existing->GetLocalPathStandard().Equals(LocalPathStandard, ESearchCase::IgnoreCase))
 			{
-				ContentPathTree.FindOrAdd(NewContentPathAbsolute->GetContentPathAbsolute()) = NewContentPathAbsolute;
+				if (FMountPoint::IsHigherOverridePriority(Existing, HighestPriorityExisting))
+				{
+					HighestPriorityExisting = Existing;
+				}
 			}
 		}
+		if (HighestPriorityExisting)
+		{
+			return TRefCountPtr<IMountPoint>(HighestPriorityExisting);
+		}
+	}
+	return TRefCountPtr<UE::PackageName::IMountPoint>();
+}
 
-		bool bRemoveFromContainer;
-		ExistingMount->SetMounted(false, ScopeLock, ExistingMount, bRemoveFromContainer);
-		if (bRemoveFromContainer)
+TRefCountPtr<UE::PackageName::IMountPoint>
+FLongPackagePathsSingleton::FindMountPoint(FStringView RootLongPackageName, FStringView RootLocalPath)
+{
+	using namespace UE::PackageName;
+
+	FMountPointKeyByStrings SearchKey;
+	FString LocalPathAbsolute;
+	if (!TryNormalizePaths(RootLongPackageName, RootLocalPath, SearchKey.LongPackageName,
+		SearchKey.LocalPathStandard, LocalPathAbsolute, true /* bSilenceLogs */))
+	{
+		return TRefCountPtr<UE::PackageName::IMountPoint>();
+	}
+	uint32 SearchHash = GetTypeHash(SearchKey);
+
+	{
+		UE::TReadScopeLock ScopeLock(MountLock);
+		FMountPoint* MountPoint = nullptr;
+
+		FMountPointKeyByPtr* Existing = MountPoints.FindByHash(SearchHash, SearchKey);
+		if (Existing)
 		{
-			MountPoints.RemoveByHash(SearchHash, SearchKey);
+			return TRefCountPtr<IMountPoint>(Existing->GetMountPoint());
 		}
+		return TRefCountPtr<UE::PackageName::IMountPoint>();
 	}
+}
+
+TRefCountPtr<UE::PackageName::IMountPoint>
+FLongPackagePathsSingleton::FindOrAddMountPoint(FStringView RootLongPackageName, FStringView RootLocalPath)
+{
+	using namespace UE::PackageName;
 
-	if (ExistingMount)
+	FMountPointKeyByStrings SearchKey;
+	FString LocalPathAbsolute;
+	if (!TryNormalizePaths(RootLongPackageName, RootLocalPath, SearchKey.LongPackageName,
+		SearchKey.LocalPathStandard, LocalPathAbsolute, false /* bSilenceLogs */))
+	{
+		return TRefCountPtr<UE::PackageName::IMountPoint>();
+	}
+	uint32 SearchHash = GetTypeHash(SearchKey);
+
+	// Look in the existing MountPoints to see if this one already exists, so we can return the existing pointer.
+	{
+		UE::TReadScopeLock ReadScopeLock(MountLock);
+		FMountPointKeyByPtr* Existing = MountPoints.FindByHash(SearchHash, SearchKey);
+		if (Existing)
+		{
+			return TRefCountPtr<IMountPoint>(Existing->GetMountPoint());
+		}
+	}
+
+	{
+		UE::TWriteScopeLock WriteScopeLock(MountLock);
+
+		// Look again in MountPoints in case it was added when we dropped the readlock.
+		FMountPointKeyByPtr* Existing = MountPoints.FindByHash(SearchHash, SearchKey);
+		if (Existing)
+		{
+			return TRefCountPtr<IMountPoint>(Existing->GetMountPoint());
+		}
+
+		FMountPoint* MountPoint = new FMountPoint(MoveTemp(SearchKey.LongPackageName),
+			MoveTemp(SearchKey.LocalPathStandard), MoveTemp(LocalPathAbsolute), EMountFlags::Removable);
+		MountPoints.AddByHash(SearchHash, FMountPointKeyByPtr(*MountPoint));
+		return TRefCountPtr<UE::PackageName::IMountPoint>(MountPoint);
+	}
+}
+
+TRefCountPtr<UE::PackageName::IMountPoint>
+FLongPackagePathsSingleton::FindMountPointByChildLongPackageName(FStringView ChildLongPackageName)
+{
+	using namespace UE::PackageName;
+
+	FLongPackagePathsSingleton& Paths = FLongPackagePathsSingleton::Get();
+
+	UE::TReadScopeLock ScopeLock(Paths.MountLock);
+	FMountPoint* const* MountPointPtr = Paths.RootPathTree.FindClosestValue(ChildLongPackageName);
+	if (!MountPointPtr)
 	{
-		UE_LOG(LogPackageName, Display, TEXT("FPackageName: Mount point removed: '%s' unmounted from '%s'"),
-			*SearchKey.LocalPathStandard, *SearchKey.LongPackageName);
-		// Let subscribers know that a content path was unmounted
-		FPackageName::OnContentPathDismounted().Broadcast(SearchKey.LongPackageName, SearchKey.LocalPathStandard);
+		return TRefCountPtr<IMountPoint>();
 	}
+	FMountPoint* MountPoint = *MountPointPtr;
+	// DirectoryTree is less conservative than IsParentPathOf; e.g. d:/Dir//Root will match
+	// d:/Dir/Root/Child in DirectoryTree but not in IsParentPathOf. Treat this as not found,
+	// since no mountdir will match it according to IsParentPathOf.
+	if (!FPathViews::IsParentPathOf(MountPoint->GetRootPath(), ChildLongPackageName))
+	{
+		return TRefCountPtr<IMountPoint>();
+	}
+	return TRefCountPtr<IMountPoint>(MountPoint);
 }
 
 bool FLongPackagePathsSingleton::MountPointExists(FStringView RootPath)
@@ -544,15 +846,16 @@ FLongPackagePathsSingleton::FLongPackagePathsSingleton()
 			return;
 		}
 
-		TRefCountPtr<FMountPoint> MountPoint = new FMountPoint(FString(RootPath), FString(ContentPathRelative),
+		FMountPoint* MountPoint = new FMountPoint(FString(RootPath), FString(ContentPathRelative),
 			FString(AbsolutePathBuilder), InFlags);
+		TRefCountPtr<IMountPoint> MountPointRefCount = MountPoint;
 		MountPoint->SetMountOrder(NextMountOrder++);
 		bool bAlreadyExisted;
 		MountPoints.Add(FMountPointKeyByPtr(*MountPoint), &bAlreadyExisted);
 		check(!bAlreadyExisted); // Impossible for it to already exist, because we called Find above.
 
 		bool bRemoveFromContainer;
-		MountPoint->SetMounted(true, ScopeLock, MountPoint, bRemoveFromContainer);
+		MountPoint->SetMounted(true, ScopeLock, MountPointRefCount, bRemoveFromContainer);
 		check(!bRemoveFromContainer); // Should only be settable to true when calling SetMounted(false).
 
 		// We are appending to the MountPoint list, and will have only a few duplicates. If any duplicates do occur,
@@ -561,7 +864,7 @@ FLongPackagePathsSingleton::FLongPackagePathsSingleton()
 		FMountPoint*& ExistingRootPath = RootPathTree.FindOrAdd(MountPoint->GetRootPath());
 		if (!ExistingRootPath)
 		{
-			ExistingRootPath = MountPoint.GetReference();
+			ExistingRootPath = MountPoint;
 		}
 		else
 		{
@@ -573,7 +876,7 @@ FLongPackagePathsSingleton::FLongPackagePathsSingleton()
 		FMountPoint*& ExistingContentPathRelative = ContentPathTree.FindOrAdd(MountPoint->GetContentPathRelative());
 		if (!ExistingContentPathRelative)
 		{
-			ExistingContentPathRelative = MountPoint.GetReference();
+			ExistingContentPathRelative = MountPoint;
 		}
 		else
 		{
@@ -585,7 +888,7 @@ FLongPackagePathsSingleton::FLongPackagePathsSingleton()
 		FMountPoint*& ExistingContentPathAbsolute = ContentPathTree.FindOrAdd(MountPoint->GetContentPathAbsolute());
 		if (!ExistingContentPathAbsolute)
 		{
-			ExistingContentPathAbsolute = MountPoint.GetReference();
+			ExistingContentPathAbsolute = MountPoint;
 		}
 		else
 		{
@@ -627,18 +930,19 @@ FLongPackagePathsSingleton::~FLongPackagePathsSingleton()
 	using namespace UE::PackageName;
 
 	// We require that our external callers drop their IMountPoint refcounts, but we do not require that they unmount
-	// all of the MountPoints they registered with RegisterMountPoint. So we may have MountPoints remaining with that
-	// our still mounted and therefore we still hold a refcount. Unmount them all now to drop that final refcount and
+	// all of the MountPoints they registered with RegisterMountPoint. So we may have MountPoints remaining that
+	// are still mounted and therefore we still hold a refcount. Unmount them all now to drop that final refcount and
 	// free the memory. This happens at end of process, so this is important mainly for leak detection.
-	TArray<TRefCountPtr<FMountPoint>> RemainingMountPoints;
+	TArray<TRefCountPtr<IMountPoint>> RemainingMountPoints;
 	{
 		UE::TWriteScopeLock ScopeLock(MountLock);
 		RemainingMountPoints.Reserve(MountPoints.Num());
 		for (FMountPointKeyByPtr& Key : MountPoints)
 		{
-			TRefCountPtr<FMountPoint>& MountPoint = RemainingMountPoints.Emplace_GetRef(Key.GetMountPoint());
+			FMountPoint* MountPoint = Key.GetMountPoint();
+			TRefCountPtr<IMountPoint>& MountPointRefCount = RemainingMountPoints.Emplace_GetRef(MountPoint);
 			bool bRemoveFromContainer;
-			MountPoint->SetMounted(false, ScopeLock, MountPoint, bRemoveFromContainer);
+			MountPoint->SetMounted(false, ScopeLock, MountPointRefCount, bRemoveFromContainer);
 			if (!bRemoveFromContainer)
 			{
 				// UE_LOG is likely no longer available; we will have to rely on leak detector reporting this instead.
@@ -1572,9 +1876,14 @@ bool FPackageName::IsValidPath(FStringView InPath)
 		IsMemoryPackage(InPath);
 }
 
-void FPackageName::RegisterMountPoint(const FString& RootPath, const FString& ContentPath)
+TRefCountPtr<UE::PackageName::IMountPoint> FPackageName::RegisterMountPoint(const FString& RootPath, const FString& ContentPath)
+{
+	return FLongPackagePathsSingleton::Get().InsertMountPoint(RootPath, ContentPath);
+}
+
+void FPackageName::RegisterMountPoint(const TRefCountPtr<UE::PackageName::IMountPoint>& MountPoint)
 {
-	FLongPackagePathsSingleton::Get().InsertMountPoint(RootPath, ContentPath);
+	FLongPackagePathsSingleton::Get().InsertMountPoint(MountPoint);
 }
 
 void FPackageName::UnRegisterMountPoint(const FString& RootPath, const FString& ContentPath)
@@ -1582,6 +1891,11 @@ void FPackageName::UnRegisterMountPoint(const FString& RootPath, const FString&
 	FLongPackagePathsSingleton::Get().RemoveMountPoint(RootPath, ContentPath);
 }
 
+void FPackageName::UnRegisterMountPoint(const TRefCountPtr<UE::PackageName::IMountPoint>& MountPoint)
+{
+	FLongPackagePathsSingleton::Get().RemoveMountPoint(MountPoint);
+}
+
 bool FPackageName::MountPointExists(FStringView RootPath)
 {
 	return FLongPackagePathsSingleton::Get().MountPointExists(RootPath);
@@ -1597,6 +1911,36 @@ bool FPackageName::MountPointExists(const TCHAR* RootPath)
 	return FLongPackagePathsSingleton::Get().MountPointExists(FStringView(RootPath));
 }
 
+TRefCountPtr<UE::PackageName::IMountPoint>
+FPackageName::FindMountPointByRootPackageName(FStringView RootLongPackageName)
+{
+	return FLongPackagePathsSingleton::Get().FindMountPointByRootPackageName(RootLongPackageName);
+}
+
+TRefCountPtr<UE::PackageName::IMountPoint>
+FPackageName::FindMountPointByRootLocalPath(FStringView RootLocalPath)
+{
+	return FLongPackagePathsSingleton::Get().FindMountPointByRootLocalPath(RootLocalPath);
+}
+
+TRefCountPtr<UE::PackageName::IMountPoint>
+FPackageName::FindMountPoint(FStringView RootLongPackageName, FStringView RootLocalPath)
+{
+	return FLongPackagePathsSingleton::Get().FindMountPoint(RootLongPackageName, RootLocalPath);
+}
+
+TRefCountPtr<UE::PackageName::IMountPoint>
+FPackageName::FindOrAddMountPoint(FStringView RootLongPackageName, FStringView RootLocalPath)
+{
+	return FLongPackagePathsSingleton::Get().FindOrAddMountPoint(RootLongPackageName, RootLocalPath);
+}
+
+TRefCountPtr<UE::PackageName::IMountPoint>
+FPackageName::FindMountPointByChildLongPackageName(FStringView ChildLongPackageName)
+{
+	return FLongPackagePathsSingleton::Get().FindMountPointByChildLongPackageName(ChildLongPackageName);
+}
+
 FName FPackageName::GetPackageMountPoint(FStringView InPackagePath, bool InWithoutSlashes)
 {
 	using namespace UE::PackageName;
diff --git a/Engine/Source/Runtime/CoreUObject/Private/Misc/PackageName/MountPoint.cpp b/Engine/Source/Runtime/CoreUObject/Private/Misc/PackageName/MountPoint.cpp
--- a/Engine/Source/Runtime/CoreUObject/Private/Misc/PackageName/MountPoint.cpp
+++ b/Engine/Source/Runtime/CoreUObject/Private/Misc/PackageName/MountPoint.cpp
@@ -25,7 +25,7 @@ bool FMountPoint::IsMounted() const
 }
 
 void FMountPoint::SetMounted(bool bInMounted, UE::TWriteScopeLock<FTransactionallySafeRWLock>& /* ProofOfWriteLock */,
-	const TRefCountPtr<FMountPoint>& /* ProofOfCallerRefCount */, bool& bOutRemoveFromContainer)
+	const TRefCountPtr<IMountPoint>& /* ProofOfCallerRefCount */, bool& bOutRemoveFromContainer)
 {
 	bOutRemoveFromContainer = false;
 
diff --git a/Engine/Source/Runtime/CoreUObject/Private/Misc/PackageName/MountPointPrivate.h b/Engine/Source/Runtime/CoreUObject/Private/Misc/PackageName/MountPointPrivate.h
--- a/Engine/Source/Runtime/CoreUObject/Private/Misc/PackageName/MountPointPrivate.h
+++ b/Engine/Source/Runtime/CoreUObject/Private/Misc/PackageName/MountPointPrivate.h
@@ -55,6 +55,9 @@ public:
 	virtual void AddRef() const override;
 	virtual FReturnedRefCountValue Release() const override;
 	virtual FReturnedRefCountValue GetRefCount() const override;
+protected:
+	virtual FMountPoint* AsFMountPoint() override;
+public:
 
 	// FMountPoint additional API
 
@@ -77,7 +80,7 @@ public:
 	 * be set to true and we will skip the effort in TryUnregisterAndDelete of removing it from the container.
 	 */
 	void SetMounted(bool bInMounted, UE::TWriteScopeLock<FTransactionallySafeRWLock>& ProofOfWriteLock,
-		const TRefCountPtr<FMountPoint>& ProofOfCallerRefCount, bool& bOutRemoveFromContainer);
+		const TRefCountPtr<IMountPoint>& ProofOfCallerRefCount, bool& bOutRemoveFromContainer);
 
 	/**
 	 * Get/Set the Ephemeral ordering value. MountPoints that have the same RootPath or ContentPath as other
@@ -94,6 +97,7 @@ public:
 	static uint32 GetTypeHashKey(const FString& LongPackageNameOfKey, const FString& LocalPathStandardOfKey);
 	static bool KeysAreEqual(const FString& LongPackageNameA, const FString& LocalPathStandardA,
 		const FString& LongPackageNameB, const FString& LocalPathStandardB);
+	static bool IsHigherOverridePriority(FMountPoint* A, FMountPoint* B);
 
 private:
 	/**
@@ -182,6 +186,11 @@ inline FStringView FMountPoint::GetLocalPathAbsolute() const
 	return LocalPathAbsolute;
 }
 
+inline FMountPoint* FMountPoint::AsFMountPoint()
+{
+	return this;
+}
+
 inline bool FMountPoint::IsRemovable() const
 {
 	return bRemovable;
@@ -295,4 +304,24 @@ inline bool FMountPoint::KeysAreEqual(const FString& LongPackageNameA, const FSt
 		&& LocalPathStandardA.Equals(LocalPathStandardB, ESearchCase::IgnoreCase);
 }
 
+inline bool FMountPoint::IsHigherOverridePriority(FMountPoint* A, FMountPoint* B)
+{
+	// Must be called within FLongPackagePathsSingleton->MountLock so that GetMountOrder is readable.
+	if ((A != nullptr) != (B != nullptr))
+	{
+		return A != nullptr;
+	}
+	if (A == nullptr)
+	{
+		return false; // equal
+	}
+	CA_ASSUME(B != nullptr);
+
+	if (A->IsMounted() != B->IsMounted())
+	{
+		return A->IsMounted();
+	}
+	return A->GetMountOrder() > B->GetMountOrder();
+}
+
 } // namespace UE::PackageName
\ No newline at end of file
diff --git a/Engine/Source/Runtime/CoreUObject/Private/Misc/PackageName/PackageNamePrivate.h b/Engine/Source/Runtime/CoreUObject/Private/Misc/PackageName/PackageNamePrivate.h
--- a/Engine/Source/Runtime/CoreUObject/Private/Misc/PackageName/PackageNamePrivate.h
+++ b/Engine/Source/Runtime/CoreUObject/Private/Misc/PackageName/PackageNamePrivate.h
@@ -64,17 +64,27 @@ struct FLongPackagePathsSingleton
 	 * LocalPathStandard is a relative path and is consistent with FileManager relative paths.
 	 * @return false if the input strings are not recognizable as mountpoint strings, else true.
 	 */
-	bool TryNormalizePaths(const FString& InLongPackageName, const FString& InLocalPath,
+	bool TryNormalizePaths(FStringView InLongPackageName, FStringView InLocalPath,
 		FString& OutLongPackageName, FString& OutLocalPathStandard, FString& OutLocalPathAbsolute, bool bSilenceLogs);
-	bool TryNormalizeLongPackageNamePath(const FString& InLongPackageName, FString& OutLongPackageName,
+	bool TryNormalizeLongPackageNamePath(FStringView InLongPackageName, FString& OutLongPackageName,
 		bool bSilenceLogs);
-	bool TryNormalizeLocalPath(const FString& InLocalPath, FString& OutLocalPathStandard,
+	bool TryNormalizeLocalPath(FStringView InLocalPath, FString& OutLocalPathStandard,
 		FString& OutLocalPathAbsolute, bool bSilenceLogs);
 
-	void InsertMountPoint(const FString& RootPath, const FString& ContentPath);
+	TRefCountPtr<UE::PackageName::IMountPoint> InsertMountPoint(const FString& RootPath, const FString& ContentPath);
+	void InsertMountPoint(const TRefCountPtr<UE::PackageName::IMountPoint>& InMountPoint);
 
 	/** This will remove a previously inserted mount point. */
 	void RemoveMountPoint(const FString& RootPath, const FString& ContentPath);
+	void RemoveMountPoint(const TRefCountPtr<UE::PackageName::IMountPoint>& InMountPoint);
+
+	TRefCountPtr<UE::PackageName::IMountPoint> FindMountPointByRootPackageName(FStringView RootLongPackageName);
+	TRefCountPtr<UE::PackageName::IMountPoint> FindMountPointByRootLocalPath(FStringView RootLocalPath);
+	TRefCountPtr<UE::PackageName::IMountPoint> FindMountPoint(FStringView RootLongPackageName,
+		FStringView RootLocalPath);
+	TRefCountPtr<UE::PackageName::IMountPoint> FindOrAddMountPoint(FStringView RootLongPackageName,
+		FStringView RootLocalPath);
+	TRefCountPtr<UE::PackageName::IMountPoint> FindMountPointByChildLongPackageName(FStringView ChildLongPackageName);
 
 	/** Checks whether the specific root path is a valid mount point. */
 	bool MountPointExists(FStringView RootPath);
@@ -89,6 +99,12 @@ private:
 	FLongPackagePathsSingleton();
 	~FLongPackagePathsSingleton();
 	void ReassignMountOrders();
+	void BroadcastMountPointMounted(const TRefCountPtr<UE::PackageName::IMountPoint>& PublicMountPoint,
+		UE::PackageName::FMountPoint* MountPoint);
+	bool TryRemoveMountPointInternal(const TRefCountPtr<UE::PackageName::IMountPoint>& InternallyHeldMountPoint,
+		UE::TWriteScopeLock<FTransactionallySafeRWLock>& ProofOfWriteLock, UE::PackageName::FMountPoint* MountPoint);
+	void BroadcastMountPointDismounted(const TRefCountPtr<UE::PackageName::IMountPoint>& PublicMountPoint,
+		UE::PackageName::FMountPoint* MountPoint);
 
 #if !UE_BUILD_SHIPPING
 	void ExecDumpMountPoints(const TArray<FString>& Args);
diff --git a/Engine/Source/Runtime/CoreUObject/Public/Misc/MountPoint.h b/Engine/Source/Runtime/CoreUObject/Public/Misc/MountPoint.h
--- a/Engine/Source/Runtime/CoreUObject/Public/Misc/MountPoint.h
+++ b/Engine/Source/Runtime/CoreUObject/Public/Misc/MountPoint.h
@@ -5,6 +5,9 @@
 #include "Containers/StringView.h"
 #include "Templates/RefCounting.h"
 
+namespace UE::PackageName { class FMountPoint; }
+struct FLongPackagePathsSingleton;
+
 namespace UE::PackageName
 {
 
@@ -77,6 +80,15 @@ public:
 	 * Mutable but atomic, can be read from any thread.
 	 */
 	virtual bool IsMounted() const = 0;
+
+protected:
+	/**
+	 * A downcast function to the implementation class inside FPackageName. There should be no other implementation of
+	 * IMountPoint; we log an error if this fails.
+	 */
+	virtual FMountPoint* AsFMountPoint() = 0;
+
+	friend FLongPackagePathsSingleton;
 };
 
 } // namespace UE::PackageName
\ No newline at end of file
diff --git a/Engine/Source/Runtime/CoreUObject/Public/Misc/PackageName.h b/Engine/Source/Runtime/CoreUObject/Public/Misc/PackageName.h
--- a/Engine/Source/Runtime/CoreUObject/Public/Misc/PackageName.h
+++ b/Engine/Source/Runtime/CoreUObject/Public/Misc/PackageName.h
@@ -314,30 +314,113 @@ public:
 	static COREUOBJECT_API bool TryConvertGameRelativePackagePathToLocalPath(FStringView RelativePackagePath, FString& OutLocalPath);
 
 	/**
-	 * This will insert a mount point at the head of the search chain (so it can overlap an existing mount point and win).
-	 * If you register a mount point (even if you do so only in certain circumstances) consider also adding a handler for GetExplanationForUnavailablePackage
-	 * to help debug cases where a package can't be found
+	 * This will mount a mount point if not already mounted, and will update its order to the front of the override
+	 * chain (so it can overlap an existing mount point and win).
+	 * If you register a mount point (even if you do so only in certain circumstances) consider also adding a handler
+	 * for GetExplanationForUnavailablePackage to help debug cases where a package can't be found.
+	 * This function can fail only if the input arguments are not valid text for RootPath or ContentPath.
 	 *
-	 * @param RootPath Logical Root Path.
-	 * @param ContentPath Content Path on disk.
+	 * @param RootPath LongPackageName location for the root of the MountPoint,
+	 *        Examples: /Engine/, /PluginName/
+	 * @param ContentPath LocalPath location for the root of the MountPoint,
+	 *        Examples: ../../../Engine/Content, d:\FullPathToUnreal\Engine\Plugins\PluginName
+	 * 
+	 * @return RefCountPtr to the mounted IMountPoint, or nullptr if failed.
+	 */
+	static COREUOBJECT_API TRefCountPtr<UE::PackageName::IMountPoint>
+		RegisterMountPoint(const FString& RootPath, const FString& ContentPath);
+	/**
+	 * Mounts a mount point if not already mounted, and updates its order to front of the override chain.
+	 * If you register a mount point (even if you do so only in certain circumstances) consider also adding a handler
+	 * for GetExplanationForUnavailablePackage to help debug cases where a package can't be found.
 	 */
-	static COREUOBJECT_API void RegisterMountPoint(const FString& RootPath, const FString& ContentPath);
+	static COREUOBJECT_API void RegisterMountPoint(const TRefCountPtr<UE::PackageName::IMountPoint>& MountPoint);
 
 	/**
-	 * This will remove a previously inserted mount point.
+	 * This will remove a previously inserted mount point. Logs a message if not found or not removable.
 	 *
-	 * @param RootPath Logical Root Path.
-	 * @param ContentPath Content Path on disk.
+	 * @param RootPath LongPackageName location for the root of the MountPoint,
+	 *        Examples: /Engine/, /PluginName/
+	 * @param ContentPath LocalPath location for the root of the MountPoint,
+	 *        Examples: ../../../Engine/Content, d:\FullPathToUnreal\Engine\Plugins\PluginName
 	 */
 	static COREUOBJECT_API void UnRegisterMountPoint(const FString& RootPath, const FString& ContentPath);
+	/** Unmounts a MountPoint; the mountpoint remains findable and updatable until all references are released. */
+	static COREUOBJECT_API void UnRegisterMountPoint(const TRefCountPtr<UE::PackageName::IMountPoint>& MountPoint);
+
 
 	/**
-	 * Returns whether the specific logical root path is a valid mount point.
+	 * Returns whether the given RootPath is the LongPackageName of a MountPoint that has been mounted, by default or
+	 * by an earlier call to RegisterMountPoint. 
 	 */
 	static COREUOBJECT_API bool MountPointExists(FStringView RootPath);
 	static COREUOBJECT_API bool MountPointExists(const FString& RootPath);
 	static COREUOBJECT_API bool MountPointExists(const TCHAR* RootPath);
 
+	/**
+	 * Looks for the first allocated MountPoint with the given RootLongPackageName. If multiple are allocated, only
+	 * returns the one with the highest override priority (the one that was added last). A discovered MountPoint might
+	 * or might not be mounted, call MountPoint->IsMounted() to check for mounted.
+	 * 
+	 * @param RootLongPackageName LongPackageName location for the root of the MountPoint,
+	 *        Examples: /Engine/, /PluginName/
+	 * @return RefCountPtr to the IMountPoint, or nullptr if not yet allocated.
+	 */
+	static COREUOBJECT_API TRefCountPtr<UE::PackageName::IMountPoint>
+		FindMountPointByRootPackageName(FStringView RootLongPackageName);
+	/**
+	 * Looks for the first MountPoint with the given RootLocalPath. If multiple are allocated, only returns the
+	 * one with the highest override priority (the one that was added last). A discovered MountPoint might or might
+	 * not be mounted, call MountPoint->IsMounted() to check for mounted.
+	 *
+	 * @param RootLocalPath LocalPath location for the root of the MountPoint,
+	 *        Examples: ../../../Engine/Content, d:\FullPathToUnreal\Engine\Plugins\PluginName
+	 * @return RefCountPtr to the IMountPoint, or nullptr if not yet allocated.
+	 */
+	static COREUOBJECT_API TRefCountPtr<UE::PackageName::IMountPoint>
+		FindMountPointByRootLocalPath(FStringView RootLocalPath);
+	/**
+	 * Looks for an existing allocated MountPoint with the given RootLongPackageName and RootLocalPath. These fields
+	 * uniquely identify a MountPoint so at most one can exist.
+	 * If the MountPoint exists, it might or might not be mounted, call MountPoint->IsMounted() to check for mounted.
+	 *
+	 * @param RootLongPackageName LongPackageName location for the root of the MountPoint,
+	 *        Examples: /Engine/, /PluginName/
+	 * @param RootLocalPath LocalPath location for the root of the MountPoint,
+	 *        Examples: ../../../Engine/Content, d:\FullPathToUnreal\Engine\Plugins\PluginName
+	 * @return RefCountPtr to the IMountPoint, or nullptr if not yet allocated.
+	 */
+	static COREUOBJECT_API TRefCountPtr<UE::PackageName::IMountPoint>
+		FindMountPoint(FStringView RootLongPackageName, FStringView RootLocalPath);
+	/**
+	 * Looks for an existing allocated MountPoint with the given RootLongPackageName and RootLocalPath, adds it
+	 * if it does not already exist. These fields uniquely identify a MountPoint so at most one can exist.
+	 * When the MountPoint is returned, it might or might not be mounted, call MountPoint->IsMounted() to check for
+	 * mounted.
+	 * This function can only return null if the input arguments are not valid text for RootLongPackageName or
+	 * RootLocalPath.
+	 *
+	 * @param RootLongPackageName LongPackageName location for the root of the MountPoint,
+	 *        Examples: /Engine/, /PluginName/
+	 * @param RootLocalPath LocalPath location for the root of the MountPoint,
+	 *        Examples: ../../../Engine/Content, d:\FullPathToUnreal\Engine\Plugins\PluginName
+	 * @return RefCountPtr to the IMountPoint, or nullptr if failed.
+	 */
+	static COREUOBJECT_API TRefCountPtr<UE::PackageName::IMountPoint>
+		FindOrAddMountPoint(FStringView RootLongPackageName, FStringView RootLocalPath);
+
+	/**
+	 * Looks for the first MountPoint that is equal to or is a parent path of the given ChildLongPackageName.
+	 * If multiple are allocated, only returns the one with the highest override priority (the one that was added
+	 * last). A discovered MountPoint might or might not be mounted, call MountPoint->IsMounted() to check for mounted.
+	 *
+	 * @param ChildLongPackageName LongPackageName for a path under the MountPoint,
+	 *        Examples: /Engine/, /Engine/Asset, /PluginName/, /PluginName/Asset
+	 * @return RefCountPtr to the IMountPoint, or nullptr if not yet allocated.
+	 */
+	static COREUOBJECT_API TRefCountPtr<UE::PackageName::IMountPoint>
+		FindMountPointByChildLongPackageName(FStringView ChildLongPackageName);
+
 	/**
 	 * Get the mount point for a given package path
 	 * 
@@ -651,14 +734,25 @@ public:
 	static COREUOBJECT_API void IteratePackagesInDirectory(const FString& RootDir, const FPackageNameVisitor& Visitor);
 	static COREUOBJECT_API void IteratePackagesInDirectory(const FString& RootDir, const FPackageNameStatVisitor& Visitor);
 
-	/** Event that is triggered when a new content path is mounted */
+	/** Event triggered when a MountPoint is mounted. */
+	DECLARE_MULTICAST_DELEGATE_OneParam(FMountPointEvent, const TRefCountPtr<UE::PackageName::IMountPoint>& MountPoint)
+	static FMountPointEvent& OnMountPointMounted()
+	{
+		return OnMountPointMountedEvent;
+	}
+	/** Legacy event triggered when a MountPoint is mounted; passes strings instead of IMountPoint. */
 	DECLARE_MULTICAST_DELEGATE_TwoParams( FOnContentPathMountedEvent, const FString& /* Asset path */, const FString& /* ContentPath */ );
 	static FOnContentPathMountedEvent& OnContentPathMounted()
 	{
 		return OnContentPathMountedEvent;
 	}
 
-	/** Event that is triggered when a new content path is removed */
+	/** Event triggered when a MountPoint is dismounted. */
+	static FMountPointEvent& OnMountPointDismounted()
+	{
+		return OnMountPointDismountedEvent;
+	}
+	/** Legacy event triggered when a MountPoint is dismounted; passes strings instead of IMountPoint. */
 	DECLARE_MULTICAST_DELEGATE_TwoParams(FOnContentPathDismountedEvent, const FString& /* Asset path */, const FString& /* ContentPath */ );
 	static FOnContentPathDismountedEvent& OnContentPathDismounted()
 	{
@@ -1012,9 +1106,11 @@ private:
 	static COREUOBJECT_API bool TryConvertToMountedPackagePath(const FString& InPath, FPackagePath& OutPackagePath, EErrorCode& OutFailureReason);
 
 	/** Event that is triggered when a new content path is mounted */
+	static COREUOBJECT_API FMountPointEvent OnMountPointMountedEvent;
 	static COREUOBJECT_API FOnContentPathMountedEvent OnContentPathMountedEvent;
 
 	/** Event that is triggered when a new content path is removed */
+	static COREUOBJECT_API FMountPointEvent OnMountPointDismountedEvent;
 	static COREUOBJECT_API FOnContentPathDismountedEvent OnContentPathDismountedEvent;
 
 	/** Delegate used to check whether a package exist without using the filesystem. */
