[Backout] - CL50612623
#rnx
Original CL Desc
-----------------------------------------------------------------
PackageName and AssetRegistry:
Add accessor functions for IMountPoint to FPackageName, including new OnContentPathMounted delegates that pass through the IMountPoint instead of strings.
Use the new IMountPoint accessors in AssetRegistry, for better performance of the upcoming ActiveMountsLoader.
Add FSuspendPathsMounting to match FSuspendPathsDismounting, to get batching performance improvement when mounting MountPoints.
Pass more stack local variable arguments into AssetRegistry helper functions, for use by the hooks for the upcoming ActiveMountsLoader.
#rnx

#Synced-CL 50612662

diff --git a/Engine/Source/Runtime/AssetRegistry/Private/ARLoader.cpp b/Engine/Source/Runtime/AssetRegistry/Private/ARLoader.cpp
--- a/Engine/Source/Runtime/AssetRegistry/Private/ARLoader.cpp
+++ b/Engine/Source/Runtime/AssetRegistry/Private/ARLoader.cpp
@@ -220,7 +220,6 @@ void FPreloader::TryLoadAsync()
 	// Trigger outside the lock so that a locked Consume function that is waiting on PreloadReady can wait inside the lock.
 	PreloadReady->Trigger();
 
-	FString StateFilePath;
 	FConsumeFunction LocalConsumeCallback;
 	{
 		FScopeLock Lock(&StateLock);
@@ -230,7 +229,6 @@ void FPreloader::TryLoadAsync()
 			LoadState = EState::Loaded;
 			if (ConsumeCallback)
 			{
-				StateFilePath = ARPath;
 				LocalConsumeCallback = MoveTemp(ConsumeCallback);
 				ConsumeCallback.Reset();
 				LoadState = EState::Consumed;
@@ -242,7 +240,7 @@ void FPreloader::TryLoadAsync()
 	{
 		// No further threads will read/write payload at this point (until destructor, which is called after all async threads are complete
 		// so we can use it outside the lock
-		LocalConsumeCallback(LocalResult, MoveTemp(Payload), StateFilePath);
+		LocalConsumeCallback(LocalResult, MoveTemp(Payload));
 		Shutdown();
 	}
 }
@@ -258,7 +256,7 @@ FPreloader::EConsumeResult FPreloader::ConsumeInternal(FConsumeFunction&& Consum
 	{
 		Lock.Unlock(); // Unlock before calling external code in Consume callback
 		ELoadResult LocalResult = (LoadState == EState::Consumed || ConsumeCallback) ? ELoadResult::AlreadyConsumed : ELoadResult::Inactive;
-		ConsumeSynchronous(LocalResult, FAssetRegistryState(), FString());
+		ConsumeSynchronous(LocalResult, FAssetRegistryState());
 		return EConsumeResult::Failed;
 	}
 
@@ -266,9 +264,8 @@ FPreloader::EConsumeResult FPreloader::ConsumeInternal(FConsumeFunction&& Consum
 	{
 		ELoadResult LocalResult = TrySetPath() ? TryLoad() : ELoadResult::NotFound;
 		LoadState = EState::Consumed;
-		FString StateFilePath = ARPath;
 		Lock.Unlock(); // Unlock before calling external code in Consume callback
-		ConsumeSynchronous(LocalResult, MoveTemp(Payload), StateFilePath);
+		ConsumeSynchronous(LocalResult, MoveTemp(Payload));
 		Shutdown(); // Shutdown can be called outside the lock since AsyncThread doesn't exist
 		return LocalResult == ELoadResult::Succeeded ? EConsumeResult::Succeeded : EConsumeResult::Failed;
 	}
@@ -296,9 +293,8 @@ FPreloader::EConsumeResult FPreloader::ConsumeInternal(FConsumeFunction&& Consum
 	LoadState = EState::Consumed;
 
 	// No further async threads exist that will read/write payload at this point so we can use it outside the lock
-	FString StateFilePath = ARPath;
 	Lock.Unlock(); // Unlock before calling external code in Consume callback
-	ConsumeSynchronous(LocalResult, MoveTemp(Payload), StateFilePath);
+	ConsumeSynchronous(LocalResult, MoveTemp(Payload));
 	Shutdown(); // Shutdown can be called outside the lock since we have set state to Consumed and the Async thread will notice and exit
 	return LocalResult == ELoadResult::Succeeded ? EConsumeResult::Succeeded : EConsumeResult::Failed;
 }
@@ -365,13 +361,10 @@ void FAsyncConsumer::Wait(UE::AssetRegistry::FInterfaceWriteScopeLock& ScopeLock
 	}
 }
 
-void FAsyncConsumer::Consume(UAssetRegistryImpl& UARI, UE::AssetRegistry::Impl::FEventContext& EventContext,
-	UE::AssetRegistry::FInterfaceWriteScopeLock& ProofOfWriteLock, ELoadResult LoadResult,
-	FAssetRegistryState&& ARState, FStringView StateFilePath)
+void FAsyncConsumer::Consume(UAssetRegistryImpl& UARI, UE::AssetRegistry::Impl::FEventContext& EventContext, ELoadResult LoadResult, FAssetRegistryState&& ARState)
 {
 	// Called within the lock
-	UARI.GuardedData.LoadPremadeAssetRegistry(UARI, EventContext, ProofOfWriteLock, LoadResult,
-		MoveTemp(ARState), StateFilePath);
+	UARI.GuardedData.LoadPremadeAssetRegistry(EventContext, LoadResult, MoveTemp(ARState));
 	check(ReferenceCount >= 1);
 	check(Consumed != nullptr);
 	Consumed->Trigger();
diff --git a/Engine/Source/Runtime/AssetRegistry/Private/ARLoader.h b/Engine/Source/Runtime/AssetRegistry/Private/ARLoader.h
--- a/Engine/Source/Runtime/AssetRegistry/Private/ARLoader.h
+++ b/Engine/Source/Runtime/AssetRegistry/Private/ARLoader.h
@@ -4,7 +4,6 @@
 
 #include "ARTypes.h"
 #include "AssetRegistry/AssetRegistryState.h"
-#include "Containers/StringView.h"
 #include "Containers/UnrealString.h"
 #include "Delegates/Delegate.h"
 #include "HAL/CriticalSection.h"
@@ -45,8 +44,7 @@ public:
 		Failed,
 		Deferred
 	};
-	using FConsumeFunction = TFunction<void(ELoadResult LoadResult, FAssetRegistryState&& ARState,
-		FStringView StateFilePath)>;
+	using FConsumeFunction = TFunction<void(ELoadResult LoadResult, FAssetRegistryState&& ARState)>;
 
 	FPreloader();
 	~FPreloader();
@@ -146,8 +144,7 @@ struct FAsyncConsumer
 	void Wait(FInterfaceWriteScopeLock& ScopeLock);
 	/** Callback from the async thread to consume the Premade ARState */
 	void Consume(UAssetRegistryImpl& UARI, UE::AssetRegistry::Impl::FEventContext& EventContext,
-		UE::AssetRegistry::FInterfaceWriteScopeLock& ProofOfWriteLock,
-		ELoadResult LoadResult, FAssetRegistryState&& ARState, FStringView StateFilePath);
+		ELoadResult LoadResult, FAssetRegistryState&& ARState);
 
 private:
 	/**
diff --git a/Engine/Source/Runtime/AssetRegistry/Private/AssetRegistry.cpp b/Engine/Source/Runtime/AssetRegistry/Private/AssetRegistry.cpp
--- a/Engine/Source/Runtime/AssetRegistry/Private/AssetRegistry.cpp
+++ b/Engine/Source/Runtime/AssetRegistry/Private/AssetRegistry.cpp
@@ -373,8 +373,7 @@ void FAssetRegistryImpl::ConditionalLoadPremadeAssetRegistry(UE::AssetRegistry::
 	AsyncConsumer.Wait(ScopeLock);
 }
 
-void FAssetRegistryImpl::ConsumeOrDeferPreloadedPremade(UAssetRegistryImpl& UARI, Impl::FEventContext& EventContext,
-	FInterfaceWriteScopeLock& ProofOfWriteLock)
+void FAssetRegistryImpl::ConsumeOrDeferPreloadedPremade(UAssetRegistryImpl& UARI, Impl::FEventContext& EventContext)
 {
 	// Called from inside WriteLock on InterfaceLock
 	using namespace UE::AssetRegistry::Premade;
@@ -388,24 +387,20 @@ void FAssetRegistryImpl::ConsumeOrDeferPreloadedPremade(UAssetRegistryImpl& UARI
 
 	if (Premade::CanLoadAsync())
 	{
-		FPreloader::FConsumeFunction ConsumeFromAsyncThread = [this, &UARI](Premade::ELoadResult LoadResult,
-			FAssetRegistryState&& ARState, FStringView StateFilePath)
+		FPreloader::FConsumeFunction ConsumeFromAsyncThread = [this, &UARI](Premade::ELoadResult LoadResult, FAssetRegistryState&& ARState)
 		{
 			Impl::FEventContext EventContext;
 			{
 				UE::AssetRegistry::FInterfaceWriteScopeLock InterfaceScopeLock(UARI.InterfaceLock);
-				AsyncConsumer.Consume(UARI, EventContext, InterfaceScopeLock, LoadResult,
-					MoveTemp(ARState), StateFilePath);
+				AsyncConsumer.Consume(UARI, EventContext, LoadResult, MoveTemp(ARState));
 			}
 			UARI.Broadcast(EventContext);
 		};
-		auto ConsumeOnCurrentThread = [ConsumeFromAsyncThread](Premade::ELoadResult LoadResult,
-			FAssetRegistryState&& ARState, FStringView InStateFilePath) mutable
+		auto ConsumeOnCurrentThread = [ConsumeFromAsyncThread](Premade::ELoadResult LoadResult, FAssetRegistryState&& ARState) mutable
 		{
-			Async(EAsyncExecution::TaskGraph, [LoadResult, ARState=MoveTemp(ARState),
-				ConsumeFromAsyncThread=MoveTemp(ConsumeFromAsyncThread), StateFilePath=FString(InStateFilePath)]() mutable
+			Async(EAsyncExecution::TaskGraph, [LoadResult, ARState=MoveTemp(ARState), ConsumeFromAsyncThread=MoveTemp(ConsumeFromAsyncThread)]() mutable
 			{
-				ConsumeFromAsyncThread(LoadResult, MoveTemp(ARState), StateFilePath);
+				ConsumeFromAsyncThread(LoadResult, MoveTemp(ARState));
 			});
 		};
 
@@ -414,12 +409,9 @@ void FAssetRegistryImpl::ConsumeOrDeferPreloadedPremade(UAssetRegistryImpl& UARI
 	}
 	else
 	{
-		GPreloader.Consume(
-			[this, &UARI, &EventContext, &ProofOfWriteLock](Premade::ELoadResult LoadResult,
-			FAssetRegistryState&& ARState, FStringView StateFilePath)
+		GPreloader.Consume([this, &EventContext](Premade::ELoadResult LoadResult, FAssetRegistryState&& ARState)
 			{
-				LoadPremadeAssetRegistry(UARI, EventContext, ProofOfWriteLock, LoadResult,
-					MoveTemp(ARState), StateFilePath);
+				LoadPremadeAssetRegistry(EventContext, LoadResult, MoveTemp(ARState));
 			});
 	}
 }
@@ -545,9 +537,8 @@ FAssetRegistryImpl::FAssetRegistryImpl()
 {
 }
 
-void FAssetRegistryImpl::LoadPremadeAssetRegistry(UAssetRegistryImpl& UARI, Impl::FEventContext& EventContext,
-	FInterfaceWriteScopeLock& ProofOfWriteLock,
-	Premade::ELoadResult LoadResult, FAssetRegistryState&& ARState, FStringView StateFilePath)
+void FAssetRegistryImpl::LoadPremadeAssetRegistry(Impl::FEventContext& EventContext,
+	Premade::ELoadResult LoadResult, FAssetRegistryState&& ARState)
 {
 	SCOPED_BOOT_TIMING("LoadPremadeAssetRegistry");
 	UE_SCOPED_ENGINE_ACTIVITY("Loading premade asset registry");
@@ -561,7 +552,6 @@ void FAssetRegistryImpl::LoadPremadeAssetRegistry(UAssetRegistryImpl& UARI, Impl
 		{
 			if (State.GetNumAssets() == 0)
 			{
-				// Add the premade state into the global AssetRegistry
 				State = MoveTemp(ARState);
 				CachePathsFromState(EventContext, State);
 				if (bEmitAssetEvents)
@@ -569,38 +559,28 @@ void FAssetRegistryImpl::LoadPremadeAssetRegistry(UAssetRegistryImpl& UARI, Impl
 					EventContext.AssetEvents.Reserve(State.GetNumAssets());
 					State.EnumerateAllAssets([&EventContext](const FAssetData& AssetData)
 					{
-						EventContext.AssetEvents.Emplace(AssetData,
-							UE::AssetRegistry::Impl::FEventContext::EEvent::Added);
+						EventContext.AssetEvents.Emplace(AssetData, UE::AssetRegistry::Impl::FEventContext::EEvent::Added);
 					});
 				}
 			}
-			else
+			else if (State.GetNumAssets() < ARState.GetNumAssets())
 			{
-				// This is unexpected, we should not have called AppendState or in any other way added Assets to the
-				// AssetRegistry before loading the premade AssetRegistry. However, we can handle it with only minor
-				// performance issues.
-
-				// Add the premade state into the global AssetRegistry
-				if (State.GetNumAssets() < ARState.GetNumAssets())
+				FAssetRegistryState ExistingState = MoveTemp(State);
+				State = MoveTemp(ARState);
+				CachePathsFromState(EventContext, State);
+				if (bEmitAssetEvents)
 				{
-					FAssetRegistryState ExistingState = MoveTemp(State);
-					State = MoveTemp(ARState);
-					CachePathsFromState(EventContext, State);
-					if (bEmitAssetEvents)
+					EventContext.AssetEvents.Reserve(State.GetNumAssets());
+					State.EnumerateAllAssets([&EventContext](const FAssetData& AssetData)
 					{
-						EventContext.AssetEvents.Reserve(State.GetNumAssets());
-						State.EnumerateAllAssets([&EventContext](const FAssetData& AssetData)
-							{
-								EventContext.AssetEvents.Emplace(AssetData,
-									UE::AssetRegistry::Impl::FEventContext::EEvent::Added);
-							});
-					}
-					AppendState(EventContext, ExistingState);
-				}
-				else
-				{
-					AppendState(EventContext, ARState, FAssetRegistryState::EInitializationMode::OnlyUpdateNew, bEmitAssetEvents);
+						EventContext.AssetEvents.Emplace(AssetData, UE::AssetRegistry::Impl::FEventContext::EEvent::Added);
+					});
 				}
+				AppendState(EventContext, ExistingState);
+			}
+			else
+			{
+				AppendState(EventContext, ARState, FAssetRegistryState::EInitializationMode::OnlyUpdateNew, bEmitAssetEvents);
 			}
 			UpdatePersistentMountPoints();
 			State.bCookedGlobalAssetRegistryState = true;
@@ -749,7 +729,7 @@ void FAssetRegistryImpl::Initialize(UE::AssetRegistry::FInterfaceWriteScopeLock&
 		}
 	}
 
-	ConsumeOrDeferPreloadedPremade(Context.UARI, Context.Events, ScopeLock);
+	ConsumeOrDeferPreloadedPremade(Context.UARI, Context.Events);
 
 	// Report startup time. This does not include DirectoryWatcher startup time.
 	double StartupDuration = FPlatformTime::Seconds() - StartupStartTime;
@@ -936,8 +916,8 @@ void UAssetRegistryImpl::InitializeEvents(UE::AssetRegistry::Impl::FInitializeCo
 
 	// Listen for new content paths being added or removed at runtime.  These are usually plugin-specific asset paths that
 	// will be loaded a bit later on.
-	FPackageName::OnMountPointMounted().AddUObject(this, &UAssetRegistryImpl::OnMountPointMounted);
-	FPackageName::OnMountPointDismounted().AddUObject(this, &UAssetRegistryImpl::OnMountPointDismounted);
+	FPackageName::OnContentPathMounted().AddUObject(this, &UAssetRegistryImpl::OnContentPathMounted);
+	FPackageName::OnContentPathDismounted().AddUObject(this, &UAssetRegistryImpl::OnContentPathDismounted);
 
 	// If we were called before engine has fully initialized, refresh classes on initialize. If not this won't do anything as it already happened
 	FCoreDelegates::OnPostEngineInit.AddUObject(this, &UAssetRegistryImpl::OnPostEngineInit);
@@ -1718,8 +1698,8 @@ void UAssetRegistryImpl::FinishDestroy()
 		UE::AssetRegistry::FInterfaceWriteScopeLock InterfaceScopeLock(InterfaceLock);
 
 		// Stop listening for content mount point events
-		FPackageName::OnMountPointMounted().RemoveAll(this);
-		FPackageName::OnMountPointDismounted().RemoveAll(this);
+		FPackageName::OnContentPathMounted().RemoveAll(this);
+		FPackageName::OnContentPathDismounted().RemoveAll(this);
 		FCoreDelegates::OnPostEngineInit.RemoveAll(this);
 		FCoreDelegates::OnPreExit.RemoveAll(this);
 		IPluginManager::Get().OnLoadingPhaseComplete().RemoveAll(this);
@@ -5645,7 +5625,7 @@ void FAssetRegistryImpl::OnAdditionalMountSearchCompleted(Impl::FEventContext& E
 	SetPerformanceMode(Impl::EPerformanceMode::MostlyStatic);
 
 	LogSearchDiagnostics(AdditionalMountSearchStartTime, /* bInitialScan */ false);
-	CSV_TRACE_REGION_END(TEXT("Additional Mount Search"), TEXT("AssetRegistry")); // Matching CSV_TRACE_REGION_BEGIN in OnMountPointMounted
+	CSV_TRACE_REGION_END(TEXT("Additional Mount Search"), TEXT("AssetRegistry")); // Matching CSV_TRACE_REGION_BEGIN in OnContentPathMounted
 
 	GlobalGatherer->UpdateCacheForSaving();
 	GlobalGatherer->OnAdditionalMountSearchCompleted();
@@ -8337,7 +8317,7 @@ void FAssetRegistryImpl::OnDirectoryRescanRequired(FInterfaceWriteScopeLock& Sco
 		{
 			if (FPackageName::TryGetMountPointForPath(MountedLongPackageName, UnusedPackageName, MountedFilePath, UnusedRelPath))
 			{
-				FString NormalizeMountedFilePath = UE::AssetRegistry::CreateStandardFilename(MountedFilePath);
+				FString NormalizeMountedFilePath = UE::AssetRegistry::CreateStandardFilename(FString(MountedFilePath));
 				if (FPaths::IsUnderDirectory(NormalizeMountedFilePath, NormalizedDirPath))
 				{
 					DirPathsAndPackageNames.Emplace(MoveTemp(NormalizeMountedFilePath), MoveTemp(MountedLongPackageName));
@@ -8386,7 +8366,7 @@ void FAssetRegistryImpl::OnDirectoryRescanRequired(FInterfaceWriteScopeLock& Sco
 					// Convert Filename to a PackagePath. We know the base dir so its faster to use that than FPackageName
 					// which has to scan all mount dirs
 					FStringView RelPath;
-					FString NormalizedFilename = UE::AssetRegistry::CreateStandardFilename(FStringView(Filename));
+					FString NormalizedFilename = UE::AssetRegistry::CreateStandardFilename(Filename);
 					if (!FPathViews::TryMakeChildPathRelativeTo(NormalizedFilename, LocalPath, RelPath))
 					{
 						return true;
@@ -8772,38 +8752,35 @@ void FAssetRegistryImpl::ScanModifiedAssetFiles(FInterfaceWriteScopeLock& ScopeL
 
 }
 
-void UAssetRegistryImpl::OnMountPointMounted(const TRefCountPtr<UE::PackageName::IMountPoint>& MountPoint)
+void UAssetRegistryImpl::OnContentPathMounted(const FString& InAssetPath, const FString& FileSystemPath)
 {
-	if (SuspendMountingCounter.load(std::memory_order_acquire) > 0)
+	TRACE_CPUPROFILER_EVENT_SCOPE(AssetRegistry_OnContentPathMounted);
+	// Sanitize
+	FString AssetPathWithTrailingSlash;
+	if (!InAssetPath.EndsWith(TEXT("/"), ESearchCase::CaseSensitive))
 	{
-		FScopeLock Guard(&MountingCriticalSection);
-		ListOfMountedMountPoints.Emplace(MountPoint);
-		return;
+		// We actually want a trailing slash here so the path can be properly converted while searching for assets
+		AssetPathWithTrailingSlash = InAssetPath + TEXT("/");
+	}
+	else
+	{
+		AssetPathWithTrailingSlash = InAssetPath;
 	}
-
-	OnMultipleMountPointsMounted(TConstArrayView<TRefCountPtr<UE::PackageName::IMountPoint>>(&MountPoint, 1));
-}
-
-void UAssetRegistryImpl::OnMultipleMountPointsMounted(TConstArrayView<TRefCountPtr<UE::PackageName::IMountPoint>> MountedMountPoints)
-{
-	TRACE_CPUPROFILER_EVENT_SCOPE(AssetRegistry_OnMultipleMountPointsMounted);
-	LLM_SCOPE(ELLMTag::AssetRegistry);
 
 #if WITH_EDITOR
 	FDirectoryWatcherModule& DirectoryWatcherModule = FModuleManager::LoadModuleChecked<FDirectoryWatcherModule>(TEXT("DirectoryWatcher"));
 	IDirectoryWatcher* DirectoryWatcher = nullptr;
 	if (UE::AssetRegistry::Impl::IsDirectoryWatcherEnabled())
 	{
+		LLM_SCOPE(ELLMTag::AssetRegistry);
 		DirectoryWatcher = DirectoryWatcherModule.Get();
-		for (const TRefCountPtr<UE::PackageName::IMountPoint>& MountPoint : MountedMountPoints)
+		if (DirectoryWatcher)
 		{
-			if (DirectoryWatcher)
-			{
-				// Make sure the directory exits on disk so that the OS level DirectoryWatcher can be used to monitor it.
-				IPlatformFile::GetPlatformPhysical().CreateDirectoryTree(*WriteToString<512>(MountPoint->GetLocalPathStandard()));
-			}
+			// Make sure the directory exits on disk so that the OS level DirectoryWatcher can be used to monitor it.
+			IPlatformFile::GetPlatformPhysical().CreateDirectoryTree(*FileSystemPath);
 		}
 	}
+		
 #endif
 
 	UE::AssetRegistry::Impl::FEventContext EventContext;
@@ -8813,29 +8790,25 @@ void UAssetRegistryImpl::OnMultipleMountPointsMounted(TConstArrayView<TRefCountP
 		UE::AssetRegistry::Impl::FClassInheritanceContext InheritanceContext;
 		UE::AssetRegistry::Impl::FClassInheritanceBuffer InheritanceBuffer;
 		GetInheritanceContextWithRequiredLock(InterfaceScopeLock, InheritanceContext, InheritanceBuffer);
-		GuardedData.OnMultipleMountPointsMounted(*this, EventContext,
-			InheritanceContext, InterfaceScopeLock, MountedMountPoints);
+		GuardedData.OnContentPathMounted(InterfaceScopeLock, EventContext, InheritanceContext, InAssetPath, AssetPathWithTrailingSlash,
+			FileSystemPath);
 
 		// Listen for directory changes in this content path
 #if WITH_EDITOR
-		for (const TRefCountPtr<UE::PackageName::IMountPoint>& MountPoint : MountedMountPoints)
+		const FString StandardFileSystemPath = UE::AssetRegistry::CreateStandardFilename(FileSystemPath);
+		// In-game doesn't listen for directory changes
+		if (DirectoryWatcher && !IsDirAlreadyWatchedByRootWatchers(StandardFileSystemPath))
 		{
-			const FString StandardFileSystemPath = UE::AssetRegistry::CreateStandardFilename(MountPoint->GetLocalPathStandard());
-			// In-game doesn't listen for directory changes
-			if (DirectoryWatcher && !IsDirAlreadyWatchedByRootWatchers(StandardFileSystemPath))
+			if (!OnDirectoryChangedDelegateHandles.Contains(AssetPathWithTrailingSlash))
 			{
-				const FString AssetPathWithTrailingSlash = UE::AssetRegistry::CreateStandardLongPackageNameDirectory(MountPoint->GetLongPackageName());
-				if (!OnDirectoryChangedDelegateHandles.Contains(AssetPathWithTrailingSlash))
-				{
-					FDelegateHandle NewHandle;
-					DirectoryWatcher->RegisterDirectoryChangedCallback_Handle(
-						StandardFileSystemPath,
-						IDirectoryWatcher::FDirectoryChanged::CreateUObject(this, &UAssetRegistryImpl::OnDirectoryChanged),
-						NewHandle,
-						IDirectoryWatcher::WatchOptions::IncludeDirectoryChanges);
+				FDelegateHandle NewHandle;
+				DirectoryWatcher->RegisterDirectoryChangedCallback_Handle(
+					StandardFileSystemPath,
+					IDirectoryWatcher::FDirectoryChanged::CreateUObject(this, &UAssetRegistryImpl::OnDirectoryChanged),
+					NewHandle, 
+					IDirectoryWatcher::WatchOptions::IncludeDirectoryChanges);
 
-					OnDirectoryChangedDelegateHandles.Add(AssetPathWithTrailingSlash, NewHandle);
-				}
+				OnDirectoryChangedDelegateHandles.Add(AssetPathWithTrailingSlash, NewHandle);
 			}
 		}
 #endif // WITH_EDITOR
@@ -8847,63 +8820,56 @@ void UAssetRegistryImpl::OnMultipleMountPointsMounted(TConstArrayView<TRefCountP
 namespace UE::AssetRegistry
 {
 
-void FAssetRegistryImpl::OnMultipleMountPointsMounted(UAssetRegistryImpl& UARI, Impl::FEventContext& EventContext,
-	Impl::FClassInheritanceContext& InheritanceContext,
-	UE::AssetRegistry::FInterfaceWriteScopeLock& ScopeLock, 
-	TConstArrayView<TRefCountPtr<UE::PackageName::IMountPoint>> MountedMountPoints)
+void FAssetRegistryImpl::OnContentPathMounted(UE::AssetRegistry::FInterfaceWriteScopeLock& ScopeLock, Impl::FEventContext& EventContext,
+	Impl::FClassInheritanceContext& InheritanceContext, const FString& InAssetPath,
+	const FString& AssetPathWithTrailingSlash, const FString& FileSystemPath)
 {
-	for (const TRefCountPtr<UE::PackageName::IMountPoint>& MountPoint : MountedMountPoints)
-	{
-		FString FileSystemPath(MountPoint->GetLocalPathStandard());
-		const FString AssetPathWithTrailingSlash = UE::AssetRegistry::CreateStandardLongPackageNameDirectory(MountPoint->GetLongPackageName());
+	// Content roots always exist
+	AddPath(EventContext, UE::String::RemoveFromEnd(FStringView(AssetPathWithTrailingSlash), TEXTVIEW("/")));
 
-		// Content roots always exist
-		AddPath(EventContext, UE::String::RemoveFromEnd(FStringView(AssetPathWithTrailingSlash), TEXTVIEW("/")));
-
-		if (GlobalGatherer.IsValid() && bSearchAllAssets)
+	if (GlobalGatherer.IsValid() && bSearchAllAssets)
+	{
+		if (GlobalGatherer->IsSynchronous())
 		{
-			if (GlobalGatherer->IsSynchronous())
-			{
-				Impl::FScanPathContext Context(EventContext, InheritanceContext, { FileSystemPath }, TArray<FString>());
-				ScanPathsSynchronous(&ScopeLock, Context);
-			}
-			else
+			Impl::FScanPathContext Context(EventContext, InheritanceContext, { FileSystemPath }, TArray<FString>());
+			ScanPathsSynchronous(&ScopeLock, Context);
+		}
+		else
+		{
+			if (!IsGathering())
 			{
-				if (!IsGathering())
-				{
-					CSV_TRACE_REGION_BEGIN(TEXT("Additional Mount Search"), TEXT("AssetRegistry")); // Matching CSV_TRACE_REGION_END in OnAdditionalMountSearchCompleted
-					AdditionalMountSearchStartTime = FPlatformTime::Seconds();
-					bAdditionalMountSearchInProgress.store(true, std::memory_order_relaxed);
-					GlobalGatherer->SetIsAdditionalMountSearchInProgress(true);
+				CSV_TRACE_REGION_BEGIN(TEXT("Additional Mount Search"), TEXT("AssetRegistry")); // Matching CSV_TRACE_REGION_END in OnAdditionalMountSearchCompleted
+				AdditionalMountSearchStartTime = FPlatformTime::Seconds();
+				bAdditionalMountSearchInProgress.store(true, std::memory_order_relaxed);
+				GlobalGatherer->SetIsAdditionalMountSearchInProgress(true);
 #if WITH_EDITOR
-					bGatherNeedsApplyObjectRedirects = true;
+				bGatherNeedsApplyObjectRedirects = true;
 #endif
-				}
-
-				GlobalGatherer->AddMountPoint(MountPoint->GetLocalPathStandard(), MountPoint->GetLongPackageName());
-				GlobalGatherer->SetIsOnAllowList(MountPoint->GetLocalPathStandard(), true);
 			}
+
+			GlobalGatherer->AddMountPoint(FileSystemPath, InAssetPath);
+			GlobalGatherer->SetIsOnAllowList(FileSystemPath, true);
 		}
 	}
 }
 
 }
 
-void UAssetRegistryImpl::OnMountPointDismounted(const TRefCountPtr<UE::PackageName::IMountPoint>& MountPoint)
+void UAssetRegistryImpl::OnContentPathDismounted(const FString& InAssetPath, const FString& FileSystemPath)
 {
 	if (SuspendDismountingCounter.load(std::memory_order_acquire) > 0)
 	{
-		FScopeLock Guard(&MountingCriticalSection);
-		ListOfDismountedMountPoints.Emplace(MountPoint);
+		FScopeLock Guard(&DismountingCriticalSection);
+		ListOfUnmountedPaths.Emplace(InAssetPath, FileSystemPath);
 		return;
 	}
 
-	OnMultipleMountPointsDismounted(TConstArrayView<TRefCountPtr<UE::PackageName::IMountPoint>>(&MountPoint, 1));
+	OnMultipleContentPathsDismounted({MakeTuple(InAssetPath, FileSystemPath)});
 }
 
-void UAssetRegistryImpl::OnMultipleMountPointsDismounted(TConstArrayView<TRefCountPtr<UE::PackageName::IMountPoint>> DismountedMountPoints)
+void UAssetRegistryImpl::OnMultipleContentPathsDismounted(TConstArrayView<TTuple<FString, FString>> UnmountedAssetFilesystemPairs)
 {
-	TRACE_CPUPROFILER_EVENT_SCOPE(AssetRegistry_OnMultipleMountPointsDismounted);
+	TRACE_CPUPROFILER_EVENT_SCOPE(AssetRegistry_OnMultipleContentPathsDismounted);
 #if WITH_EDITOR
 	FDirectoryWatcherModule& DirectoryWatcherModule = FModuleManager::LoadModuleChecked<FDirectoryWatcherModule>(TEXT("DirectoryWatcher"));
 	IDirectoryWatcher* DirectoryWatcher = nullptr;
@@ -8917,18 +8883,22 @@ void UAssetRegistryImpl::OnMultipleMountPointsDismounted(TConstArrayView<TRefCou
 	{
 		LLM_SCOPE(ELLMTag::AssetRegistry);
 		UE::AssetRegistry::FInterfaceWriteScopeLock InterfaceScopeLock(InterfaceLock);
-		GuardedData.OnMultipleMountPointsDismounted(*this, EventContext,
-			DismountedMountPoints);
+		GuardedData.OnMultipleContentPathsDismounted(EventContext, UnmountedAssetFilesystemPairs);
 
 		// Stop listening for directory changes in this content path
 #if WITH_EDITOR
-		for (const TRefCountPtr<UE::PackageName::IMountPoint>& MountPoint : DismountedMountPoints)
+		for (const TTuple<FString, FString>& Pair : UnmountedAssetFilesystemPairs)
 		{
-			const FString StandardFileSystemPath = UE::AssetRegistry::CreateStandardFilename(MountPoint->GetLocalPathStandard());
+			const FString StandardFileSystemPath = UE::AssetRegistry::CreateStandardFilename(Pair.Value);
 			if (DirectoryWatcher && !IsDirAlreadyWatchedByRootWatchers(StandardFileSystemPath))
 			{
-				// Make sure OnDirectoryChangedDelegateHandles key is symmetrical with the one used in OnMountPointMounted
-				FString AssetPathWithTrailingSlash = UE::AssetRegistry::CreateStandardLongPackageNameDirectory(MountPoint->GetLongPackageName());
+				// Make sure OnDirectoryChangedDelegateHandles key is symmetrical with the one used in OnContentPathMounted
+				FString AssetPathWithTrailingSlash = Pair.Key;
+				if (!AssetPathWithTrailingSlash.EndsWith(TEXT("/"), ESearchCase::CaseSensitive))
+				{
+					AssetPathWithTrailingSlash += TEXT("/");
+				}
+
 				FDelegateHandle DirectoryChangedHandle;
 				if (ensure(OnDirectoryChangedDelegateHandles.RemoveAndCopyValue(AssetPathWithTrailingSlash, DirectoryChangedHandle)))
 				{
@@ -8944,57 +8914,51 @@ void UAssetRegistryImpl::OnMultipleMountPointsDismounted(TConstArrayView<TRefCou
 namespace UE::AssetRegistry
 {
 
-void FAssetRegistryImpl::OnMultipleMountPointsDismounted(UAssetRegistryImpl& UARI,
+void FAssetRegistryImpl::OnMultipleContentPathsDismounted(
 	Impl::FEventContext& EventContext,
-	TConstArrayView<TRefCountPtr<UE::PackageName::IMountPoint>> DismountedMountPoints)
+	TConstArrayView<TTuple<FString, FString>> UnmountedAssetFilesystemPairs)
 {
 	// For the GlobalGatherer we remove all input MountPoints; we don't need to keep information about their gather.
 	// But for the FAssetData stored inside our AssetRegistryState, for some mount points
 	// we are configured to keep those assets in memory even when dismounted.
-	TArray<TRefCountPtr<UE::PackageName::IMountPoint>> MountPointsToRemoveAssetsFrom;
-	MountPointsToRemoveAssetsFrom.Reserve(DismountedMountPoints.Num());
-	for (const TRefCountPtr<UE::PackageName::IMountPoint>& MountPoint : DismountedMountPoints)
+	TSet<FName> MountPointsToRemoveAssetsFrom;
+	for (const TTuple<FString, FString>& Pair : UnmountedAssetFilesystemPairs)
 	{
 		if (GlobalGatherer.IsValid())
 		{
-			GlobalGatherer->RemoveMountPoint(MountPoint->GetLongPackageName());
+			GlobalGatherer->RemoveMountPoint(Pair.Value);
 		}
-
 		// The input AssetPaths of the unmounted MountPoints are unnormalized and might have a trailing slash,
 		// but we are comparing them by exact string inside our containers 
 		// and need to normalize them to not have a trailing slash.
-		FName AssetPathNoTrailingSlashFName = UE::AssetRegistry::CreateStandardLongPackageFNameNoTrailingSlash(
-			MountPoint->GetLongPackageName());
+		FStringView AssetPathNoTrailingSlash = Pair.Key;
+		if (AssetPathNoTrailingSlash.EndsWith(TEXT("/"), ESearchCase::CaseSensitive))
+		{
+			// We don't want a trailing slash here as it could interfere with RemoveAssetPath
+			AssetPathNoTrailingSlash.LeftChopInline(1);
+		}
+		FName AssetPathNoTrailingSlashFName(AssetPathNoTrailingSlash);
 		if (PersistentMountPoints.Contains(AssetPathNoTrailingSlashFName))
 		{
-			// This path is marked to never remove its AssetDatas.
+			// This path is not marked to never remove its AssetDatas.
 			continue;
 		}
 
-		MountPointsToRemoveAssetsFrom.Add(MountPoint);
+		MountPointsToRemoveAssetsFrom.Add(AssetPathNoTrailingSlashFName);
 	}
-	UnloadMountPoints(EventContext, MountPointsToRemoveAssetsFrom);
-}
 
-void FAssetRegistryImpl::UnloadMountPoints(Impl::FEventContext& EventContext,
-	TConstArrayView<TRefCountPtr<UE::PackageName::IMountPoint>> MountPoints)
-{
 	// Collect the list of paths from given list of roots
 	TSet<FName> AssetPathsToRemoveAssetsFrom;
-	TArray<FName> MountPointRootNoTrailingSlashFNames;
-
-	const bool bRecurse = true;
-	MountPointRootNoTrailingSlashFNames.Reserve(MountPoints.Num());
-	for (const TRefCountPtr<UE::PackageName::IMountPoint>& MountPoint : MountPoints)
 	{
-		FName AssetPathNoTrailingSlashFName = UE::AssetRegistry::CreateStandardLongPackageFNameNoTrailingSlash(
-			MountPoint->GetLongPackageName());
-		CachedPathTree.GetSubPaths(AssetPathNoTrailingSlashFName, AssetPathsToRemoveAssetsFrom, bRecurse);
-		AssetPathsToRemoveAssetsFrom.Add(AssetPathNoTrailingSlashFName);
-		MountPointRootNoTrailingSlashFNames.Add(AssetPathNoTrailingSlashFName);
+		const bool bRecurse = true;
+		for (FName AssetPathNoTrailingSlashFName : MountPointsToRemoveAssetsFrom)
+		{
+			CachedPathTree.GetSubPaths(AssetPathNoTrailingSlashFName, AssetPathsToRemoveAssetsFrom, bRecurse);
+			AssetPathsToRemoveAssetsFrom.Add(AssetPathNoTrailingSlashFName);
+		}
 	}
 
-	// Remove all cached assets and Verse files found in the mountpoints.
+	// Remove all cached assets and Verse files found at this location
 	{
 		TArray<FAssetData*> AllAssetDataToRemove;
 		for (FName PathName : AssetPathsToRemoveAssetsFrom)
@@ -9021,8 +8985,8 @@ void FAssetRegistryImpl::UnloadMountPoints(Impl::FEventContext& EventContext,
 		RemoveAssetDatas(EventContext, AllAssetDataToRemove);
 	}
 
-	// Remove the root paths of the mountpoints.
-	for (FName AssetPathNoTrailingSlashFName : MountPointRootNoTrailingSlashFNames)
+	// Remove the root path
+	for (FName AssetPathNoTrailingSlashFName : MountPointsToRemoveAssetsFrom)
 	{
 		const bool bEvenIfAssetsStillExist = true;
 		RemoveAssetPath(EventContext, AssetPathNoTrailingSlashFName, bEvenIfAssetsStillExist);
@@ -9283,24 +9247,6 @@ bool UAssetRegistryImpl::IsDirAlreadyWatchedByRootWatchers(const FString& Direct
 
 #endif
 
-void UAssetRegistryImpl::SuspendPathsMounting()
-{
-	SuspendMountingCounter.fetch_add(1, std::memory_order_acq_rel);
-}
-
-void UAssetRegistryImpl::ResumePathsMounting()
-{
-	if (SuspendMountingCounter.fetch_sub(1, std::memory_order_acq_rel) == 1)
-	{
-		TArray<TRefCountPtr<UE::PackageName::IMountPoint>> LocalList;
-		{
-			FScopeLock Guard(&MountingCriticalSection);
-			LocalList = MoveTemp(ListOfMountedMountPoints);
-		}
-		OnMultipleMountPointsMounted(LocalList);
-	}
-}
-
 void UAssetRegistryImpl::SuspendPathsDismounting()
 {
 	SuspendDismountingCounter.fetch_add(1, std::memory_order_acq_rel);
@@ -9308,15 +9254,13 @@ void UAssetRegistryImpl::SuspendPathsDismounting()
 
 void UAssetRegistryImpl::ResumePathsDismounting()
 {
+	TArray<TTuple<FString, FString>> LocalList;
 	if (SuspendDismountingCounter.fetch_sub(1, std::memory_order_acq_rel) == 1)
 	{
-		TArray<TRefCountPtr<UE::PackageName::IMountPoint>> LocalList;
-		{
-			FScopeLock Guard(&MountingCriticalSection);
-			LocalList = MoveTemp(ListOfDismountedMountPoints);
-		}
-		OnMultipleMountPointsDismounted(LocalList);
+		FScopeLock Guard(&DismountingCriticalSection);
+		LocalList = MoveTemp(ListOfUnmountedPaths);
 	}
+	OnMultipleContentPathsDismounted(LocalList);
 }
 
 void UAssetRegistryImpl::RequestPauseBackgroundProcessing()
@@ -10907,11 +10851,6 @@ bool FInterruptionContext::ShouldExitEarly()
 
 }
 
-FString CreateStandardFilename(FStringView InPath)
-{
-	return CreateStandardFilename(FString(InPath));
-}
-
 FString CreateStandardFilename(const FString& InPath)
 {
 	FString Result = FPaths::CreateStandardFilename(InPath);
@@ -10930,26 +10869,4 @@ FString CreateStandardFilename(const FString& InPath)
 	return Result;
 }
 
-FString CreateStandardLongPackageNameDirectory(FStringView LongPackageName)
-{
-	FString Result(LongPackageName);
-	if (!Result.EndsWith(TEXT("/"), ESearchCase::CaseSensitive))
-	{
-		Result = Result + TEXT("/");
-	}
-	return Result;
-}
-
-FName CreateStandardLongPackageFNameNoTrailingSlash(FStringView LongPackageName)
-{
-	FStringView ResultNoTrailingSlash = LongPackageName;
-	if (ResultNoTrailingSlash.EndsWith(TEXT("/"), ESearchCase::CaseSensitive))
-	{
-		// We don't want a trailing slash here as it could interfere with RemoveAssetPath
-		ResultNoTrailingSlash.LeftChopInline(1);
-	}
-	return FName(ResultNoTrailingSlash);
-}
-
-
 } // namespace AssetRegistry
diff --git a/Engine/Source/Runtime/AssetRegistry/Private/AssetRegistry.h b/Engine/Source/Runtime/AssetRegistry/Private/AssetRegistry.h
--- a/Engine/Source/Runtime/AssetRegistry/Private/AssetRegistry.h
+++ b/Engine/Source/Runtime/AssetRegistry/Private/AssetRegistry.h
@@ -4,12 +4,10 @@
 
 #include "AssetRegistryImpl.h"
 #include "Containers/Ticker.h"
-#include "Misc/PackageName.h"
 #include "UObject/Object.h"
 
 #include "AssetRegistry.generated.h"
 
-class FAssetDataGatherer;
 namespace UE::AssetRegistry::Premade { struct FAsyncConsumer; }
 
 /**
@@ -112,8 +110,7 @@ public:
 	virtual void ScanModifiedAssetFiles(const TArray<FString>& InFilePaths) override;
 	virtual void ScanModifiedAssetFiles(const TArray<FString>& InFilePaths, UE::AssetRegistry::EScanFlags ScanFlags) override;
 	virtual void Serialize(FArchive& Ar) override;
-	virtual void AppendState(const FAssetRegistryState& InState,
-		UE::AssetRegistry::EAppendMode AppendMode = UE::AssetRegistry::EAppendMode::Append) override;
+	virtual void AppendState(const FAssetRegistryState& InState, UE::AssetRegistry::EAppendMode AppendMode = UE::AssetRegistry::EAppendMode::Append) override;
 	virtual SIZE_T GetAllocatedSize(bool bLogDetailed = false) const override;
 	virtual void LoadPackageRegistryData(FArchive& Ar, FLoadPackageRegistryData& InOutData) const override;
 	virtual void LoadPackageRegistryData(const FString& PackageFilename, FLoadPackageRegistryData& InOutData) const override;
@@ -253,22 +250,29 @@ private:
 #endif
 
 	/**
-	 * Called by the engine core when a new MountPoint is added dynamically at runtime.  This is wired to
+	 * Called by the engine core when a new content path is added dynamically at runtime.  This is wired to
 	 * FPackageName's static delegate.
+	 *
+	 * @param	AssetPath		The new content root asset path that was added (e.g. "/MyPlugin/")
+	 * @param	FileSystemPath	The filesystem path that the AssetPath is mapped to
 	 */
-	void OnMountPointMounted(const TRefCountPtr<UE::PackageName::IMountPoint>& MountPoint);
+	void OnContentPathMounted(const FString& AssetPath, const FString& FileSystemPath);
 
 	/**
-	 * Called by the engine core when a MountPoint is removed dynamically at runtime.  This is wired to
+	 * Called by the engine core when a content path is removed dynamically at runtime.  This is wired to
 	 * FPackageName's static delegate.
+	 *
+	 * @param	AssetPath		The new content root asset path that was added (e.g. "/MyPlugin/")
+	 * @param	FileSystemPath	The filesystem path that the AssetPath is mapped to
 	 */
-	void OnMountPointDismounted(const TRefCountPtr<UE::PackageName::IMountPoint>& MountPoint);
+	void OnContentPathDismounted(const FString& AssetPath, const FString& FileSystemPath);
 
-	/** Batched processing version of OnMountPointMounted. */
-	void OnMultipleMountPointsMounted(TConstArrayView<TRefCountPtr<UE::PackageName::IMountPoint>> MountedMountPoints);
-
-	/** Batched processing version of the OnMountPointDismounted. */
-	void OnMultipleMountPointsDismounted(TConstArrayView<TRefCountPtr<UE::PackageName::IMountPoint>> DismountedMountPoints);
+	/**
+	 * Batched processing version of the OnContentPathDismounted
+	 *
+	 * @param	UnmountedAssetFilesystemPairs	the list of pairs AssetPath - FilesystemPath
+	 */
+	void OnMultipleContentPathsDismounted(TConstArrayView<TTuple<FString, FString>> UnmountedAssetFilesystemPairs);
 
 	/** Called to refresh the native classes list, called at end of engine initialization. */
 	void OnPostEngineInit();
@@ -316,9 +320,7 @@ private:
 	bool IsDirAlreadyWatchedByRootWatchers(const FString& Directory) const;
 #endif
 
-	// Implementation methods for FSuspendPathsMounting and FSuspendPathsDismounting
-	virtual void SuspendPathsMounting() override;
-	virtual void ResumePathsMounting() override;
+	// Implementation methods for FSuspendPathsDismounting
 	virtual void SuspendPathsDismounting() override;
 	virtual void ResumePathsDismounting() override;
 
@@ -445,17 +447,14 @@ private:
 	FTSTicker::FDelegateHandle TickRequestHandle;
 	FCriticalSection DeferredEventsCriticalSection;
 
-	// Suspension and resuming of paths mounting and dismounting. Used for batching many
-	// OnMountPointMounted/OnMountPointDismounted calls together.
-	std::atomic<int32> SuspendMountingCounter = 0;
+	/** Handles suspension and resume of paths dismounting Used for batching many Dismount calls together. */
 	std::atomic<int32> SuspendDismountingCounter = 0;
-	FCriticalSection MountingCriticalSection;
-	TArray<TRefCountPtr<UE::PackageName::IMountPoint>> ListOfMountedMountPoints;
-	TArray<TRefCountPtr<UE::PackageName::IMountPoint>> ListOfDismountedMountPoints;
+	FCriticalSection DismountingCriticalSection;
+	TArray<TTuple<FString, FString>> ListOfUnmountedPaths;
 
 
-	friend FAssetDataGatherer;
-	friend UE::AssetRegistry::FAssetRegistryImpl;
-	friend UE::AssetRegistry::Premade::FAsyncConsumer;
+	friend class FAssetDataGatherer;
+	friend class UE::AssetRegistry::FAssetRegistryImpl;
+	friend struct UE::AssetRegistry::Premade::FAsyncConsumer;
 };
 PRAGMA_ENABLE_DEPRECATION_WARNINGS
diff --git a/Engine/Source/Runtime/AssetRegistry/Private/AssetRegistryImpl.h b/Engine/Source/Runtime/AssetRegistry/Private/AssetRegistryImpl.h
--- a/Engine/Source/Runtime/AssetRegistry/Private/AssetRegistryImpl.h
+++ b/Engine/Source/Runtime/AssetRegistry/Private/AssetRegistryImpl.h
@@ -31,13 +31,14 @@ struct FFileChangeData;
 class FPackageReader;
 class UAssetRegistryImpl;
 namespace UE::AssetRegistry::DependsNode { struct FGatherDependencyDataScratchpad; }
+namespace UE::AssetRegistry::Impl { struct FEventContext; }
 namespace UE::AssetRegistry::Premade { enum class ELoadResult : uint8; }
 namespace UE::AssetRegistry::Premade { struct FAsyncConsumer; }
-namespace UE::AssetRegistry::Impl { struct FEventContext; }
 namespace UE::AssetRegistry::Impl { struct FInitializeContext; }
 namespace UE::AssetRegistry::Impl { struct FScanPathContext; }
 namespace UE::AssetRegistry::Impl { struct FTickContext; }
 namespace UE::AssetRegistry::Impl { struct FClassInheritanceContext; }
+
 #if WITH_EDITOR
 namespace UE::AssetDependencyGatherer::Private { class FRegisteredAssetDependencyGatherer; }
 #endif
@@ -256,12 +257,9 @@ public:
 		Impl::FClassInheritanceContext& InheritanceContext, FString& DirPath, int64 BeforeTimeStamp);
 	void AddLoadedAssetToProcess(const UObject& AssetLoaded);
 #endif
-	void OnMultipleMountPointsMounted(UAssetRegistryImpl& UARI, Impl::FEventContext& EventContext,
-		Impl::FClassInheritanceContext& InheritanceContext,
-		UE::AssetRegistry::FInterfaceWriteScopeLock& ScopeLock,
-		TConstArrayView<TRefCountPtr<UE::PackageName::IMountPoint>> MountedMountPoints);
-	void OnMultipleMountPointsDismounted(UAssetRegistryImpl& UARI, Impl::FEventContext& EventContext,
-		TConstArrayView<TRefCountPtr<UE::PackageName::IMountPoint>> DismountedMountPoints);
+	void OnContentPathMounted(UE::AssetRegistry::FInterfaceWriteScopeLock& ScopeLock, Impl::FEventContext& EventContext, Impl::FClassInheritanceContext& InheritanceContext,
+		const FString& InAssetPath, const FString& AssetPathWithTrailingSlash, const FString& FileSystemPath);
+	void OnMultipleContentPathsDismounted(Impl::FEventContext& EventContext, TConstArrayView<TTuple<FString, FString>> UnmountedAssetFilesystemPairs);
 
 	// Other helper functions called by UAssetRegistryImpl
 
@@ -407,12 +405,7 @@ public:
 	 * For performance reasons, this must be called within the WriteScopeLock, and the caller 
 	 * must handle the possibility that it leaves and reenters the lock.
 	 */
-	void ConditionalLoadPremadeAssetRegistry(UE::AssetRegistry::FInterfaceWriteScopeLock& ScopeLock,
-		UE::AssetRegistry::Impl::FEventContext& EventContext);
-
-	/** Unload specified MountPoints, helper for OnMultipleMountPointsDismounted and others. */
-	void UnloadMountPoints(Impl::FEventContext& EventContext,
-		TConstArrayView<TRefCountPtr<UE::PackageName::IMountPoint>> MountPoints);
+	void ConditionalLoadPremadeAssetRegistry(UE::AssetRegistry::FInterfaceWriteScopeLock& ScopeLock, UE::AssetRegistry::Impl::FEventContext& EventContext);
 
 #if WITH_EDITOR
 	/** Request to pause or resume background processing of scan results.
@@ -535,12 +528,10 @@ private:
 	void UpdateInheritanceBuffer(Impl::FClassInheritanceBuffer& OutBuffer) const;
 
 	/** Conditionally loads the Premade AssetRegistry from disk, or queues it for asynchronous loading if not yet loaded. */
-	void ConsumeOrDeferPreloadedPremade(UAssetRegistryImpl& UARI, Impl::FEventContext& EventContext,
-		FInterfaceWriteScopeLock& ProofOfWriteLock);
+	void ConsumeOrDeferPreloadedPremade(UAssetRegistryImpl& UARI, Impl::FEventContext& EventContext);
 	/** Moves a premade asset registry state into this AR */
-	void LoadPremadeAssetRegistry(UAssetRegistryImpl& UARI, Impl::FEventContext& Context,
-		FInterfaceWriteScopeLock& ProofOfWriteLock,
-		Premade::ELoadResult LoadResult, FAssetRegistryState&& ARState, FStringView StateFilepath);
+	void LoadPremadeAssetRegistry(Impl::FEventContext& Context,
+		Premade::ELoadResult LoadResult, FAssetRegistryState&& ARState);
 	/** Add MountPoints of all AssetDatas currently registered in this->State to the list of PersistentMountPoints. */
 	void UpdatePersistentMountPoints();
 	void OnInitialSearchCompleted(Impl::FEventContext& EventContext);
@@ -1021,12 +1012,7 @@ void EnumerateMemoryAssets(const FARCompiledFilter& InFilter, TSet<FName>& OutPa
 // Normalization functions for paths. We require that normalized paths are case-insensitively string-identical for the same input
 // path specification. e.g. all slashes must be replaced with forward-slashes. Special case: a leading \\ is left unreplaced,
 // so //share will be seen as different than \\share.
-FString CreateStandardFilename(FStringView InPath);
 FString CreateStandardFilename(const FString& InFilename);
-// Normalization for LongPackageNames that identify a directory. Ensures it ends with a trailing forward-slash.
-FString CreateStandardLongPackageNameDirectory(FStringView LongPackageName);
-// Normalization for LongPackageNames that should not end in a trailing slash.
-FName CreateStandardLongPackageFNameNoTrailingSlash(FStringView LongPackageName);
 
 } // namespace UE::AssetRegistry
 
diff --git a/Engine/Source/Runtime/AssetRegistry/Public/AssetRegistry/IAssetRegistry.h b/Engine/Source/Runtime/AssetRegistry/Public/AssetRegistry/IAssetRegistry.h
--- a/Engine/Source/Runtime/AssetRegistry/Public/AssetRegistry/IAssetRegistry.h
+++ b/Engine/Source/Runtime/AssetRegistry/Public/AssetRegistry/IAssetRegistry.h
@@ -1155,20 +1155,7 @@ public:
 	DECLARE_TS_MULTICAST_DELEGATE_TwoParams(FEnumerateAssetsEvent, const FARCompiledFilter& /*Filter*/, UE::AssetRegistry::EEnumerateAssetsFlags /*Flag*/)
 	virtual FEnumerateAssetsEvent& OnEnumerateAssetsEvent() = 0;
 
-	/**
-	 * Request to pause paths mounting.
-	 * Can be used by the calling code to temporarily pause loading of data inside the AssetRegistry in response
-	 * to mounted MountPoints (e.g. ActiveMountLoading). All mounting requests will be executed in one batch at the
-	 * end of that scope.
-	 */
-	struct FSuspendPathsMounting
-	{
-		UE_NODISCARD_CTOR FSuspendPathsMounting() { IAssetRegistry::GetChecked().SuspendPathsMounting(); }
-		~FSuspendPathsMounting() { IAssetRegistry::GetChecked().ResumePathsMounting(); }
-	};
-
-	/**
-	 * Request to pause paths dismounting.
+	/** Request to pause paths dismounting.
 	 * Can be used by the calling code to temporarily pause clearing of data inside the AssetRegistry.
 	 * All unmounting requests will be executed in one batch at the end of that scope.
 	 */
@@ -1180,9 +1167,7 @@ public:
 
 protected:
 
-	// Implementation methods for FSuspendPathsMounting and FSuspendPathsDismounting
-	virtual void SuspendPathsMounting() {};
-	virtual void ResumePathsMounting() {};
+	// Implementation methods for FSuspendPathsDismounting
 	virtual void SuspendPathsDismounting() {};
 	virtual void ResumePathsDismounting() {};
 
diff --git a/Engine/Source/Runtime/CoreUObject/Private/Misc/PackageName.cpp b/Engine/Source/Runtime/CoreUObject/Private/Misc/PackageName.cpp
--- a/Engine/Source/Runtime/CoreUObject/Private/Misc/PackageName.cpp
+++ b/Engine/Source/Runtime/CoreUObject/Private/Misc/PackageName.cpp
@@ -51,10 +51,10 @@ DEFINE_LOG_CATEGORY(LogPackageName);
 
 #define LOCTEXT_NAMESPACE "PackageNames"
 
-FPackageName::FMountPointEvent FPackageName::OnMountPointMountedEvent;
+/** Event that is triggered when a new content path is mounted */
 FPackageName::FOnContentPathMountedEvent FPackageName::OnContentPathMountedEvent;
 
-FPackageName::FMountPointEvent FPackageName::OnMountPointDismountedEvent;
+/** Event that is triggered when a content path is dismounted */
 FPackageName::FOnContentPathDismountedEvent FPackageName::OnContentPathDismountedEvent;
 
 /** Delegate used to check whether a package exist without using the filesystem. */
@@ -175,26 +175,17 @@ FLongPackagePathsSingleton& FLongPackagePathsSingleton::Get()
 	return Singleton;
 }
 
-/**
- * Function to bind to IPluginManager::FRegisterMountPointDelegate, which expects a function with the same
- * arguments as RegisterMountPoint but with a void return value.
- */
-static void RegisterMountPointForPluginManager(const FString& RootPath, const FString& ContentPath)
-{
-	(void)FPackageName::RegisterMountPoint(RootPath, ContentPath);
-}
-
 void FLongPackagePathsSingleton::OnCoreUObjectInitialized()
 {
 	// Allow the plugin manager to mount new content paths by exposing access through a delegate.  PluginManager is 
 	// a Core class, but content path functionality is added at the CoreUObject level.
 	IPluginManager& PluginManager = IPluginManager::Get();
-	PluginManager.SetRegisterMountPointDelegate(IPluginManager::FRegisterMountPointDelegate::CreateStatic(&RegisterMountPointForPluginManager));
+	PluginManager.SetRegisterMountPointDelegate(IPluginManager::FRegisterMountPointDelegate::CreateStatic(&FPackageName::RegisterMountPoint));
 	PluginManager.SetUnRegisterMountPointDelegate(IPluginManager::FRegisterMountPointDelegate::CreateStatic(&FPackageName::UnRegisterMountPoint));
 	UE::CoreUObject::Private::PluginHandler::Install();
 }
 
-bool FLongPackagePathsSingleton::TryNormalizePaths(FStringView InLongPackageName, FStringView InLocalPath,
+bool FLongPackagePathsSingleton::TryNormalizePaths(const FString& InLongPackageName, const FString& InLocalPath,
 	FString& OutLongPackageName, FString& OutLocalPathStandard, FString& OutLocalPathAbsolute, bool bSilenceLogs)
 {
 	if (!TryNormalizeLongPackageNamePath(InLongPackageName, OutLongPackageName, bSilenceLogs))
@@ -208,26 +199,25 @@ bool FLongPackagePathsSingleton::TryNormalizePaths(FStringView InLongPackageName
 	return true;
 }
 
-bool FLongPackagePathsSingleton::TryNormalizeLongPackageNamePath(FStringView InLongPackageName,
+bool FLongPackagePathsSingleton::TryNormalizeLongPackageNamePath(const FString& InLongPackageName,
 	FString& OutLongPackageName, bool bSilenceLogs)
 {
 	int32 NumBackSlashes = 0;
 	int32 LenLongPackageName = InLongPackageName.Len();
 	for (int32 N = 0; N < LenLongPackageName; ++N)
 	{
-		TCHAR C = InLongPackageName.GetData()[N];
+		TCHAR C = (*InLongPackageName)[N];
 		if (C == '\\')
 		{
 			++NumBackSlashes;
 		}
 	}
-	if (LenLongPackageName == 0 || InLongPackageName[0] != '/' || NumBackSlashes > 0)
+	if (!InLongPackageName.StartsWith("/") || NumBackSlashes > 0)
 	{
 		if (!bSilenceLogs)
 		{
 			UE_LOG(LogPackageName, Error,
-				TEXT("Invalid LongPackageName for mount point: \"%.*s\". Callstack:"),
-				InLongPackageName.Len(), InLongPackageName.GetData());
+				TEXT("Invalid LongPackageName for mount point: \"%s\". Callstack:"), *InLongPackageName);
 			FDebug::DumpStackTraceToLog(ELogVerbosity::Display);
 		}
 		return false;
@@ -241,14 +231,14 @@ bool FLongPackagePathsSingleton::TryNormalizeLongPackageNamePath(FStringView InL
 	return true;
 }
 
-bool FLongPackagePathsSingleton::TryNormalizeLocalPath(FStringView InLocalPath,
+bool FLongPackagePathsSingleton::TryNormalizeLocalPath(const FString& InLocalPath,
 	FString& OutLocalPathStandard, FString& OutLocalPathAbsolute, bool bSilenceLogs)
 {
 	// If a relative path is passed, convert to an absolute path 
 	OutLocalPathAbsolute = InLocalPath;
 	if (FPaths::IsRelative(OutLocalPathAbsolute))
 	{
-		OutLocalPathAbsolute = FPaths::ConvertRelativePathToFull(OutLocalPathAbsolute);
+		OutLocalPathAbsolute = FPaths::ConvertRelativePathToFull(InLocalPath);
 
 		// Revert to original path if unable to convert to full path
 		if (OutLocalPathAbsolute.Len() <= 1)
@@ -256,8 +246,7 @@ bool FLongPackagePathsSingleton::TryNormalizeLocalPath(FStringView InLocalPath,
 			OutLocalPathAbsolute = InLocalPath;
 			if (!bSilenceLogs)
 			{
-				UE_LOG(LogPackageName, Warning, TEXT("Unable to convert mount point relative path: %.*s"),
-					InLocalPath.Len(), InLocalPath.GetData());
+				UE_LOG(LogPackageName, Warning, TEXT("Unable to convert mount point relative path: %s"), *InLocalPath);
 			}
 		}
 	}
@@ -280,7 +269,7 @@ bool FLongPackagePathsSingleton::TryNormalizeLocalPath(FStringView InLocalPath,
 	return true;
 }
 
-TRefCountPtr<UE::PackageName::IMountPoint> FLongPackagePathsSingleton::InsertMountPoint(const FString& RootPath, const FString& ContentPath)
+void FLongPackagePathsSingleton::InsertMountPoint(const FString& RootPath, const FString& ContentPath)
 {
 	using namespace UE::PackageName;
 
@@ -289,12 +278,11 @@ TRefCountPtr<UE::PackageName::IMountPoint> FLongPackagePathsSingleton::InsertMou
 	if (!TryNormalizePaths(RootPath, ContentPath, SearchKey.LongPackageName,
 		SearchKey.LocalPathStandard, LocalPathAbsolute, false /* bSilenceLogs */))
 	{
-		return TRefCountPtr<UE::PackageName::IMountPoint>();
+		return;
 	}
 	uint32 SearchHash = GetTypeHash(SearchKey);
 
-	TRefCountPtr<IMountPoint> PublicMountPoint;
-	FMountPoint* MountPoint = nullptr;
+	TRefCountPtr<FMountPoint> MountPoint;
 	{
 		UE::TWriteScopeLock ScopeLock(MountLock);
 
@@ -310,13 +298,12 @@ TRefCountPtr<UE::PackageName::IMountPoint> FLongPackagePathsSingleton::InsertMou
 				FString(LocalPathAbsolute), EMountFlags::Removable);
 			MountPoints.AddByHash(SearchHash, FMountPointKeyByPtr(*MountPoint));
 		}
-		PublicMountPoint = MountPoint;
-		RootPathTree.FindOrAdd(MountPoint->GetRootPath()) = MountPoint;
-		ContentPathTree.FindOrAdd(MountPoint->GetContentPathRelative()) = MountPoint;
-		ContentPathTree.FindOrAdd(MountPoint->GetContentPathAbsolute()) = MountPoint;
+		RootPathTree.FindOrAdd(MountPoint->GetRootPath()) = MountPoint.GetReference();
+		ContentPathTree.FindOrAdd(MountPoint->GetContentPathRelative()) = MountPoint.GetReference();
+		ContentPathTree.FindOrAdd(MountPoint->GetContentPathAbsolute()) = MountPoint.GetReference();
 
 		bool bRemoveFromContainer;
-		MountPoint->SetMounted(true, ScopeLock, PublicMountPoint, bRemoveFromContainer);
+		MountPoint->SetMounted(true, ScopeLock, MountPoint, bRemoveFromContainer);
 		// Set mount order to highest priority, replacing earlier mount order if it already existed.
 		MountPoint->SetMountOrder(NextMountOrder++);
 		if (NextMountOrder == 0)
@@ -327,21 +314,10 @@ TRefCountPtr<UE::PackageName::IMountPoint> FLongPackagePathsSingleton::InsertMou
 		check(!bRemoveFromContainer); // Should only be settable to true when calling SetMounted(false).
 	}
 
-	BroadcastMountPointMounted(PublicMountPoint, MountPoint);
-	return PublicMountPoint;
-}
-
-void FLongPackagePathsSingleton::BroadcastMountPointMounted(
-	const TRefCountPtr<UE::PackageName::IMountPoint>& PublicMountPoint, UE::PackageName::FMountPoint* MountPoint)
-{
-	using namespace UE::PackageName;
-
-	check(PublicMountPoint && PublicMountPoint.GetReference() == MountPoint);
-
 	// Downgrade the log messages on startup to Verbose to reduce startup spam
 #if !NO_LOGGING
 	FString Message = FString::Printf(TEXT("FPackageName: Mount point added: '%s' mounted to '%s'"),
-		*MountPoint->GetContentPathRelative(), *MountPoint->GetRootPath());
+		*SearchKey.LocalPathStandard, *SearchKey.LongPackageName);
 	if (IsEngineStartupModuleLoadingComplete())
 	{
 		UE_LOG(LogPackageName, Log, TEXT("%s"), *Message);
@@ -353,50 +329,7 @@ void FLongPackagePathsSingleton::BroadcastMountPointMounted(
 #endif
 
 	// Let subscribers know that a new content path was mounted
-	FPackageName::OnMountPointMounted().Broadcast(PublicMountPoint);
-	FPackageName::OnContentPathMounted().Broadcast(MountPoint->GetRootPath(), MountPoint->GetContentPathRelative());
-}
-
-void FLongPackagePathsSingleton::InsertMountPoint(const TRefCountPtr<UE::PackageName::IMountPoint>& PublicMountPoint)
-{
-	using namespace UE::PackageName;
-
-	if (!PublicMountPoint)
-	{
-		return;
-	}
-	FMountPoint* MountPoint = PublicMountPoint->AsFMountPoint();
-	if (!MountPoint)
-	{
-		UE_LOG(LogPackageName, Error,
-			TEXT("IMountPoint passed to FPackageName that is not an FMountPoint subclass. This is unexpected, all IMountPoints should be FMountPoint. ")
-			TEXT("LongPackageName: %.*s, LocalPathStandard: %.*s"),
-			PublicMountPoint->GetLongPackageName().Len(), PublicMountPoint->GetLongPackageName().GetData(),
-			PublicMountPoint->GetLocalPathStandard().Len(), PublicMountPoint->GetLocalPathStandard().GetData());
-		FDebug::DumpStackTraceToLog(ELogVerbosity::Display);
-		return;
-	}
-
-	{
-		UE::TWriteScopeLock ScopeLock(MountLock);
-
-		RootPathTree.FindOrAdd(MountPoint->GetRootPath()) = MountPoint;
-		ContentPathTree.FindOrAdd(MountPoint->GetContentPathRelative()) = MountPoint;
-		ContentPathTree.FindOrAdd(MountPoint->GetContentPathAbsolute()) = MountPoint;
-
-		bool bRemoveFromContainer;
-		MountPoint->SetMounted(true, ScopeLock, PublicMountPoint, bRemoveFromContainer);
-		// Set mount order to highest priority, replacing earlier mount order if it already existed.
-		MountPoint->SetMountOrder(NextMountOrder++);
-		if (NextMountOrder == 0)
-		{
-			// Overflow, due to e.g. removing and adding a mountpoint too many times.
-			ReassignMountOrders();
-		}
-		check(!bRemoveFromContainer); // Should only be settable to true when calling SetMounted(false).
-	}
-
-	BroadcastMountPointMounted(PublicMountPoint, MountPoint);
+	FPackageName::OnContentPathMounted().Broadcast( SearchKey.LongPackageName, SearchKey.LocalPathStandard);
 }
 
 void FLongPackagePathsSingleton::ReassignMountOrders()
@@ -440,8 +373,8 @@ void FLongPackagePathsSingleton::RemoveMountPoint(const FString& RootPath, const
 	}
 	uint32 SearchHash = GetTypeHash(SearchKey);
 
-	TRefCountPtr<IMountPoint> PublicMountPoint;
-	FMountPoint* MountPoint = nullptr;
+	TRefCountPtr<FMountPoint> ExistingMount;
+	bool bUnmounted = false;
 	{
 		UE::TWriteScopeLock ScopeLock(MountLock);
 		FMountPointKeyByPtr* ExistingPtr = MountPoints.FindByHash(SearchHash, SearchKey);
@@ -452,328 +385,93 @@ void FLongPackagePathsSingleton::RemoveMountPoint(const FString& RootPath, const
 				*SearchKey.LocalPathStandard, *SearchKey.LongPackageName);
 			return;
 		}
-		MountPoint = ExistingPtr->GetMountPoint();
-		PublicMountPoint = MountPoint;
-		if (!TryRemoveMountPointInternal(PublicMountPoint, ScopeLock, MountPoint))
+		ExistingMount = ExistingPtr->GetMountPoint();
+		if (!ExistingMount->IsRemovable())
 		{
+			UE_LOG(LogPackageName, Display, TEXT("FPackageName: Mount point remove failed: MountPoint mapping '%s' to '%s' is a default mapping and is not removable."),
+				*SearchKey.LocalPathStandard, *SearchKey.LongPackageName);
 			return;
 		}
-	}
-	BroadcastMountPointDismounted(PublicMountPoint, MountPoint);
-}
 
-bool FLongPackagePathsSingleton::TryRemoveMountPointInternal(
-	const TRefCountPtr<UE::PackageName::IMountPoint>& InternallyHeldMountPoint,
-	UE::TWriteScopeLock<FTransactionallySafeRWLock>& ProofOfWriteLock, UE::PackageName::FMountPoint* MountPoint)
-{
-	using namespace UE::PackageName;
-
-	// Called within the WriteLock
-	// Caller must hold a single internal refcount in addition to any externally held refcount; that single internal
-	// refcount is required by the SetMounted call within this function to decide whether the MountPoint is externally referenced
-	if (!MountPoint->IsRemovable())
-	{
-		UE_LOG(LogPackageName, Display, TEXT("FPackageName: Mount point remove failed: MountPoint mapping '%s' to '%s' is a default mapping and is not removable."),
-			*MountPoint->GetContentPathRelative(), *MountPoint->GetRootPath());
-		return false;
-	}
-
-	bool bRootPathExisted;
-	bool bContentPathRelativeExisted;
-	bool bContentPathAbsoluteExisted;
-	RootPathTree.Remove(MountPoint->GetRootPath(), &bRootPathExisted);
-	ContentPathTree.Remove(MountPoint->GetContentPathRelative(), &bContentPathRelativeExisted);
-	ContentPathTree.Remove(MountPoint->GetContentPathAbsolute(), &bContentPathAbsoluteExisted);
-	if (bRootPathExisted || bContentPathRelativeExisted || bContentPathAbsoluteExisted)
-	{
-		// See if there is an existing mounted MountPoint that was overridden by the removed mount point and now
-		// needs to be published as the MountPoint for one of its paths.
-		FMountPoint* NewRootPath = nullptr;
-		FMountPoint* NewContentPathRelative = nullptr;
-		FMountPoint* NewContentPathAbsolute = nullptr;
-		for (const FMountPointKeyByPtr& OtherMountPtr : MountPoints)
-		{
-			FMountPoint* OtherMount = OtherMountPtr.GetMountPoint();
-			CA_ASSUME(OtherMount != nullptr);
-			if (OtherMount == MountPoint || !OtherMount->IsMounted())
-			{
-				continue;
-			}
-			if (bRootPathExisted && OtherMount->GetRootPath().Equals(
-				MountPoint->GetRootPath(), ESearchCase::IgnoreCase))
+		bool bRootPathExisted;
+		bool bContentPathRelativeExisted;
+		bool bContentPathAbsoluteExisted;
+		RootPathTree.Remove(ExistingMount->GetRootPath(), &bRootPathExisted);
+		ContentPathTree.Remove(ExistingMount->GetContentPathRelative(), &bContentPathRelativeExisted);
+		ContentPathTree.Remove(ExistingMount->GetContentPathAbsolute(), &bContentPathAbsoluteExisted);
+		if (bRootPathExisted || bContentPathRelativeExisted || bContentPathAbsoluteExisted)
+		{
+			// See if there is an existing mounted MountPoint that was overridden by the removed mount point and now
+			// needs to be published as the MountPoint for one of its paths.
+			FMountPoint* NewRootPath = nullptr;
+			FMountPoint* NewContentPathRelative = nullptr;
+			FMountPoint* NewContentPathAbsolute = nullptr;
+			for (const FMountPointKeyByPtr& OtherMountPtr : MountPoints)
 			{
-				if (NewRootPath == nullptr
-					|| OtherMount->GetMountOrder() > NewRootPath->GetMountOrder())
+				FMountPoint* OtherMount = OtherMountPtr.GetMountPoint();
+				CA_ASSUME(OtherMount != nullptr);
+				if (OtherMount == ExistingMount.GetReference()
+					|| !OtherMount->IsMounted())
 				{
-					NewRootPath = OtherMount;
+					continue;
 				}
-			}
-			if (bContentPathRelativeExisted && OtherMount->GetContentPathRelative().Equals(
-				MountPoint->GetContentPathRelative(), ESearchCase::IgnoreCase))
-			{
-				if (NewContentPathRelative == nullptr
-					|| OtherMount->GetMountOrder() > NewContentPathRelative->GetMountOrder())
+				if (bRootPathExisted && OtherMount->GetRootPath().Equals(
+					SearchKey.LongPackageName, ESearchCase::IgnoreCase))
 				{
-					NewContentPathRelative = OtherMount;
+					if (NewRootPath == nullptr
+						|| OtherMount->GetMountOrder() > NewRootPath->GetMountOrder())
+					{
+						NewRootPath = OtherMount;
+					}
 				}
-			}
-			if (bContentPathAbsoluteExisted && OtherMount->GetContentPathAbsolute().Equals(
-				MountPoint->GetContentPathAbsolute(), ESearchCase::IgnoreCase))
-			{
-				if (NewContentPathAbsolute == nullptr
-					|| OtherMount->GetMountOrder() > NewContentPathAbsolute->GetMountOrder())
+				if (bContentPathRelativeExisted && OtherMount->GetContentPathRelative().Equals(
+					SearchKey.LocalPathStandard, ESearchCase::IgnoreCase))
 				{
-					NewContentPathAbsolute = OtherMount;
+					if (NewContentPathRelative == nullptr
+						|| OtherMount->GetMountOrder() > NewContentPathRelative->GetMountOrder())
+					{
+						NewContentPathRelative = OtherMount;
+					}
 				}
-			}
-		}
-		if (NewRootPath)
-		{
-			RootPathTree.FindOrAdd(NewRootPath->GetRootPath()) = NewRootPath;
-		}
-		if (NewContentPathRelative)
-		{
-			ContentPathTree.FindOrAdd(NewContentPathRelative->GetContentPathRelative()) = NewContentPathRelative;
-		}
-		if (NewContentPathAbsolute)
-		{
-			ContentPathTree.FindOrAdd(NewContentPathAbsolute->GetContentPathAbsolute()) = NewContentPathAbsolute;
-		}
-	}
-
-	bool bRemoveFromContainer;
-	MountPoint->SetMounted(false, ProofOfWriteLock, InternallyHeldMountPoint, bRemoveFromContainer);
-	if (bRemoveFromContainer)
-	{
-		MountPoints.Remove(FMountPointKeyByPtr(*MountPoint));
-	}
-
-	return true;
-}
-
-void FLongPackagePathsSingleton::BroadcastMountPointDismounted(
-	const TRefCountPtr<UE::PackageName::IMountPoint>& PublicMountPoint, UE::PackageName::FMountPoint* MountPoint)
-{
-	using namespace UE::PackageName;
-
-	check(PublicMountPoint && PublicMountPoint.GetReference() == MountPoint);
-	UE_LOG(LogPackageName, Display, TEXT("FPackageName: Mount point removed: '%s' unmounted from '%s'"),
-		*MountPoint->GetContentPathRelative(), *MountPoint->GetRootPath());
-	// Let subscribers know that a content path was unmounted
-	FPackageName::OnMountPointDismounted().Broadcast(PublicMountPoint);
-	FPackageName::OnContentPathDismounted().Broadcast(MountPoint->GetRootPath(), MountPoint->GetContentPathRelative());
-}
-
-void FLongPackagePathsSingleton::RemoveMountPoint(const TRefCountPtr<UE::PackageName::IMountPoint>& PublicMountPoint)
-{
-	using namespace UE::PackageName;
-
-	if (!PublicMountPoint)
-	{
-		return;
-	}
-	FMountPoint* MountPoint = PublicMountPoint->AsFMountPoint();
-	if (!MountPoint)
-	{
-		UE_LOG(LogPackageName, Error,
-			TEXT("IMountPoint passed to FPackageName that is not an FMountPoint subclass. This is unexpected, all IMountPoints should be FMountPoint. ")
-			TEXT("LongPackageName: %.*s, LocalPathStandard: %.*s"),
-			PublicMountPoint->GetLongPackageName().Len(), PublicMountPoint->GetLongPackageName().GetData(),
-			PublicMountPoint->GetLocalPathStandard().Len(), PublicMountPoint->GetLocalPathStandard().GetData());
-		FDebug::DumpStackTraceToLog(ELogVerbosity::Display);
-		return;
-	}
-
-	// The SetMounted function assumes we take an additional refcount in this function, and uses that to decide
-	// that the MountPoint is no longer externally referenced. It is externally referenced (our caller passed in
-	// a refcountptr to it that they still have), so we need to increment refcount to let SetMounted know that
-	// it is externally referenced.
-	TRefCountPtr<IMountPoint> InternallyHeldMountPoint = PublicMountPoint;
-	{
-		UE::TWriteScopeLock ScopeLock(MountLock);
-		if (!TryRemoveMountPointInternal(InternallyHeldMountPoint, ScopeLock, MountPoint))
-		{
-			return;
-		}
-	}
-	BroadcastMountPointDismounted(InternallyHeldMountPoint, MountPoint);
-}
-
-TRefCountPtr<UE::PackageName::IMountPoint>
-FLongPackagePathsSingleton::FindMountPointByRootPackageName(FStringView RootLongPackageName)
-{
-	using namespace UE::PackageName;
-
-	FString LongPackageName;
-	if (!TryNormalizeLongPackageNamePath(RootLongPackageName, LongPackageName, true /* bSilenceLogs */))
-	{
-		return TRefCountPtr<UE::PackageName::IMountPoint>();
-	}
-
-	{
-		UE::TReadScopeLock ScopeLock(MountLock);
-
-		// First look for a Mounted MountPoint because its a common case that we can do fast
-		FMountPoint** ExistingPtr = RootPathTree.Find(LongPackageName);
-		if (ExistingPtr)
-		{
-			return TRefCountPtr<IMountPoint>(*ExistingPtr);
-		}
-
-		// We don't have a lookup map by RootName for overridden or non-mounted MountPoints, so search our entire list.
-		FMountPoint* HighestPriorityExisting = nullptr;
-		for (const FMountPointKeyByPtr& ExistingKey : MountPoints)
-		{
-			FMountPoint* Existing = ExistingKey.GetMountPoint();
-			if (Existing->GetLongPackageName().Equals(LongPackageName, ESearchCase::IgnoreCase))
-			{
-				if (FMountPoint::IsHigherOverridePriority(Existing, HighestPriorityExisting))
+				if (bContentPathAbsoluteExisted && OtherMount->GetContentPathAbsolute().Equals(
+					LocalPathAbsolute, ESearchCase::IgnoreCase))
 				{
-					HighestPriorityExisting = Existing;
+					if (NewContentPathAbsolute == nullptr
+						|| OtherMount->GetMountOrder() > NewContentPathAbsolute->GetMountOrder())
+					{
+						NewContentPathAbsolute = OtherMount;
+					}
 				}
 			}
-		}
-		if (HighestPriorityExisting)
-		{
-			return TRefCountPtr<IMountPoint>(HighestPriorityExisting);
-		}
-	}
-	return TRefCountPtr<UE::PackageName::IMountPoint>();
-}
-
-TRefCountPtr<UE::PackageName::IMountPoint>
-FLongPackagePathsSingleton::FindMountPointByRootLocalPath(FStringView RootLocalPath)
-{
-	using namespace UE::PackageName;
-
-	FString LocalPathStandard;
-	FString LocalPathAbsolute;
-	if (!TryNormalizeLocalPath(RootLocalPath, LocalPathStandard, LocalPathAbsolute, true /* bSilenceLogs */))
-	{
-		return TRefCountPtr<UE::PackageName::IMountPoint>();
-	}
-
-	{
-		UE::TReadScopeLock ScopeLock(MountLock);
-
-		// First look for a Mounted MountPoint because its a common case that we can do fast
-		FMountPoint** ExistingPtr = ContentPathTree.Find(LocalPathStandard);
-		if (ExistingPtr)
-		{
-			return TRefCountPtr<IMountPoint>(*ExistingPtr);
-		}
-
-		// We don't have a lookup map by LocalPath for overridden or non-mounted MountPoints, so search our entire list.
-		FMountPoint* HighestPriorityExisting = nullptr;
-		for (const FMountPointKeyByPtr& ExistingKey : MountPoints)
-		{
-			FMountPoint* Existing = ExistingKey.GetMountPoint();
-			if (Existing->GetLocalPathStandard().Equals(LocalPathStandard, ESearchCase::IgnoreCase))
+			if (NewRootPath)
 			{
-				if (FMountPoint::IsHigherOverridePriority(Existing, HighestPriorityExisting))
-				{
-					HighestPriorityExisting = Existing;
-				}
+				RootPathTree.FindOrAdd(NewRootPath->GetRootPath()) = NewRootPath;
+			}
+			if (NewContentPathRelative)
+			{
+				ContentPathTree.FindOrAdd(NewContentPathRelative->GetContentPathRelative()) = NewContentPathRelative;
+			}
+			if (NewContentPathAbsolute)
+			{
+				ContentPathTree.FindOrAdd(NewContentPathAbsolute->GetContentPathAbsolute()) = NewContentPathAbsolute;
 			}
 		}
-		if (HighestPriorityExisting)
-		{
-			return TRefCountPtr<IMountPoint>(HighestPriorityExisting);
-		}
-	}
-	return TRefCountPtr<UE::PackageName::IMountPoint>();
-}
-
-TRefCountPtr<UE::PackageName::IMountPoint>
-FLongPackagePathsSingleton::FindMountPoint(FStringView RootLongPackageName, FStringView RootLocalPath)
-{
-	using namespace UE::PackageName;
-
-	FMountPointKeyByStrings SearchKey;
-	FString LocalPathAbsolute;
-	if (!TryNormalizePaths(RootLongPackageName, RootLocalPath, SearchKey.LongPackageName,
-		SearchKey.LocalPathStandard, LocalPathAbsolute, true /* bSilenceLogs */))
-	{
-		return TRefCountPtr<UE::PackageName::IMountPoint>();
-	}
-	uint32 SearchHash = GetTypeHash(SearchKey);
-
-	{
-		UE::TReadScopeLock ScopeLock(MountLock);
-		FMountPoint* MountPoint = nullptr;
-
-		FMountPointKeyByPtr* Existing = MountPoints.FindByHash(SearchHash, SearchKey);
-		if (Existing)
-		{
-			return TRefCountPtr<IMountPoint>(Existing->GetMountPoint());
-		}
-		return TRefCountPtr<UE::PackageName::IMountPoint>();
-	}
-}
-
-TRefCountPtr<UE::PackageName::IMountPoint>
-FLongPackagePathsSingleton::FindOrAddMountPoint(FStringView RootLongPackageName, FStringView RootLocalPath)
-{
-	using namespace UE::PackageName;
-
-	FMountPointKeyByStrings SearchKey;
-	FString LocalPathAbsolute;
-	if (!TryNormalizePaths(RootLongPackageName, RootLocalPath, SearchKey.LongPackageName,
-		SearchKey.LocalPathStandard, LocalPathAbsolute, false /* bSilenceLogs */))
-	{
-		return TRefCountPtr<UE::PackageName::IMountPoint>();
-	}
-	uint32 SearchHash = GetTypeHash(SearchKey);
-
-	// Look in the existing MountPoints to see if this one already exists, so we can return the existing pointer.
-	{
-		UE::TReadScopeLock ReadScopeLock(MountLock);
-		FMountPointKeyByPtr* Existing = MountPoints.FindByHash(SearchHash, SearchKey);
-		if (Existing)
-		{
-			return TRefCountPtr<IMountPoint>(Existing->GetMountPoint());
-		}
-	}
 
-	{
-		UE::TWriteScopeLock WriteScopeLock(MountLock);
-
-		// Look again in MountPoints in case it was added when we dropped the readlock.
-		FMountPointKeyByPtr* Existing = MountPoints.FindByHash(SearchHash, SearchKey);
-		if (Existing)
+		bool bRemoveFromContainer;
+		ExistingMount->SetMounted(false, ScopeLock, ExistingMount, bRemoveFromContainer);
+		if (bRemoveFromContainer)
 		{
-			return TRefCountPtr<IMountPoint>(Existing->GetMountPoint());
+			MountPoints.RemoveByHash(SearchHash, SearchKey);
 		}
-
-		FMountPoint* MountPoint = new FMountPoint(MoveTemp(SearchKey.LongPackageName),
-			MoveTemp(SearchKey.LocalPathStandard), MoveTemp(LocalPathAbsolute), EMountFlags::Removable);
-		MountPoints.AddByHash(SearchHash, FMountPointKeyByPtr(*MountPoint));
-		return TRefCountPtr<UE::PackageName::IMountPoint>(MountPoint);
 	}
-}
-
-TRefCountPtr<UE::PackageName::IMountPoint>
-FLongPackagePathsSingleton::FindMountPointByChildLongPackageName(FStringView ChildLongPackageName)
-{
-	using namespace UE::PackageName;
 
-	FLongPackagePathsSingleton& Paths = FLongPackagePathsSingleton::Get();
-
-	UE::TReadScopeLock ScopeLock(Paths.MountLock);
-	FMountPoint* const* MountPointPtr = Paths.RootPathTree.FindClosestValue(ChildLongPackageName);
-	if (!MountPointPtr)
+	if (ExistingMount)
 	{
-		return TRefCountPtr<IMountPoint>();
+		UE_LOG(LogPackageName, Display, TEXT("FPackageName: Mount point removed: '%s' unmounted from '%s'"),
+			*SearchKey.LocalPathStandard, *SearchKey.LongPackageName);
+		// Let subscribers know that a content path was unmounted
+		FPackageName::OnContentPathDismounted().Broadcast(SearchKey.LongPackageName, SearchKey.LocalPathStandard);
 	}
-	FMountPoint* MountPoint = *MountPointPtr;
-	// DirectoryTree is less conservative than IsParentPathOf; e.g. d:/Dir//Root will match
-	// d:/Dir/Root/Child in DirectoryTree but not in IsParentPathOf. Treat this as not found,
-	// since no mountdir will match it according to IsParentPathOf.
-	if (!FPathViews::IsParentPathOf(MountPoint->GetRootPath(), ChildLongPackageName))
-	{
-		return TRefCountPtr<IMountPoint>();
-	}
-	return TRefCountPtr<IMountPoint>(MountPoint);
 }
 
 bool FLongPackagePathsSingleton::MountPointExists(FStringView RootPath)
@@ -846,16 +544,15 @@ FLongPackagePathsSingleton::FLongPackagePathsSingleton()
 			return;
 		}
 
-		FMountPoint* MountPoint = new FMountPoint(FString(RootPath), FString(ContentPathRelative),
+		TRefCountPtr<FMountPoint> MountPoint = new FMountPoint(FString(RootPath), FString(ContentPathRelative),
 			FString(AbsolutePathBuilder), InFlags);
-		TRefCountPtr<IMountPoint> MountPointRefCount = MountPoint;
 		MountPoint->SetMountOrder(NextMountOrder++);
 		bool bAlreadyExisted;
 		MountPoints.Add(FMountPointKeyByPtr(*MountPoint), &bAlreadyExisted);
 		check(!bAlreadyExisted); // Impossible for it to already exist, because we called Find above.
 
 		bool bRemoveFromContainer;
-		MountPoint->SetMounted(true, ScopeLock, MountPointRefCount, bRemoveFromContainer);
+		MountPoint->SetMounted(true, ScopeLock, MountPoint, bRemoveFromContainer);
 		check(!bRemoveFromContainer); // Should only be settable to true when calling SetMounted(false).
 
 		// We are appending to the MountPoint list, and will have only a few duplicates. If any duplicates do occur,
@@ -864,7 +561,7 @@ FLongPackagePathsSingleton::FLongPackagePathsSingleton()
 		FMountPoint*& ExistingRootPath = RootPathTree.FindOrAdd(MountPoint->GetRootPath());
 		if (!ExistingRootPath)
 		{
-			ExistingRootPath = MountPoint;
+			ExistingRootPath = MountPoint.GetReference();
 		}
 		else
 		{
@@ -876,7 +573,7 @@ FLongPackagePathsSingleton::FLongPackagePathsSingleton()
 		FMountPoint*& ExistingContentPathRelative = ContentPathTree.FindOrAdd(MountPoint->GetContentPathRelative());
 		if (!ExistingContentPathRelative)
 		{
-			ExistingContentPathRelative = MountPoint;
+			ExistingContentPathRelative = MountPoint.GetReference();
 		}
 		else
 		{
@@ -888,7 +585,7 @@ FLongPackagePathsSingleton::FLongPackagePathsSingleton()
 		FMountPoint*& ExistingContentPathAbsolute = ContentPathTree.FindOrAdd(MountPoint->GetContentPathAbsolute());
 		if (!ExistingContentPathAbsolute)
 		{
-			ExistingContentPathAbsolute = MountPoint;
+			ExistingContentPathAbsolute = MountPoint.GetReference();
 		}
 		else
 		{
@@ -930,19 +627,18 @@ FLongPackagePathsSingleton::~FLongPackagePathsSingleton()
 	using namespace UE::PackageName;
 
 	// We require that our external callers drop their IMountPoint refcounts, but we do not require that they unmount
-	// all of the MountPoints they registered with RegisterMountPoint. So we may have MountPoints remaining that
-	// are still mounted and therefore we still hold a refcount. Unmount them all now to drop that final refcount and
+	// all of the MountPoints they registered with RegisterMountPoint. So we may have MountPoints remaining with that
+	// our still mounted and therefore we still hold a refcount. Unmount them all now to drop that final refcount and
 	// free the memory. This happens at end of process, so this is important mainly for leak detection.
-	TArray<TRefCountPtr<IMountPoint>> RemainingMountPoints;
+	TArray<TRefCountPtr<FMountPoint>> RemainingMountPoints;
 	{
 		UE::TWriteScopeLock ScopeLock(MountLock);
 		RemainingMountPoints.Reserve(MountPoints.Num());
 		for (FMountPointKeyByPtr& Key : MountPoints)
 		{
-			FMountPoint* MountPoint = Key.GetMountPoint();
-			TRefCountPtr<IMountPoint>& MountPointRefCount = RemainingMountPoints.Emplace_GetRef(MountPoint);
+			TRefCountPtr<FMountPoint>& MountPoint = RemainingMountPoints.Emplace_GetRef(Key.GetMountPoint());
 			bool bRemoveFromContainer;
-			MountPoint->SetMounted(false, ScopeLock, MountPointRefCount, bRemoveFromContainer);
+			MountPoint->SetMounted(false, ScopeLock, MountPoint, bRemoveFromContainer);
 			if (!bRemoveFromContainer)
 			{
 				// UE_LOG is likely no longer available; we will have to rely on leak detector reporting this instead.
@@ -1876,14 +1572,9 @@ bool FPackageName::IsValidPath(FStringView InPath)
 		IsMemoryPackage(InPath);
 }
 
-TRefCountPtr<UE::PackageName::IMountPoint> FPackageName::RegisterMountPoint(const FString& RootPath, const FString& ContentPath)
-{
-	return FLongPackagePathsSingleton::Get().InsertMountPoint(RootPath, ContentPath);
-}
-
-void FPackageName::RegisterMountPoint(const TRefCountPtr<UE::PackageName::IMountPoint>& MountPoint)
+void FPackageName::RegisterMountPoint(const FString& RootPath, const FString& ContentPath)
 {
-	FLongPackagePathsSingleton::Get().InsertMountPoint(MountPoint);
+	FLongPackagePathsSingleton::Get().InsertMountPoint(RootPath, ContentPath);
 }
 
 void FPackageName::UnRegisterMountPoint(const FString& RootPath, const FString& ContentPath)
@@ -1891,11 +1582,6 @@ void FPackageName::UnRegisterMountPoint(const FString& RootPath, const FString&
 	FLongPackagePathsSingleton::Get().RemoveMountPoint(RootPath, ContentPath);
 }
 
-void FPackageName::UnRegisterMountPoint(const TRefCountPtr<UE::PackageName::IMountPoint>& MountPoint)
-{
-	FLongPackagePathsSingleton::Get().RemoveMountPoint(MountPoint);
-}
-
 bool FPackageName::MountPointExists(FStringView RootPath)
 {
 	return FLongPackagePathsSingleton::Get().MountPointExists(RootPath);
@@ -1911,36 +1597,6 @@ bool FPackageName::MountPointExists(const TCHAR* RootPath)
 	return FLongPackagePathsSingleton::Get().MountPointExists(FStringView(RootPath));
 }
 
-TRefCountPtr<UE::PackageName::IMountPoint>
-FPackageName::FindMountPointByRootPackageName(FStringView RootLongPackageName)
-{
-	return FLongPackagePathsSingleton::Get().FindMountPointByRootPackageName(RootLongPackageName);
-}
-
-TRefCountPtr<UE::PackageName::IMountPoint>
-FPackageName::FindMountPointByRootLocalPath(FStringView RootLocalPath)
-{
-	return FLongPackagePathsSingleton::Get().FindMountPointByRootLocalPath(RootLocalPath);
-}
-
-TRefCountPtr<UE::PackageName::IMountPoint>
-FPackageName::FindMountPoint(FStringView RootLongPackageName, FStringView RootLocalPath)
-{
-	return FLongPackagePathsSingleton::Get().FindMountPoint(RootLongPackageName, RootLocalPath);
-}
-
-TRefCountPtr<UE::PackageName::IMountPoint>
-FPackageName::FindOrAddMountPoint(FStringView RootLongPackageName, FStringView RootLocalPath)
-{
-	return FLongPackagePathsSingleton::Get().FindOrAddMountPoint(RootLongPackageName, RootLocalPath);
-}
-
-TRefCountPtr<UE::PackageName::IMountPoint>
-FPackageName::FindMountPointByChildLongPackageName(FStringView ChildLongPackageName)
-{
-	return FLongPackagePathsSingleton::Get().FindMountPointByChildLongPackageName(ChildLongPackageName);
-}
-
 FName FPackageName::GetPackageMountPoint(FStringView InPackagePath, bool InWithoutSlashes)
 {
 	using namespace UE::PackageName;
diff --git a/Engine/Source/Runtime/CoreUObject/Private/Misc/PackageName/MountPoint.cpp b/Engine/Source/Runtime/CoreUObject/Private/Misc/PackageName/MountPoint.cpp
--- a/Engine/Source/Runtime/CoreUObject/Private/Misc/PackageName/MountPoint.cpp
+++ b/Engine/Source/Runtime/CoreUObject/Private/Misc/PackageName/MountPoint.cpp
@@ -25,7 +25,7 @@ bool FMountPoint::IsMounted() const
 }
 
 void FMountPoint::SetMounted(bool bInMounted, UE::TWriteScopeLock<FTransactionallySafeRWLock>& /* ProofOfWriteLock */,
-	const TRefCountPtr<IMountPoint>& /* ProofOfCallerRefCount */, bool& bOutRemoveFromContainer)
+	const TRefCountPtr<FMountPoint>& /* ProofOfCallerRefCount */, bool& bOutRemoveFromContainer)
 {
 	bOutRemoveFromContainer = false;
 
diff --git a/Engine/Source/Runtime/CoreUObject/Private/Misc/PackageName/MountPointPrivate.h b/Engine/Source/Runtime/CoreUObject/Private/Misc/PackageName/MountPointPrivate.h
--- a/Engine/Source/Runtime/CoreUObject/Private/Misc/PackageName/MountPointPrivate.h
+++ b/Engine/Source/Runtime/CoreUObject/Private/Misc/PackageName/MountPointPrivate.h
@@ -55,9 +55,6 @@ public:
 	virtual void AddRef() const override;
 	virtual FReturnedRefCountValue Release() const override;
 	virtual FReturnedRefCountValue GetRefCount() const override;
-protected:
-	virtual FMountPoint* AsFMountPoint() override;
-public:
 
 	// FMountPoint additional API
 
@@ -80,7 +77,7 @@ public:
 	 * be set to true and we will skip the effort in TryUnregisterAndDelete of removing it from the container.
 	 */
 	void SetMounted(bool bInMounted, UE::TWriteScopeLock<FTransactionallySafeRWLock>& ProofOfWriteLock,
-		const TRefCountPtr<IMountPoint>& ProofOfCallerRefCount, bool& bOutRemoveFromContainer);
+		const TRefCountPtr<FMountPoint>& ProofOfCallerRefCount, bool& bOutRemoveFromContainer);
 
 	/**
 	 * Get/Set the Ephemeral ordering value. MountPoints that have the same RootPath or ContentPath as other
@@ -97,7 +94,6 @@ public:
 	static uint32 GetTypeHashKey(const FString& LongPackageNameOfKey, const FString& LocalPathStandardOfKey);
 	static bool KeysAreEqual(const FString& LongPackageNameA, const FString& LocalPathStandardA,
 		const FString& LongPackageNameB, const FString& LocalPathStandardB);
-	static bool IsHigherOverridePriority(FMountPoint* A, FMountPoint* B);
 
 private:
 	/**
@@ -186,11 +182,6 @@ inline FStringView FMountPoint::GetLocalPathAbsolute() const
 	return LocalPathAbsolute;
 }
 
-inline FMountPoint* FMountPoint::AsFMountPoint()
-{
-	return this;
-}
-
 inline bool FMountPoint::IsRemovable() const
 {
 	return bRemovable;
@@ -304,24 +295,4 @@ inline bool FMountPoint::KeysAreEqual(const FString& LongPackageNameA, const FSt
 		&& LocalPathStandardA.Equals(LocalPathStandardB, ESearchCase::IgnoreCase);
 }
 
-inline bool FMountPoint::IsHigherOverridePriority(FMountPoint* A, FMountPoint* B)
-{
-	// Must be called within FLongPackagePathsSingleton->MountLock so that GetMountOrder is readable.
-	if ((A != nullptr) != (B != nullptr))
-	{
-		return A != nullptr;
-	}
-	if (A == nullptr)
-	{
-		return false; // equal
-	}
-	CA_ASSUME(B != nullptr);
-
-	if (A->IsMounted() != B->IsMounted())
-	{
-		return A->IsMounted();
-	}
-	return A->GetMountOrder() > B->GetMountOrder();
-}
-
 } // namespace UE::PackageName
\ No newline at end of file
diff --git a/Engine/Source/Runtime/CoreUObject/Private/Misc/PackageName/PackageNamePrivate.h b/Engine/Source/Runtime/CoreUObject/Private/Misc/PackageName/PackageNamePrivate.h
--- a/Engine/Source/Runtime/CoreUObject/Private/Misc/PackageName/PackageNamePrivate.h
+++ b/Engine/Source/Runtime/CoreUObject/Private/Misc/PackageName/PackageNamePrivate.h
@@ -64,27 +64,17 @@ struct FLongPackagePathsSingleton
 	 * LocalPathStandard is a relative path and is consistent with FileManager relative paths.
 	 * @return false if the input strings are not recognizable as mountpoint strings, else true.
 	 */
-	bool TryNormalizePaths(FStringView InLongPackageName, FStringView InLocalPath,
+	bool TryNormalizePaths(const FString& InLongPackageName, const FString& InLocalPath,
 		FString& OutLongPackageName, FString& OutLocalPathStandard, FString& OutLocalPathAbsolute, bool bSilenceLogs);
-	bool TryNormalizeLongPackageNamePath(FStringView InLongPackageName, FString& OutLongPackageName,
+	bool TryNormalizeLongPackageNamePath(const FString& InLongPackageName, FString& OutLongPackageName,
 		bool bSilenceLogs);
-	bool TryNormalizeLocalPath(FStringView InLocalPath, FString& OutLocalPathStandard,
+	bool TryNormalizeLocalPath(const FString& InLocalPath, FString& OutLocalPathStandard,
 		FString& OutLocalPathAbsolute, bool bSilenceLogs);
 
-	TRefCountPtr<UE::PackageName::IMountPoint> InsertMountPoint(const FString& RootPath, const FString& ContentPath);
-	void InsertMountPoint(const TRefCountPtr<UE::PackageName::IMountPoint>& InMountPoint);
+	void InsertMountPoint(const FString& RootPath, const FString& ContentPath);
 
 	/** This will remove a previously inserted mount point. */
 	void RemoveMountPoint(const FString& RootPath, const FString& ContentPath);
-	void RemoveMountPoint(const TRefCountPtr<UE::PackageName::IMountPoint>& InMountPoint);
-
-	TRefCountPtr<UE::PackageName::IMountPoint> FindMountPointByRootPackageName(FStringView RootLongPackageName);
-	TRefCountPtr<UE::PackageName::IMountPoint> FindMountPointByRootLocalPath(FStringView RootLocalPath);
-	TRefCountPtr<UE::PackageName::IMountPoint> FindMountPoint(FStringView RootLongPackageName,
-		FStringView RootLocalPath);
-	TRefCountPtr<UE::PackageName::IMountPoint> FindOrAddMountPoint(FStringView RootLongPackageName,
-		FStringView RootLocalPath);
-	TRefCountPtr<UE::PackageName::IMountPoint> FindMountPointByChildLongPackageName(FStringView ChildLongPackageName);
 
 	/** Checks whether the specific root path is a valid mount point. */
 	bool MountPointExists(FStringView RootPath);
@@ -99,12 +89,6 @@ private:
 	FLongPackagePathsSingleton();
 	~FLongPackagePathsSingleton();
 	void ReassignMountOrders();
-	void BroadcastMountPointMounted(const TRefCountPtr<UE::PackageName::IMountPoint>& PublicMountPoint,
-		UE::PackageName::FMountPoint* MountPoint);
-	bool TryRemoveMountPointInternal(const TRefCountPtr<UE::PackageName::IMountPoint>& InternallyHeldMountPoint,
-		UE::TWriteScopeLock<FTransactionallySafeRWLock>& ProofOfWriteLock, UE::PackageName::FMountPoint* MountPoint);
-	void BroadcastMountPointDismounted(const TRefCountPtr<UE::PackageName::IMountPoint>& PublicMountPoint,
-		UE::PackageName::FMountPoint* MountPoint);
 
 #if !UE_BUILD_SHIPPING
 	void ExecDumpMountPoints(const TArray<FString>& Args);
diff --git a/Engine/Source/Runtime/CoreUObject/Public/Misc/MountPoint.h b/Engine/Source/Runtime/CoreUObject/Public/Misc/MountPoint.h
--- a/Engine/Source/Runtime/CoreUObject/Public/Misc/MountPoint.h
+++ b/Engine/Source/Runtime/CoreUObject/Public/Misc/MountPoint.h
@@ -5,9 +5,6 @@
 #include "Containers/StringView.h"
 #include "Templates/RefCounting.h"
 
-namespace UE::PackageName { class FMountPoint; }
-struct FLongPackagePathsSingleton;
-
 namespace UE::PackageName
 {
 
@@ -80,15 +77,6 @@ public:
 	 * Mutable but atomic, can be read from any thread.
 	 */
 	virtual bool IsMounted() const = 0;
-
-protected:
-	/**
-	 * A downcast function to the implementation class inside FPackageName. There should be no other implementation of
-	 * IMountPoint; we log an error if this fails.
-	 */
-	virtual FMountPoint* AsFMountPoint() = 0;
-
-	friend FLongPackagePathsSingleton;
 };
 
 } // namespace UE::PackageName
\ No newline at end of file
diff --git a/Engine/Source/Runtime/CoreUObject/Public/Misc/PackageName.h b/Engine/Source/Runtime/CoreUObject/Public/Misc/PackageName.h
--- a/Engine/Source/Runtime/CoreUObject/Public/Misc/PackageName.h
+++ b/Engine/Source/Runtime/CoreUObject/Public/Misc/PackageName.h
@@ -314,113 +314,30 @@ public:
 	static COREUOBJECT_API bool TryConvertGameRelativePackagePathToLocalPath(FStringView RelativePackagePath, FString& OutLocalPath);
 
 	/**
-	 * This will mount a mount point if not already mounted, and will update its order to the front of the override
-	 * chain (so it can overlap an existing mount point and win).
-	 * If you register a mount point (even if you do so only in certain circumstances) consider also adding a handler
-	 * for GetExplanationForUnavailablePackage to help debug cases where a package can't be found.
-	 * This function can fail only if the input arguments are not valid text for RootPath or ContentPath.
+	 * This will insert a mount point at the head of the search chain (so it can overlap an existing mount point and win).
+	 * If you register a mount point (even if you do so only in certain circumstances) consider also adding a handler for GetExplanationForUnavailablePackage
+	 * to help debug cases where a package can't be found
 	 *
-	 * @param RootPath LongPackageName location for the root of the MountPoint,
-	 *        Examples: /Engine/, /PluginName/
-	 * @param ContentPath LocalPath location for the root of the MountPoint,
-	 *        Examples: ../../../Engine/Content, d:\FullPathToUnreal\Engine\Plugins\PluginName
-	 * 
-	 * @return RefCountPtr to the mounted IMountPoint, or nullptr if failed.
-	 */
-	static COREUOBJECT_API TRefCountPtr<UE::PackageName::IMountPoint>
-		RegisterMountPoint(const FString& RootPath, const FString& ContentPath);
-	/**
-	 * Mounts a mount point if not already mounted, and updates its order to front of the override chain.
-	 * If you register a mount point (even if you do so only in certain circumstances) consider also adding a handler
-	 * for GetExplanationForUnavailablePackage to help debug cases where a package can't be found.
+	 * @param RootPath Logical Root Path.
+	 * @param ContentPath Content Path on disk.
 	 */
-	static COREUOBJECT_API void RegisterMountPoint(const TRefCountPtr<UE::PackageName::IMountPoint>& MountPoint);
+	static COREUOBJECT_API void RegisterMountPoint(const FString& RootPath, const FString& ContentPath);
 
 	/**
-	 * This will remove a previously inserted mount point. Logs a message if not found or not removable.
+	 * This will remove a previously inserted mount point.
 	 *
-	 * @param RootPath LongPackageName location for the root of the MountPoint,
-	 *        Examples: /Engine/, /PluginName/
-	 * @param ContentPath LocalPath location for the root of the MountPoint,
-	 *        Examples: ../../../Engine/Content, d:\FullPathToUnreal\Engine\Plugins\PluginName
+	 * @param RootPath Logical Root Path.
+	 * @param ContentPath Content Path on disk.
 	 */
 	static COREUOBJECT_API void UnRegisterMountPoint(const FString& RootPath, const FString& ContentPath);
-	/** Unmounts a MountPoint; the mountpoint remains findable and updatable until all references are released. */
-	static COREUOBJECT_API void UnRegisterMountPoint(const TRefCountPtr<UE::PackageName::IMountPoint>& MountPoint);
-
 
 	/**
-	 * Returns whether the given RootPath is the LongPackageName of a MountPoint that has been mounted, by default or
-	 * by an earlier call to RegisterMountPoint. 
+	 * Returns whether the specific logical root path is a valid mount point.
 	 */
 	static COREUOBJECT_API bool MountPointExists(FStringView RootPath);
 	static COREUOBJECT_API bool MountPointExists(const FString& RootPath);
 	static COREUOBJECT_API bool MountPointExists(const TCHAR* RootPath);
 
-	/**
-	 * Looks for the first allocated MountPoint with the given RootLongPackageName. If multiple are allocated, only
-	 * returns the one with the highest override priority (the one that was added last). A discovered MountPoint might
-	 * or might not be mounted, call MountPoint->IsMounted() to check for mounted.
-	 * 
-	 * @param RootLongPackageName LongPackageName location for the root of the MountPoint,
-	 *        Examples: /Engine/, /PluginName/
-	 * @return RefCountPtr to the IMountPoint, or nullptr if not yet allocated.
-	 */
-	static COREUOBJECT_API TRefCountPtr<UE::PackageName::IMountPoint>
-		FindMountPointByRootPackageName(FStringView RootLongPackageName);
-	/**
-	 * Looks for the first MountPoint with the given RootLocalPath. If multiple are allocated, only returns the
-	 * one with the highest override priority (the one that was added last). A discovered MountPoint might or might
-	 * not be mounted, call MountPoint->IsMounted() to check for mounted.
-	 *
-	 * @param RootLocalPath LocalPath location for the root of the MountPoint,
-	 *        Examples: ../../../Engine/Content, d:\FullPathToUnreal\Engine\Plugins\PluginName
-	 * @return RefCountPtr to the IMountPoint, or nullptr if not yet allocated.
-	 */
-	static COREUOBJECT_API TRefCountPtr<UE::PackageName::IMountPoint>
-		FindMountPointByRootLocalPath(FStringView RootLocalPath);
-	/**
-	 * Looks for an existing allocated MountPoint with the given RootLongPackageName and RootLocalPath. These fields
-	 * uniquely identify a MountPoint so at most one can exist.
-	 * If the MountPoint exists, it might or might not be mounted, call MountPoint->IsMounted() to check for mounted.
-	 *
-	 * @param RootLongPackageName LongPackageName location for the root of the MountPoint,
-	 *        Examples: /Engine/, /PluginName/
-	 * @param RootLocalPath LocalPath location for the root of the MountPoint,
-	 *        Examples: ../../../Engine/Content, d:\FullPathToUnreal\Engine\Plugins\PluginName
-	 * @return RefCountPtr to the IMountPoint, or nullptr if not yet allocated.
-	 */
-	static COREUOBJECT_API TRefCountPtr<UE::PackageName::IMountPoint>
-		FindMountPoint(FStringView RootLongPackageName, FStringView RootLocalPath);
-	/**
-	 * Looks for an existing allocated MountPoint with the given RootLongPackageName and RootLocalPath, adds it
-	 * if it does not already exist. These fields uniquely identify a MountPoint so at most one can exist.
-	 * When the MountPoint is returned, it might or might not be mounted, call MountPoint->IsMounted() to check for
-	 * mounted.
-	 * This function can only return null if the input arguments are not valid text for RootLongPackageName or
-	 * RootLocalPath.
-	 *
-	 * @param RootLongPackageName LongPackageName location for the root of the MountPoint,
-	 *        Examples: /Engine/, /PluginName/
-	 * @param RootLocalPath LocalPath location for the root of the MountPoint,
-	 *        Examples: ../../../Engine/Content, d:\FullPathToUnreal\Engine\Plugins\PluginName
-	 * @return RefCountPtr to the IMountPoint, or nullptr if failed.
-	 */
-	static COREUOBJECT_API TRefCountPtr<UE::PackageName::IMountPoint>
-		FindOrAddMountPoint(FStringView RootLongPackageName, FStringView RootLocalPath);
-
-	/**
-	 * Looks for the first MountPoint that is equal to or is a parent path of the given ChildLongPackageName.
-	 * If multiple are allocated, only returns the one with the highest override priority (the one that was added
-	 * last). A discovered MountPoint might or might not be mounted, call MountPoint->IsMounted() to check for mounted.
-	 *
-	 * @param ChildLongPackageName LongPackageName for a path under the MountPoint,
-	 *        Examples: /Engine/, /Engine/Asset, /PluginName/, /PluginName/Asset
-	 * @return RefCountPtr to the IMountPoint, or nullptr if not yet allocated.
-	 */
-	static COREUOBJECT_API TRefCountPtr<UE::PackageName::IMountPoint>
-		FindMountPointByChildLongPackageName(FStringView ChildLongPackageName);
-
 	/**
 	 * Get the mount point for a given package path
 	 * 
@@ -734,25 +651,14 @@ public:
 	static COREUOBJECT_API void IteratePackagesInDirectory(const FString& RootDir, const FPackageNameVisitor& Visitor);
 	static COREUOBJECT_API void IteratePackagesInDirectory(const FString& RootDir, const FPackageNameStatVisitor& Visitor);
 
-	/** Event triggered when a MountPoint is mounted. */
-	DECLARE_MULTICAST_DELEGATE_OneParam(FMountPointEvent, const TRefCountPtr<UE::PackageName::IMountPoint>& MountPoint)
-	static FMountPointEvent& OnMountPointMounted()
-	{
-		return OnMountPointMountedEvent;
-	}
-	/** Legacy event triggered when a MountPoint is mounted; passes strings instead of IMountPoint. */
+	/** Event that is triggered when a new content path is mounted */
 	DECLARE_MULTICAST_DELEGATE_TwoParams( FOnContentPathMountedEvent, const FString& /* Asset path */, const FString& /* ContentPath */ );
 	static FOnContentPathMountedEvent& OnContentPathMounted()
 	{
 		return OnContentPathMountedEvent;
 	}
 
-	/** Event triggered when a MountPoint is dismounted. */
-	static FMountPointEvent& OnMountPointDismounted()
-	{
-		return OnMountPointDismountedEvent;
-	}
-	/** Legacy event triggered when a MountPoint is dismounted; passes strings instead of IMountPoint. */
+	/** Event that is triggered when a new content path is removed */
 	DECLARE_MULTICAST_DELEGATE_TwoParams(FOnContentPathDismountedEvent, const FString& /* Asset path */, const FString& /* ContentPath */ );
 	static FOnContentPathDismountedEvent& OnContentPathDismounted()
 	{
@@ -1106,11 +1012,9 @@ private:
 	static COREUOBJECT_API bool TryConvertToMountedPackagePath(const FString& InPath, FPackagePath& OutPackagePath, EErrorCode& OutFailureReason);
 
 	/** Event that is triggered when a new content path is mounted */
-	static COREUOBJECT_API FMountPointEvent OnMountPointMountedEvent;
 	static COREUOBJECT_API FOnContentPathMountedEvent OnContentPathMountedEvent;
 
 	/** Event that is triggered when a new content path is removed */
-	static COREUOBJECT_API FMountPointEvent OnMountPointDismountedEvent;
 	static COREUOBJECT_API FOnContentPathDismountedEvent OnContentPathDismountedEvent;
 
 	/** Delegate used to check whether a package exist without using the filesystem. */
