[StructUtils] Skip struct validation when the struct is duplicated. The struct is not linked. The validation will always fail due to this.

#Synced-CL 50535091

diff --git a/Engine/Source/Runtime/CoreUObject/Private/StructUtils/UserDefinedStruct.cpp b/Engine/Source/Runtime/CoreUObject/Private/StructUtils/UserDefinedStruct.cpp
--- a/Engine/Source/Runtime/CoreUObject/Private/StructUtils/UserDefinedStruct.cpp
+++ b/Engine/Source/Runtime/CoreUObject/Private/StructUtils/UserDefinedStruct.cpp
@@ -156,11 +156,15 @@ void UUserDefinedStruct::Serialize(FStructuredArchive::FRecord Record)
 				}
 			}
 
-			const FUserDefinedStructEditorUtils::EStructureError Result = FUserDefinedStructEditorUtils::IsStructureValid(this, NULL, &ErrorMessage);
-			if (FUserDefinedStructEditorUtils::EStructureError::Ok != Result)
+			// While duplicating, the struct is not linked and is therefore invalid.
+			if (!(UnderlyingArchive.GetPortFlags() & PPF_Duplicate))
 			{
-				Status = EUserDefinedStructureStatus::UDSS_Error;
-				UE_LOG(LogClass, Log, TEXT("UUserDefinedStruct.Serialize '%s' validation: %s"), *GetName(), *ErrorMessage);
+				const FUserDefinedStructEditorUtils::EStructureError Result = FUserDefinedStructEditorUtils::IsStructureValid(this, NULL, &ErrorMessage);
+				if (FUserDefinedStructEditorUtils::EStructureError::Ok != Result)
+				{
+					Status = EUserDefinedStructureStatus::UDSS_Error;
+					UE_LOG(LogClass, Log, TEXT("UUserDefinedStruct.Serialize '%s' validation: %s"), *GetName(), *ErrorMessage);
+				}
 			}
 		}
 	}
