CoreUObject: Move out the composite enumerators in EStructFlags to be defines just like for the other core enums.
Add the comment found on core enums in ObjectMacrors.h to EStructFlags to clarify that composite flags should not be made a enumerators within the enums.

#rnx

#Synced-CL 50480306

diff --git a/Engine/Source/Runtime/CoreUObject/Public/UObject/Class.h b/Engine/Source/Runtime/CoreUObject/Public/UObject/Class.h
--- a/Engine/Source/Runtime/CoreUObject/Public/UObject/Class.h
+++ b/Engine/Source/Runtime/CoreUObject/Public/UObject/Class.h
@@ -1126,13 +1126,16 @@ enum EStructFlags
 	/** If set, this struct will have PostLoad on it in the editor. */
 	STRUCT_PostLoad UE_INTERNAL = 0x08000000,
 	
-	/** Struct flags that are automatically inherited */
-	STRUCT_Inherit				= STRUCT_HasInstancedReference|STRUCT_Atomic,
-
-	/** Flags that are always computed, never loaded or done with code generation */
-	STRUCT_ComputedFlags		= STRUCT_NetDeltaSerializeNative | STRUCT_NetSerializeNative | STRUCT_SerializeNative | STRUCT_PostSerializeNative | STRUCT_CopyNative | STRUCT_IsPlainOldData | STRUCT_NoDestructor | STRUCT_ZeroConstructor | STRUCT_IdenticalNative | STRUCT_AddStructReferencedObjects | STRUCT_ExportTextItemNative | STRUCT_ImportTextItemNative | STRUCT_SerializeFromMismatchedTag | STRUCT_PostScriptConstruct | STRUCT_NetSharedSerialization
+	// !!!!!!!!!!!!!!!!!!!!!!!!!!!!!!!!!!!!!!!!!!!!!!!!!!!!!!!!!!!!!!!!!!!!!!!!!!!!!!!!!!!!!!!!!!!!!!!!!!!!!!!!!!!!!!!!!!!!!!!!!!!!!
+	// DO NOT ADD composite flags to EStructFlags. Debugger visualizations have trouble displaying composite flag values as text.
+	// !!!!!!!!!!!!!!!!!!!!!!!!!!!!!!!!!!!!!!!!!!!!!!!!!!!!!!!!!!!!!!!!!!!!!!!!!!!!!!!!!!!!!!!!!!!!!!!!!!!!!!!!!!!!!!!!!!!!!!!!!!!!!
 };
 
+/** Struct flags that are automatically inherited */
+#define STRUCT_Inherit			(STRUCT_HasInstancedReference | STRUCT_Atomic)
+
+/** Flags that are always computed, never loaded or done with code generation */
+#define STRUCT_ComputedFlags	(STRUCT_NetDeltaSerializeNative | STRUCT_NetSerializeNative | STRUCT_SerializeNative | STRUCT_PostSerializeNative | STRUCT_CopyNative | STRUCT_IsPlainOldData | STRUCT_NoDestructor | STRUCT_ZeroConstructor | STRUCT_IdenticalNative | STRUCT_AddStructReferencedObjects | STRUCT_ExportTextItemNative | STRUCT_ImportTextItemNative | STRUCT_SerializeFromMismatchedTag | STRUCT_PostScriptConstruct | STRUCT_NetSharedSerialization)
 
 
 #if CHECK_PUREVIRTUALS
