Suppress UBSan issues with casting function pointer types in source.

#Synced-CL 50647680

diff --git a/Engine/Source/Runtime/Core/Public/Clang/ClangPlatformCodeAnalysis.h b/Engine/Source/Runtime/Core/Public/Clang/ClangPlatformCodeAnalysis.h
--- a/Engine/Source/Runtime/Core/Public/Clang/ClangPlatformCodeAnalysis.h
+++ b/Engine/Source/Runtime/Core/Public/Clang/ClangPlatformCodeAnalysis.h
@@ -119,6 +119,10 @@
 	#define USING_UNDEFINED_BEHAVIOR_SANITISER 0
 #endif
 
+#if USING_UNDEFINED_BEHAVIOR_SANITISER
+	#define UE_DISABLE_UBSAN [[clang::no_sanitize("undefined")]]
+#endif
+
 #ifndef USING_INSTRUMENTATION
 	#define USING_INSTRUMENTATION 0
 #endif
diff --git a/Engine/Source/Runtime/Core/Public/HAL/Platform.h b/Engine/Source/Runtime/Core/Public/HAL/Platform.h
--- a/Engine/Source/Runtime/Core/Public/HAL/Platform.h
+++ b/Engine/Source/Runtime/Core/Public/HAL/Platform.h
@@ -955,6 +955,11 @@
 	#define PLATFORM_CACHE_LINE_SIZE	64
 #endif
 
+// Disable UBSan on a given function
+#ifndef UE_DISABLE_UBSAN
+	#define UE_DISABLE_UBSAN
+#endif
+
 // Compile-time warnings and errors. Use these as "#pragma COMPILE_WARNING("XYZ")". GCC does not expand macro parameters to _Pragma, so we can't wrap the #pragma part.
 #ifdef _MSC_VER
 	#define MSC_FORMAT_DIAGNOSTIC_HELPER_2(x) #x
diff --git a/Engine/Source/Runtime/CoreUObject/Private/UObject/UObjectConstructInternal.h b/Engine/Source/Runtime/CoreUObject/Private/UObject/UObjectConstructInternal.h
--- a/Engine/Source/Runtime/CoreUObject/Private/UObject/UObjectConstructInternal.h
+++ b/Engine/Source/Runtime/CoreUObject/Private/UObject/UObjectConstructInternal.h
@@ -5,6 +5,7 @@
 #if !defined(UE_WITH_CONSTINIT_UOBJECT) || !UE_WITH_CONSTINIT_UOBJECT
 
 #include "CoreMinimal.h"
+#include "HAL/Platform.h"
 #include "UObject/Class.h"
 #include "UObject/UObjectAllocator.h"
 #include "UObject/UObjectGlobals.h"
@@ -130,6 +131,9 @@ void ConstructUClassNoInitHelper(
 }
 
 template<typename UEnumClass, typename EnumParams, typename PostNewFn>
+// It is UB to call `OuterFunc` below because it can return a `UClass*` not `UObject*` (even though
+// there is a parent-child relation there). So we suppress UBSan for this function.
+UE_DISABLE_UBSAN
 void ConstructUEnumHelper(UEnum*& OutEnum, const EnumParams& Params, PostNewFn&& InPostNewFn)
 {
 	UObject* (*OuterFunc)() = Params.OuterFunc;
@@ -167,6 +171,9 @@ void ConstructUEnumHelper(UEnum*& OutEnum, const EnumParams& Params, PostNewFn&&
 }
 
 template<typename UScriptStructClass, typename StructParams, typename PostNewFn>
+// It is UB to call `OuterFunc` below because it can return a `UClass*` not `UObject*` (even though
+// there is a parent-child relation there). So we suppress UBSan for this function.
+UE_DISABLE_UBSAN
 void ConstructUScriptStructHelper(UScriptStruct*& OutStruct, const StructParams& Params, PostNewFn&& InPostNewFn)
 {
 	UObject* (*OuterFunc)() = Params.OuterFunc;
@@ -194,6 +201,9 @@ void ConstructUScriptStructHelper(UScriptStruct*& OutStruct, const StructParams&
 }
 
 template<typename UClassClass, typename ClassParams, typename PostNewFn>
+// It is UB to call `SingletonFunc` below because it can return a `UClass*` not `UObject*` (even though
+// there is a parent-child relation there). So we suppress UBSan for this function.
+UE_DISABLE_UBSAN
 void ConstructUClassHelper(UClass*& OutClass, const ClassParams& Params, PostNewFn&& InPostNewFn)
 {
 	if (OutClass && (OutClass->ClassFlags & CLASS_Constructed))
@@ -265,6 +275,9 @@ void ConstructUClassHelper(UClass*& OutClass, const ClassParams& Params, PostNew
 }
 
 template<typename UFunctionClass, typename FunctionParams, typename NewFn>
+// It is UB to call `OuterFunc` below because it can return a `UClass*` not `UObject*` (even though
+// there is a parent-child relation there). So we suppress UBSan for this function.
+UE_DISABLE_UBSAN
 void ConstructUFunctionHelper(UFunction*& InOutFunction, const FunctionParams& Params, UFunction** InSingletonPtr, NewFn&& InNewFn)
 {
 	UObject* (*OuterFunc)() = Params.OuterFunc;
diff --git a/Engine/Source/Runtime/CoreUObject/Private/UObject/UObjectGlobals.cpp b/Engine/Source/Runtime/CoreUObject/Private/UObject/UObjectGlobals.cpp
--- a/Engine/Source/Runtime/CoreUObject/Private/UObject/UObjectGlobals.cpp
+++ b/Engine/Source/Runtime/CoreUObject/Private/UObject/UObjectGlobals.cpp
@@ -6573,6 +6573,9 @@ namespace UECodeGen_Private
 		ConstructUClassHelper<UClass>(OutClass, Params, [](UClass*, const FClassParams&) {});
 	}
 
+	// It is UB to call `SingletonFunc` below because it can return a `UClass*` not `UObject*` (even though
+	// there is a parent-child relation there). So we suppress UBSan for this function.
+	UE_DISABLE_UBSAN
 	void ConstructUPackage(UPackage*& OutPackage, const FPackageParams& Params)
 	{
 		if (OutPackage)
