[CookedMetaData/ALT]
- Make sure cooked metadata are ready ASAP during loading to avoid racing between ALT and postload

#Synced-CL 50642529

diff --git a/Engine/Source/Runtime/CoreUObject/Private/Serialization/AsyncLoading2.cpp b/Engine/Source/Runtime/CoreUObject/Private/Serialization/AsyncLoading2.cpp
--- a/Engine/Source/Runtime/CoreUObject/Private/Serialization/AsyncLoading2.cpp
+++ b/Engine/Source/Runtime/CoreUObject/Private/Serialization/AsyncLoading2.cpp
@@ -70,6 +70,7 @@
 #include "Serialization/LargeMemoryReader.h"
 #include "Serialization/MemoryReader.h"
 #include "UObject/UObjectClusters.h"
+#include "UObject/CookedMetaData.h"
 #include "UObject/LinkerInstancingContext.h"
 #include "UObject/LinkerLoadImportBehavior.h"
 #include "UObject/ObjectHandlePrivate.h"
@@ -8466,13 +8467,14 @@ void FAsyncPackage2::InitializeExportArchive(FExportArchive& Ar, bool bIsOptiona
 }
 
 // When possible, deferring serialization will enable some nodes to run in parallel
-// As some point, we will also need to defer them if we want to better mix cooked and uncooked loading.
+// At some point, we will also need to defer them if we want to better mix cooked and uncooked loading.
 inline bool CanDeferSerialization(FExportObject& Export)
 {
 	return GAllowMultithreadedLoading &&
 		// Classes, CDOs and Archetype are serialized during the create export phase since they are required to be loaded by dependencies
 		!Export.Object->GetClass()->IsChildOf(UField::StaticClass()) &&
-		!Export.Object->HasAnyFlags(RF_ClassDefaultObject | RF_ArchetypeObject);
+		!Export.Object->HasAnyFlags(RF_ClassDefaultObject | RF_ArchetypeObject) &&
+		!Export.Object->GetClass()->IsChildOf(UCookedMetaDataBase::StaticClass());
 }
 
 bool FAsyncPackage2::CreateRuntimeExports(FAsyncLoadingThreadState2& ThreadState)
diff --git a/Engine/Source/Runtime/CoreUObject/Private/UObject/CookedMetaData.cpp b/Engine/Source/Runtime/CoreUObject/Private/UObject/CookedMetaData.cpp
--- a/Engine/Source/Runtime/CoreUObject/Private/UObject/CookedMetaData.cpp
+++ b/Engine/Source/Runtime/CoreUObject/Private/UObject/CookedMetaData.cpp
@@ -25,14 +25,14 @@ void PrepareCookedMetaDataForPurge(UObject* CookedMetaDataPtr)
 }
 
 template <typename CookedMetaDataOuterType, typename CookedMetaDataType>
-void PostLoadCookedMetaData(CookedMetaDataType* CookedMetaDataPtr)
+void PreloadCookedMetaData(CookedMetaDataType* CookedMetaDataPtr)
 {
 #if WITH_METADATA
 	checkf(CookedMetaDataPtr->GetPackage()->HasAnyPackageFlags(PKG_Cooked), TEXT("Cooked meta-data should only be loaded for a cooked package!"));
 
 	if (CookedMetaDataOuterType* Owner = CastChecked<CookedMetaDataOuterType>(CookedMetaDataPtr->GetOuter()))
 	{
-		Owner->ConditionalPostLoad();
+		Owner->ConditionalPreload();
 		CookedMetaDataPtr->ApplyMetaData(Owner);
 		PrepareCookedMetaDataForPurge(CookedMetaDataPtr);
 	}
@@ -208,11 +208,14 @@ void FStructCookedMetaDataStore::ApplyMetaData(UStruct* TargetStruct) const
 	}
 }
 
-
-void UEnumCookedMetaData::PostLoad()
+void UEnumCookedMetaData::Serialize(FArchive& Ar)
 {
-	Super::PostLoad();
-	CookedMetaDataUtil::Internal::PostLoadCookedMetaData<UEnum, UEnumCookedMetaData>(this);
+	Super::Serialize(Ar);
+
+	if (Ar.IsLoading())
+	{
+		CookedMetaDataUtil::Internal::PreloadCookedMetaData<UEnum, UEnumCookedMetaData>(this);
+	}
 }
 
 bool UEnumCookedMetaData::HasMetaData() const
@@ -230,12 +233,14 @@ void UEnumCookedMetaData::ApplyMetaData(UEnum* TargetEnum) const
 	EnumMetaData.ApplyMetaData(TargetEnum);
 }
 
-
-
-void UStructCookedMetaData::PostLoad()
+void UStructCookedMetaData::Serialize(FArchive& Ar)
 {
-	Super::PostLoad();
-	CookedMetaDataUtil::Internal::PostLoadCookedMetaData<UScriptStruct, UStructCookedMetaData>(this);
+	Super::Serialize(Ar);
+
+	if (Ar.IsLoading())
+	{
+		CookedMetaDataUtil::Internal::PreloadCookedMetaData<UScriptStruct, UStructCookedMetaData>(this);
+	}
 }
 
 bool UStructCookedMetaData::HasMetaData() const
@@ -254,10 +259,14 @@ void UStructCookedMetaData::ApplyMetaData(UScriptStruct* TargetStruct) const
 }
 
 
-void UClassCookedMetaData::PostLoad()
+void UClassCookedMetaData::Serialize(FArchive& Ar)
 {
-	Super::PostLoad();
-	CookedMetaDataUtil::Internal::PostLoadCookedMetaData<UClass, UClassCookedMetaData>(this);
+	Super::Serialize(Ar);
+
+	if (Ar.IsLoading())
+	{
+		CookedMetaDataUtil::Internal::PreloadCookedMetaData<UClass, UClassCookedMetaData>(this);
+	}
 }
 
 bool UClassCookedMetaData::HasMetaData() const
diff --git a/Engine/Source/Runtime/CoreUObject/Public/UObject/CookedMetaData.h b/Engine/Source/Runtime/CoreUObject/Public/UObject/CookedMetaData.h
--- a/Engine/Source/Runtime/CoreUObject/Public/UObject/CookedMetaData.h
+++ b/Engine/Source/Runtime/CoreUObject/Public/UObject/CookedMetaData.h
@@ -117,16 +117,26 @@ private:
 	TMap<FName, FFieldCookedMetaDataStore> PropertiesMetaData;
 };
 
+/**
+ * Cooked meta-data base class.
+ */
+UCLASS(Abstract)
+class UCookedMetaDataBase : public UObject
+{
+public:
+	GENERATED_BODY()
+};
+
 /**
  * Cooked meta-data for a UEnum.
  */
 UCLASS(Optional, Within=Enum, MinimalAPI)
-class UEnumCookedMetaData : public UObject
+class UEnumCookedMetaData : public UCookedMetaDataBase
 {
 public:
 	GENERATED_BODY()
 
-	COREUOBJECT_API virtual void PostLoad() override;
+	COREUOBJECT_API virtual void Serialize(FArchive& Ar) override;
 	
 	COREUOBJECT_API virtual bool HasMetaData() const;
 	COREUOBJECT_API virtual void CacheMetaData(const UEnum* SourceEnum);
@@ -141,12 +151,12 @@ protected:
  * Cooked meta-data for a UScriptStruct, including its nested FProperty data.
  */
 UCLASS(Optional, Within=ScriptStruct, MinimalAPI)
-class UStructCookedMetaData : public UObject
+class UStructCookedMetaData : public UCookedMetaDataBase
 {
 public:
 	GENERATED_BODY()
 
-	COREUOBJECT_API virtual void PostLoad() override;
+	COREUOBJECT_API virtual void Serialize(FArchive& Ar) override;
 	
 	COREUOBJECT_API virtual bool HasMetaData() const;
 	COREUOBJECT_API virtual void CacheMetaData(const UScriptStruct* SourceStruct);
@@ -161,12 +171,12 @@ protected:
  * Cooked meta-data for a UClass, including its nested FProperty and UFunction data.
  */
 UCLASS(Optional, Within=Class, MinimalAPI)
-class UClassCookedMetaData : public UObject
+class UClassCookedMetaData : public UCookedMetaDataBase
 {
 public:
 	GENERATED_BODY()
 
-	COREUOBJECT_API virtual void PostLoad() override;
+	COREUOBJECT_API virtual void Serialize(FArchive& Ar) override;
 	
 	COREUOBJECT_API virtual bool HasMetaData() const;
 	COREUOBJECT_API virtual void CacheMetaData(const UClass* SourceClass);
