Fix crash when loading a placeholder native object in the editor
Refactor PropertyBagRepository to make the simple accessors return false in cooked builds

#Synced-CL 50604117

diff --git a/Engine/Source/Runtime/CoreUObject/Internal/UObject/PropertyBagRepository.h b/Engine/Source/Runtime/CoreUObject/Internal/UObject/PropertyBagRepository.h
--- a/Engine/Source/Runtime/CoreUObject/Internal/UObject/PropertyBagRepository.h
+++ b/Engine/Source/Runtime/CoreUObject/Internal/UObject/PropertyBagRepository.h
@@ -17,21 +17,32 @@ namespace UE
 {
 	class FPropertyPathName;
 
-#if WITH_EDITORONLY_DATA
-
-#if STATS
-struct FPropertyBagRepositoryStats
-{
-	size_t IDOMemoryBytes;
-	int32 NumIDOs;
-	int32 NumIDOsWithLooseProperties;
-	int32 NumPlaceholderTypes;
-};
-#endif
-
 // Singleton class tracking property bag association with objects
 class FPropertyBagRepository
 {
+public:
+
+	// Placeholder object feature flags; external callers can query a particular subfeature to see if it's enabled.
+	// Note: These have CVar counterparts defined in the .cpp file. Add an entry there to also create a toggle switch.
+	enum class EPlaceholderObjectFeature : uint8
+	{
+		/** Replace missing type imports with a placeholder type on load so that exports of the missing type can be serialized */
+		ReplaceMissingTypeImportsOnLoad,
+		/** Serialize references to placeholder exports on load so that they remain persistent and become visible to referencers */
+		SerializeExportReferencesOnLoad,
+		/** Replace missing types after reinstancing with a placeholder type so that data is not lost after rebuilding scripts */
+		ReplaceMissingReinstancedTypes,
+		/** Replace dead class types after a script compile with a placeholder type so that data is not lost after a script build failure */
+		ReplaceDeadClassInstanceTypes,
+		/** Always replace missing type imports with a placeholder type on load, regardless of the underlying UClass derivative type */
+		DoNotRestrictMissingTypeImports,
+		/** Disable validation check that stops placeholders from being saved */
+		AllowSavingPlaceholderExports,
+	};
+
+#if WITH_EDITORONLY_DATA
+
+private:
 	struct FPropertyBagAssociationData
 	{
 		void Destroy();
@@ -42,8 +53,6 @@ class FPropertyBagRepository
 	//friend UObjectBase;
 	//friend UStruct;
 	friend struct FScopedInstanceDataObjectLoad;
-
-private:
 	friend class FPropertyBagRepositoryLock;
 	mutable FTransactionallySafeCriticalSection CriticalSection;
 
@@ -130,30 +139,20 @@ public:
 	COREUOBJECT_API const UObject* FindInstanceForDataObject(const UObject* InstanceDataObject) const;
 
 	#if STATS
+	struct FPropertyBagRepositoryStats
+	{
+		size_t IDOMemoryBytes;
+		int32 NumIDOs;
+		int32 NumIDOsWithLooseProperties;
+		int32 NumPlaceholderTypes;
+	};
+
 	UE_INTERNAL void GatherStats(FPropertyBagRepositoryStats& Stats);
 	#endif
 
 	UE_INTERNAL void DumpIDOs(FOutputDevice& OutputDevice);
 	UE_INTERNAL void DumpPropertyBagPlaceholderTypes(FOutputDevice& OutputDevice);
 
-	// Placeholder object feature flags; external callers can query a particular subfeature to see if it's enabled.
-	// Note: These have CVar counterparts defined in the .cpp file. Add an entry there to also create a toggle switch.
-	enum class EPlaceholderObjectFeature : uint8
-	{
-		/** Replace missing type imports with a placeholder type on load so that exports of the missing type can be serialized */
-		ReplaceMissingTypeImportsOnLoad,
-		/** Serialize references to placeholder exports on load so that they remain persistent and become visible to referencers */
-		SerializeExportReferencesOnLoad,
-		/** Replace missing types after reinstancing with a placeholder type so that data is not lost after rebuilding scripts */
-		ReplaceMissingReinstancedTypes,
-		/** Replace dead class types after a script compile with a placeholder type so that data is not lost after a script build failure */
-		ReplaceDeadClassInstanceTypes,
-		/** Always replace missing type imports with a placeholder type on load, regardless of the underlying UClass derivative type */
-		DoNotRestrictMissingTypeImports,
-		/** Disable validation check that stops placeholders from being saved */
-		AllowSavingPlaceholderExports,
-	};
-
 	// query for whether or not the given struct/class is a placeholder type
 	static COREUOBJECT_API bool IsPropertyBagPlaceholderType(const UStruct* Type);
 	// query for whether or not the given object was created as a placeholder type
@@ -191,8 +190,21 @@ private:
 
 	// Delete owner reference and disassociate all data. Returns success.
 	bool RemoveAssociationUnsafe(const UObject* Owner);
-};
+
+#else
+	// Stub out accessors when the system is not available in cooked builds
+
+	// query for whether or not the given struct/class is a placeholder type
+	static inline bool IsPropertyBagPlaceholderType(const UStruct* Type) { return false; }
+	// query for whether or not the given object was created as a placeholder type
+	static inline bool IsPropertyBagPlaceholderObject(const UObject* Object) { return false; }
+	// query for whether or not creating property bag placeholder objects should be allowed
+	static inline bool IsPropertyBagPlaceholderObjectSupportEnabled() { return false; }
+	// query for whether or not a specific property bag placeholder object feature is enabled
+	static inline bool IsPropertyBagPlaceholderObjectFeatureEnabled(EPlaceholderObjectFeature Feature) { return false; }
 
 #endif // WITH_EDITORONLY_DATA
 
+};
+
 } // UE
diff --git a/Engine/Source/Runtime/CoreUObject/Private/UObject/InstanceDataObjectUtils.cpp b/Engine/Source/Runtime/CoreUObject/Private/UObject/InstanceDataObjectUtils.cpp
--- a/Engine/Source/Runtime/CoreUObject/Private/UObject/InstanceDataObjectUtils.cpp
+++ b/Engine/Source/Runtime/CoreUObject/Private/UObject/InstanceDataObjectUtils.cpp
@@ -529,7 +529,7 @@ namespace UE
 			return Metric;
 		};
 
-		FPropertyBagRepositoryStats Stats;
+		FPropertyBagRepository::FPropertyBagRepositoryStats Stats;
 		FPropertyBagRepository::Get().GatherStats(Stats);
 
 		int32 NumIDOClasses = 0;
diff --git a/Engine/Source/Runtime/CoreUObject/Private/UObject/LinkerLoad.cpp b/Engine/Source/Runtime/CoreUObject/Private/UObject/LinkerLoad.cpp
--- a/Engine/Source/Runtime/CoreUObject/Private/UObject/LinkerLoad.cpp
+++ b/Engine/Source/Runtime/CoreUObject/Private/UObject/LinkerLoad.cpp
@@ -5751,11 +5751,9 @@ void FLinkerLoad::Preload( UObject* Object )
 					{
 						UE_ASSET_LOG(LogLinker, Warning, PackagePath, TEXT("%s: Serial size mismatch: Got %d, Expected %" INT64_FMT), *Object->GetFullName(), (int32)SizeSerialized, ExpectedSerialSize);
 					}
-					else
+					else if (!UE::FPropertyBagRepository::IsPropertyBagPlaceholderObject(Object))
 					{
-#if !UE_INTERNAL_UNIVERSAL_PLACEHOLDERS
 						UE_ASSET_LOG(LogLinker, Fatal, PackagePath, TEXT("%s: Serial size mismatch: Got %d, Expected %" INT64_FMT), *Object->GetFullName(), (int32)SizeSerialized, ExpectedSerialSize);
-#endif //!UE_INTERNAL_UNIVERSAL_PLACEHOLDERS
 					}
 				}
 
