Allow boxing/unboxing with nested boxes.  This fixes a variety of issues with parametric types involving nested type instantiation.

#was-backout 50771316

#Synced-CL 50592211
#Synced-CL 50592211

#Synced-CL 50770334

diff --git a/Engine/Source/Runtime/CoreUObject/Private/UObject/InstanceDataTransforms.cpp b/Engine/Source/Runtime/CoreUObject/Private/UObject/InstanceDataTransforms.cpp
--- a/Engine/Source/Runtime/CoreUObject/Private/UObject/InstanceDataTransforms.cpp
+++ b/Engine/Source/Runtime/CoreUObject/Private/UObject/InstanceDataTransforms.cpp
@@ -287,7 +287,7 @@ namespace UE::Private
 				FString PropertyName = Segment.Name.ToString();
 				if (!PropertyName.StartsWith(TEXT("__verse_0x")))
 				{
-					PropertyName = Verse::Names::Private::MangleCasedName(PropertyName, PropertyName);
+					PropertyName = ::Verse::Names::Private::MangleCasedName(PropertyName, PropertyName);
 					Segment.Name = FName(PropertyName);
 					bUpdateSegment = true;
 				}
