Remove 7 virtual functions in UObject (either deprecated or toolonly)
Deprecate 4 functions in non editor

#Synced-CL 50763424

diff --git a/Engine/Plugins/Editor/DataValidation/Source/DataValidation/Private/EditorValidator_Localization.cpp b/Engine/Plugins/Editor/DataValidation/Source/DataValidation/Private/EditorValidator_Localization.cpp
--- a/Engine/Plugins/Editor/DataValidation/Source/DataValidation/Private/EditorValidator_Localization.cpp
+++ b/Engine/Plugins/Editor/DataValidation/Source/DataValidation/Private/EditorValidator_Localization.cpp
@@ -37,8 +37,9 @@ EDataValidationResult UEditorValidator_Localization::ValidateLoadedAsset_Impleme
 	IAssetRegistry& AssetRegistry = FModuleManager::LoadModuleChecked<FAssetRegistryModule>(NAME_AssetRegistry).Get();
 
 	const FString AssetPathName = InAsset->GetPathName();
-
-	if (InAsset->IsLocalizedResource())
+	
+	const UPackage* AssetPackage = InAsset->GetOutermost();
+	if (AssetPackage && FPackageName::IsLocalizedPackage(AssetPackage->GetPathName()))
 	{
 		// Get the source path for this asset
 		FString SourceObjectPath;
@@ -68,7 +69,9 @@ EDataValidationResult UEditorValidator_Localization::ValidateLoadedAsset_Impleme
 						if (UObject* ThisDestinationObject = ThisRedirector->DestinationObject)
 						{
 							FString ThisDestinationSourceObjectPath;
-							if (ThisDestinationObject->IsLocalizedResource() && FPackageLocalizationUtil::ConvertLocalizedToSource(ThisDestinationObject->GetPathName(), ThisDestinationSourceObjectPath))
+							const UPackage* DestinationObjectPackage = ThisDestinationObject->GetOutermost();
+							if (DestinationObjectPackage && FPackageName::IsLocalizedPackage(DestinationObjectPackage->GetPathName()) 
+								&& FPackageLocalizationUtil::ConvertLocalizedToSource(ThisDestinationObject->GetPathName(), ThisDestinationSourceObjectPath))
 							{
 								FString SourceRedirectorDest;
 								if (SourceAssetData.GetTagValue(TEXT("DestinationObject"), SourceRedirectorDest))
diff --git a/Engine/Plugins/MovieScene/TemplateSequence/Source/TemplateSequence/Public/CameraAnimationSequencePlayer.h b/Engine/Plugins/MovieScene/TemplateSequence/Source/TemplateSequence/Public/CameraAnimationSequencePlayer.h
--- a/Engine/Plugins/MovieScene/TemplateSequence/Source/TemplateSequence/Public/CameraAnimationSequencePlayer.h
+++ b/Engine/Plugins/MovieScene/TemplateSequence/Source/TemplateSequence/Public/CameraAnimationSequencePlayer.h
@@ -186,7 +186,6 @@ public:
 	virtual void UpdateCameraCut(UObject* CameraObject, const EMovieSceneCameraCutParams& CameraCutParams) override {}
 	
 	// UObject interface 
-	virtual bool IsDestructionThreadSafe() const override { return false; }
 	UE_API virtual void BeginDestroy() override;
 
 protected:
diff --git a/Engine/Source/Editor/UnrealEd/Private/Commandlets/ExportDialogueScriptCommandlet.cpp b/Engine/Source/Editor/UnrealEd/Private/Commandlets/ExportDialogueScriptCommandlet.cpp
--- a/Engine/Source/Editor/UnrealEd/Private/Commandlets/ExportDialogueScriptCommandlet.cpp
+++ b/Engine/Source/Editor/UnrealEd/Private/Commandlets/ExportDialogueScriptCommandlet.cpp
@@ -448,7 +448,8 @@ void UExportDialogueScriptCommandlet::PopulateDialogueScriptEntry(const UDialogu
 			{
 				if (LocalizedContextMapping.Context == InContext)
 				{
-					return LocalizedContextMapping.SoundWave && LocalizedContextMapping.SoundWave->IsLocalizedResource();
+					const UPackage* ObjPackage = LocalizedContextMapping.SoundWave ? LocalizedContextMapping.SoundWave->GetOutermost() : nullptr;
+                	return ObjPackage && FPackageName::IsLocalizedPackage(ObjPackage->GetPathName());
 				}
 			}
 		}
diff --git a/Engine/Source/Runtime/CoreUObject/Private/Blueprint/BlueprintSupport.cpp b/Engine/Source/Runtime/CoreUObject/Private/Blueprint/BlueprintSupport.cpp
--- a/Engine/Source/Runtime/CoreUObject/Private/Blueprint/BlueprintSupport.cpp
+++ b/Engine/Source/Runtime/CoreUObject/Private/Blueprint/BlueprintSupport.cpp
@@ -785,6 +785,7 @@ bool FLinkerLoad::RegenerateBlueprintClass(UClass* LoadClass, UObject* ClassDefa
 		FPreloadMembersHelper::PreloadExternalNativeDependencies(ClassDefaultObject);
 
 		{
+#if WITH_EDITOR
 			// RegenerateClass largely performs redundant work since we already preloaded the remaining exports.
 			// However, in some circumstances, a re-preload might occur on an export that doesn't have RF_LoadCompleted set due to a cycle.
 			// This is largely dependent on how the derived Blueprint (eg: Widget BPs, Animation BPs, etc.) decides to handle this.
@@ -794,6 +795,7 @@ bool FLinkerLoad::RegenerateBlueprintClass(UClass* LoadClass, UObject* ClassDefa
 
 			// A regenerated class won't be post-loaded, so we can clear these flags.
 			LoadClass->ClearFlags(RF_NeedPostLoad | RF_NeedPostLoadSubobjects);
+#endif
 		}
 
 		// After regeneration, we can now notify placeholders that it's safe to create instances of our class.
@@ -873,6 +875,7 @@ bool FLinkerLoad::RegenerateBlueprintClass(UClass* LoadClass, UObject* ClassDefa
 		}
 
 		{
+#if WITH_EDITOR
 			ClassSourceObject = GetClassSourceObjectLambda(LoadClass);
 
 			// Preload the blueprint to make sure it has all the data the class needs for regeneration
@@ -886,7 +889,6 @@ bool FLinkerLoad::RegenerateBlueprintClass(UClass* LoadClass, UObject* ClassDefa
 				LoadClass->ClearFlags(RF_NeedLoad | RF_NeedPostLoad | RF_NeedPostLoadSubobjects);
 			}
 
-#if WITH_EDITOR
 			// Ensure that the class source object is marked standalone so it doesn't get GC'd in the editor.
 			// In particular, this is needed for a BPGC asset in a cooked package.
 			if (LoadClass->bCooked)
diff --git a/Engine/Source/Runtime/CoreUObject/Private/UObject/GarbageCollection.cpp b/Engine/Source/Runtime/CoreUObject/Private/UObject/GarbageCollection.cpp
--- a/Engine/Source/Runtime/CoreUObject/Private/UObject/GarbageCollection.cpp
+++ b/Engine/Source/Runtime/CoreUObject/Private/UObject/GarbageCollection.cpp
@@ -6522,11 +6522,6 @@ void UObject::AddReferencedObjects(UObject* This, FReferenceCollector& Collector
 #endif
 }
 
-bool UObject::IsDestructionThreadSafe() const
-{
-	return false;
-}
-
 bool UObjectBaseUtility::ThisThreadAtomicallyClearedRFUnreachable()
 {
 	return UE::GC::Private::FGCFlags::ThisThreadAtomicallyClearedRFUnreachable(GUObjectArray.IndexToObject(InternalIndex));
diff --git a/Engine/Source/Runtime/CoreUObject/Private/UObject/Obj.cpp b/Engine/Source/Runtime/CoreUObject/Private/UObject/Obj.cpp
--- a/Engine/Source/Runtime/CoreUObject/Private/UObject/Obj.cpp
+++ b/Engine/Source/Runtime/CoreUObject/Private/UObject/Obj.cpp
@@ -2136,23 +2136,6 @@ void UObject::BuildSubobjectMapping(UObject* OtherObject, TMap<UObject*, UObject
 	};
 }
 
-
-void UObject::CollectDefaultSubobjects( TArray<UObject*>& OutSubobjectArray, bool bIncludeNestedSubobjects/*=false*/ ) const
-{
-		OutSubobjectArray.Empty();
-		GetObjectsWithOuter(this, OutSubobjectArray, bIncludeNestedSubobjects);
-
-		// Remove contained objects that are not subobjects.
-		for (int32 ComponentIndex = 0; ComponentIndex < OutSubobjectArray.Num(); ComponentIndex++)
-		{
-			UObject* PotentialComponent = OutSubobjectArray[ComponentIndex];
-			if (!PotentialComponent->IsDefaultSubobject())
-			{
-				OutSubobjectArray.RemoveAtSwap(ComponentIndex--);
-			}
-		}
-	}
-
 /**
  * FSubobjectReferenceFinder.
  * Helper class used to collect default subobjects of other objects than the referencing object.
@@ -2794,11 +2777,13 @@ FPrimaryAssetId UObject::GetPrimaryAssetId() const
 	return FPrimaryAssetId();
 }
 
+#if !UE_STRIP_DEPRECATED_PROPERTIES 
 bool UObject::IsLocalizedResource() const
 {
 	const UPackage* ObjPackage = GetOutermost();
 	return ObjPackage && FPackageName::IsLocalizedPackage(ObjPackage->GetPathName());
 }
+#endif
 
 bool UObject::IsSafeForRootSet() const
 {
diff --git a/Engine/Source/Runtime/CoreUObject/Public/UObject/Object.h b/Engine/Source/Runtime/CoreUObject/Public/UObject/Object.h
--- a/Engine/Source/Runtime/CoreUObject/Public/UObject/Object.h
+++ b/Engine/Source/Runtime/CoreUObject/Public/UObject/Object.h
@@ -616,14 +616,6 @@ public:
 		return false;
 	}
 
-	/**
-	* Called during garbage collection to determine if an object can have its destructor called on a worker thread.
-	*
-	* @return	true if this object's destructor is thread safe
-	*/
-	UE_DEPRECATED(5.5, "Multithreaded destruction of UObjects is no longer supported.")
-	COREUOBJECT_API virtual bool IsDestructionThreadSafe() const;
-
 	/**
 	* Called during cooking. Must return all objects that will be Preload()ed when this is serialized at load time. Only used by the EDL.
 	*
@@ -928,11 +920,6 @@ public:
 
 #if WITH_EDITOR
 
-	UE_DEPRECATED(5.5, "This function is no longer called. If you require this functionality, please implement ThreadedPostLoadAssetRegistryTagsOverride. " \
-		"Note that this function MUST be safe to call on a background thread. If you have functionality that doesn't meet that requirement, " \
-		"consider using the AssetRegistry's AssetAddedEvent which fires on the game thread.")
-	virtual void PostLoadAssetRegistryTags(const FAssetData& InAssetData, TArray<FAssetRegistryTag>& OutTagsAndValuesToUpdate) const final {}
-
 	/** Contains an AssetData referencing the asset to be processed and an array of tags and values to be updated on that asset */
 	struct FPostLoadAssetRegistryTagsContext
 	{
@@ -1031,7 +1018,11 @@ public:
 	COREUOBJECT_API virtual FPrimaryAssetId GetPrimaryAssetId() const;
 
 	/** Returns true if this object is considered a localized resource. */
+
+#if !UE_STRIP_DEPRECATED_PROPERTIES 
+	UE_DEPRECATED(5.8, "Use FPackageName::IsLocalizedPackage instead")
 	COREUOBJECT_API virtual bool IsLocalizedResource() const;
+#endif
 
 	/** Returns true if this object is safe to add to the root set. */
 	COREUOBJECT_API virtual bool IsSafeForRootSet() const;
@@ -1073,8 +1064,10 @@ public:
 	COREUOBJECT_API virtual bool IsSupportedForNetworking() const;
 
 	/** Returns a list of sub-objects that have stable names for networking */
+#if !UE_STRIP_DEPRECATED_PROPERTIES 
 	UE_DEPRECATED(5.7, "This function is useless now and will be deleted")
 	virtual void GetSubobjectsWithStableNamesForNetworking(TArray<UObject*> &ObjList) {}
+#endif
 
 	/** Called right before receiving a bunch */
 	COREUOBJECT_API virtual void PreNetReceive();
@@ -1275,13 +1268,6 @@ public:
 	 */
 	COREUOBJECT_API virtual void BuildSubobjectMapping(UObject* OtherObject, TMap<UObject*, UObject*>& ObjectMapping) const;
 
-	/** 
-	 * Gets all subobjects inside this object that return true for IsDefaultSubobject.
-	 * The nested behavior is inconsistent because IsDefaultSubobject does not work reliably for nested subobjects and it is less efficient than GetDefaultSubobjects
-	 */
-	UE_DEPRECATED(5.4, "Call GetDefaultSubobjects for top level subobjects or use ForEachObjectWithOuter with a more precise check")
-	COREUOBJECT_API void CollectDefaultSubobjects( TArray<UObject*>& OutDefaultSubobjects, bool bIncludeNestedSubobjects=false ) const;
-
 	/**
 	 * Checks default sub-object assumptions.
 	 *
@@ -1562,8 +1548,14 @@ public:
 	 * Called on the target when a class is loaded with ClassGeneratedBy is loaded.  Should regenerate the class if needed, and return the updated class
 	 * @return	Updated instance of the class, if needed, or NULL if no regeneration was necessary
 	 */
-	virtual UClass* RegenerateClass(UClass* ClassToRegenerate, UObject* PreviousCDO) { return NULL; };
-
+#if WITH_EDITOR
+	virtual UClass* RegenerateClass(UClass* ClassToRegenerate, UObject* PreviousCDO) { return nullptr; };
+#else
+#if !UE_STRIP_DEPRECATED_PROPERTIES 
+	UE_DEPRECATED(5.8, "This function is now editor only")
+	virtual UClass* RegenerateClass(UClass* ClassToRegenerate, UObject* PreviousCDO) { return nullptr; };
+#endif
+#endif
 	/** 
 	 * Returns whether this object is contained in or part of a blueprint object
 	 */
@@ -1575,7 +1567,14 @@ public:
 	COREUOBJECT_API void DestroyNonNativeProperties();
 
 	/** Called during subobject creation to mark this component as editor only, which causes it to get stripped in packaged builds */
+#if WITH_EDITORONLY_DATA
 	virtual void MarkAsEditorOnlySubobject() { }
+#else
+#if !UE_STRIP_DEPRECATED_PROPERTIES 
+	UE_DEPRECATED(5.8, "This function is now editor only")
+	virtual void MarkAsEditorOnlySubobject() { }
+#endif
+#endif
 
 	/**
 	 * Abort with a member function call at the top of the callstack, helping to ensure that most platforms will stuff this object's memory into the resulting minidump.
@@ -1809,7 +1808,12 @@ public:
 		NETFIELD_REP_END = -1
 	};
 
-	virtual void ValidateGeneratedRepEnums(const TArray<struct FRepRecord>& ClassReps) const {}	
+#if VALIDATE_CLASS_REPS
+	virtual void ValidateGeneratedRepEnums(const TArray<struct FRepRecord>& ClassReps) const {}
+#elif !UE_STRIP_DEPRECATED_PROPERTIES 
+	UE_DEPRECATED(5.8, "VALIDATE_CLASS_REPS is not defined")
+	virtual void ValidateGeneratedRepEnums(const TArray<struct FRepRecord>& ClassReps) const {}
+#endif
 
 private:
 	
