IoStore: Decrement the unused serialized CookedHeaderSize by 8 x NumImports to eliminate the impact of CL47033040 on unversioned cooked game patch sizes.
Also hack serialization of FObjectImport to avoid harvesting/serializing the name "None".
Note: The CookedHeaderSize field should ultimately be removed or be serialized as 0.

#rnx
#tests BuildCookRun Lyra, -singlepackagenorefs cook

#Synced-CL 50480306

diff --git a/Engine/Source/Developer/IoStoreUtilities/Private/PackageStoreOptimizer.cpp b/Engine/Source/Developer/IoStoreUtilities/Private/PackageStoreOptimizer.cpp
--- a/Engine/Source/Developer/IoStoreUtilities/Private/PackageStoreOptimizer.cpp
+++ b/Engine/Source/Developer/IoStoreUtilities/Private/PackageStoreOptimizer.cpp
@@ -1147,13 +1147,16 @@ void FPackageStoreOptimizer::FinalizePackageHeader(FPackageStorePackage* Package
 	//       It should ultimately be removed, until then freeze the FPackageFileSummary baseline impact on serialized CookedHeaderSize at
 	//       EUnrealEngineObjectUE5Version::OS_SUB_OBJECT_SHADOW_SERIALIZATION.
 	// NOTE: EUnrealEngineObjectUE5Version::IMPORT_TYPE_HIERARCHIES adds two int32 fields to FPackageFileSummary. This increases the total header size by 8 bytes.
+	// NOTE: operator<<(FStructuredArchive::FSlot Slot, FObjectImport& I) added unversioned cooked serialization of I.PackageName, i.e. two int32 fields for every import entry.
 #if WITH_EDITORONLY_DATA
 	constexpr size_t ExpectedPackageFileSummarySize = 360;
 #else
 	constexpr size_t ExpectedPackageFileSummarySize = 344;
 #endif
-	static_assert(sizeof(FPackageFileSummary) == ExpectedPackageFileSummarySize, "Any change to FPackageFileSummary serialization requires adjusting the offset below, transient changes can just update the size in this ssert.");
-	constexpr uint32 CookedHeaderSizeOffsetToMinimizePatchSize = 8;
+	static_assert(sizeof(FPackageFileSummary) == ExpectedPackageFileSummarySize, "Any change to FPackageFileSummary serialization requires adjusting the offset below, transient changes can just update the size in this assert.");
+	constexpr uint32 CookedHeaderSizeForImportTypeHierarchies = 8; // sizeof(ImportTypeHierarchiesCount) + sizeof(ImportTypeHierarchiesOffset) = 8
+	const uint32 CookedHeaderSizeForImportPackageNames = 8 * Package->Imports.Num(); // sizeof(NameIndex) + sizeof(NameNumber) = 8
+	const uint32 CookedHeaderSizeOffsetToMinimizePatchSize = CookedHeaderSizeForImportTypeHierarchies + CookedHeaderSizeForImportPackageNames;
 
 	Package->HeaderBuffer = FIoBuffer(HeaderSize);
 	uint8* HeaderData = Package->HeaderBuffer.Data();
diff --git a/Engine/Source/Runtime/CoreUObject/Private/UObject/ObjectResource.cpp b/Engine/Source/Runtime/CoreUObject/Private/UObject/ObjectResource.cpp
--- a/Engine/Source/Runtime/CoreUObject/Private/UObject/ObjectResource.cpp
+++ b/Engine/Source/Runtime/CoreUObject/Private/UObject/ObjectResource.cpp
@@ -363,7 +363,26 @@ void operator<<(FStructuredArchive::FSlot Slot, FObjectImport& I)
 
 	if (Slot.GetUnderlyingArchive().UEVer() >= VER_UE4_NON_OUTER_PACKAGE_IMPORT)
 	{
-		Record << SA_VALUE(TEXT("PackageName"), I.PackageName);
+		if (Slot.GetUnderlyingArchive().IsSaving())
+		{
+			FName SavePackageName = I.PackageName;
+			if (Slot.GetUnderlyingArchive().IsFilterEditorOnly() && SavePackageName.IsNone())
+			{
+				// Hack to avoid collecting the "None" name in cooked packages that does not use the external package feature.
+				// This minimizes the patch size impact, else the PackageName addition would impact the majority of cooked packages.
+				// Serialize an existing name table entry instead, and reset it on load for the legacy cooked loading path.
+				// Note that the PackageStoreOptimizer strips this data anyway, so it will never be loaded in the iostore loading path.
+				ensureMsgf(I.ObjectName.IsNone() || I.ObjectName != I.PackageName,
+					TEXT("Unexpected import names in archive %s: ObjectName=%s, PackageName=%s"),
+					*Slot.GetUnderlyingArchive().GetArchiveName(), *I.ObjectName.ToString(), *I.PackageName.ToString());
+				SavePackageName = I.ObjectName;
+			}
+			Record << SA_VALUE(TEXT("PackageName"), SavePackageName);
+		}
+		else
+		{
+			Record << SA_VALUE(TEXT("PackageName"), I.PackageName);
+		}
 	}
 
 	if (Slot.GetUnderlyingArchive().UEVer() >= EUnrealEngineObjectUE5Version::OPTIONAL_RESOURCES)
@@ -376,6 +395,12 @@ void operator<<(FStructuredArchive::FSlot Slot, FObjectImport& I)
 		I.SourceLinker = NULL;
 		I.SourceIndex = INDEX_NONE;
 		I.XObject = NULL;
+		if (Slot.GetUnderlyingArchive().IsFilterEditorOnly() && I.PackageName == I.ObjectName)
+		{
+			// Reset the data from the above save hack to avoid collecting the "None" name in cooked packages that
+			// does not use the external package feature.
+			I.PackageName = FName{};
+		}
 	}
 }
 
diff --git a/Engine/Source/Runtime/CoreUObject/Private/UObject/SavePackage/PackageHarvester.cpp b/Engine/Source/Runtime/CoreUObject/Private/UObject/SavePackage/PackageHarvester.cpp
--- a/Engine/Source/Runtime/CoreUObject/Private/UObject/SavePackage/PackageHarvester.cpp
+++ b/Engine/Source/Runtime/CoreUObject/Private/UObject/SavePackage/PackageHarvester.cpp
@@ -931,7 +931,7 @@ void FPackageHarvester::ProcessImport(TObjectPtr<UObject> InObjPtr,
 				EIllegalRefReason::ExternalPackage, Package->GetName());
 		}
 	}
-	else
+	else if (!IsFilterEditorOnly())
 	{
 		// operator<<(FStructuredArchive::FSlot Slot, FObjectImport& I) will need to write NAME_None for this empty ExternalPackage pointer
 		HarvestPackageHeaderName(NAME_None);
