AssetRegistry: Change AssetDataTagMap to support a feature required by ActiveMountLoader: reusing an FStore from an earlier load of the same AssetRegistryState rather than recreating it. Added a weakpointer class to pass the FStore from the previous load into the new load if it still exists, and added SkipBytes to FixedTagPrivate::TSerializer.
#rnx

#Synced-CL 50777996

diff --git a/Engine/Source/Runtime/AssetRegistry/Private/AssetDataGatherer.cpp b/Engine/Source/Runtime/AssetRegistry/Private/AssetDataGatherer.cpp
--- a/Engine/Source/Runtime/AssetRegistry/Private/AssetDataGatherer.cpp
+++ b/Engine/Source/Runtime/AssetRegistry/Private/AssetDataGatherer.cpp
@@ -3,9 +3,9 @@
 #include "AssetDataGatherer.h"
 #include "AssetDataGathererPrivate.h"
 #include "AssetRegistry.h"
-#include "AssetRegistryImpl.h"
 #include "AssetRegistry/AssetDataTagMapSerializationDetails.h"
 #include "AssetRegistry/AssetRegistryTelemetry.h"
+#include "AssetRegistryImpl.h"
 
 #include "Algo/AnyOf.h"
 #include "Algo/Find.h"
@@ -5439,7 +5439,9 @@ TArray<FCachePayload> LoadCacheFiles(TConstArrayView<FString> InCacheFilenames)
 			Header.Version = VersionToRead;
 			Header.bFilterEditorOnlyData = false;
 
-			FAssetRegistryReader RegistryReader(FilemappedReader, UE::AssetDataGather::Private::bMemoryMappedCacheFiles ? BackingFile : nullptr, NumAsyncWorkers, Header);
+			FAssetRegistryReader RegistryReader(FilemappedReader,
+				UE::AssetDataGather::Private::bMemoryMappedCacheFiles ? BackingFile : nullptr,
+				FixedTagPrivate::FWeakStorePtr(), NumAsyncWorkers, Header);
 			return RegistryReader.IsError() ? FCachePayload() : SerializeCacheLoad(RegistryReader);
 		};
 
diff --git a/Engine/Source/Runtime/AssetRegistry/Private/AssetRegistryArchive.cpp b/Engine/Source/Runtime/AssetRegistry/Private/AssetRegistryArchive.cpp
--- a/Engine/Source/Runtime/AssetRegistry/Private/AssetRegistryArchive.cpp
+++ b/Engine/Source/Runtime/AssetRegistry/Private/AssetRegistryArchive.cpp
@@ -268,7 +268,10 @@ void FAssetRegistryWriter::SerializeTagsAndBundles(const FAssetData& Out)
 
 #endif
 
-FAssetRegistryReader::FAssetRegistryReader(FArchive& Inner, TSharedPtr<UE::AssetRegistry::FMemoryMappedFile> InBackingFile, int32 NumWorkers, FAssetRegistryHeader Header)
+FAssetRegistryReader::FAssetRegistryReader(FArchive& Inner,
+	TSharedPtr<UE::AssetRegistry::FMemoryMappedFile> InBackingFile,
+	const FixedTagPrivate::FWeakStorePtr& PreviousLoadedStore,
+	int32 NumWorkers, FAssetRegistryHeader Header)
 	: FArchiveProxy(Inner)
 {
 	check(IsLoading());
@@ -284,7 +287,8 @@ FAssetRegistryReader::FAssetRegistryReader(FArchive& Inner, TSharedPtr<UE::Asset
 	}
 	else
 	{
-		if (NumWorkers > 0)
+		TRefCountPtr<const FixedTagPrivate::FStore> PreviousStore = PreviousLoadedStore.Pin();
+		if (NumWorkers > 0 && !PreviousStore) // Multithreaded loading does not yet support PreviousStore
 		{
 			ENameBatchLoadingFlags LoadingFlags = ENameBatchLoadingFlags::None;
 #if ASSET_REGISTRY_LOAD_NAMES_IN_ORDER
@@ -302,7 +306,15 @@ FAssetRegistryReader::FAssetRegistryReader(FArchive& Inner, TSharedPtr<UE::Asset
 		else
 		{
 			Names = LoadNameBatch(Inner, ENameBatchLoadingFlags::RespectOrder);
-			Tags = FixedTagPrivate::LoadStore(*this, Header.Version, InBackingFile);
+			if (PreviousStore)
+			{
+				Tags = MoveTemp(PreviousStore);
+				FixedTagPrivate::SkipBytesStore(*this, Header.Version);
+			}
+			else
+			{
+				Tags = FixedTagPrivate::LoadStore(*this, Header.Version, InBackingFile);
+			}
 		}
 	}
 }
@@ -312,6 +324,11 @@ FAssetRegistryReader::~FAssetRegistryReader()
 	WaitForTasks();
 }
 
+const TRefCountPtr<const FixedTagPrivate::FStore>& FAssetRegistryReader::GetTags() const
+{
+	return Tags;
+}
+
 void FAssetRegistryReader::WaitForTasks()
 {
 	if (Task.IsValid())
@@ -438,7 +455,7 @@ bool FAssetRegistryTagSerializationTest::RunTest(const FString& Parameters)
 	{
 		FMemoryReader DataReader(Data);
 		FAssetRegistryHeader Header;
-		FAssetRegistryReader RegistryReader(DataReader, nullptr, 0, Header);
+		FAssetRegistryReader RegistryReader(DataReader, nullptr, FixedTagPrivate::FWeakStorePtr(), 0, Header);
 		RegistryReader.WaitForTasks();
 
 		for (FAssetDataTagMapSharedView& FixedMap : FixedMaps)
@@ -453,7 +470,7 @@ bool FAssetRegistryTagSerializationTest::RunTest(const FString& Parameters)
 	{
 		FMemoryReader DataReader(Data);
 		FAssetRegistryHeader Header;
-		FAssetRegistryReader RegistryReader(DataReader, nullptr, 0, Header);
+		FAssetRegistryReader RegistryReader(DataReader, nullptr, FixedTagPrivate::FWeakStorePtr(), 0, Header);
 		RegistryReader.WaitForTasks();
 
 		for (const FAssetDataTagMapSharedView& FixedMap1 : FixedMaps)
diff --git a/Engine/Source/Runtime/AssetRegistry/Private/AssetRegistryArchive.h b/Engine/Source/Runtime/AssetRegistry/Private/AssetRegistryArchive.h
--- a/Engine/Source/Runtime/AssetRegistry/Private/AssetRegistryArchive.h
+++ b/Engine/Source/Runtime/AssetRegistry/Private/AssetRegistryArchive.h
@@ -4,6 +4,7 @@
 
 #include "AssetRegistry/AssetDataTagMapSerializationDetails.h"
 #include "AssetRegistry/AssetRegistrySerializationDetails.h"
+#include "AssetRegistry/AssetRegistryStateSerialization.h"
 #include "Serialization/LargeMemoryWriter.h"
 #include "UObject/NameBatchSerialization.h"
 
@@ -25,8 +26,10 @@ class FAssetRegistryReader : public FArchiveProxy
 public:
 	/// @param InBackingFile shared pointer to a class that allows mapping the underlying file into memory
 	/// @param NumWorkers > 0 for parallel loading
-	FAssetRegistryReader(FArchive& Inner, TSharedPtr<UE::AssetRegistry::FMemoryMappedFile> InBackingFile, int32 NumWorkers, FAssetRegistryHeader Header);
+	FAssetRegistryReader(FArchive& Inner, TSharedPtr<UE::AssetRegistry::FMemoryMappedFile> InBackingFile,
+		const FixedTagPrivate::FWeakStorePtr& PreviousLoadedStore, int32 NumWorkers, FAssetRegistryHeader Header);
 	~FAssetRegistryReader();
+	const TRefCountPtr<const FixedTagPrivate::FStore>& GetTags() const;
 
 	virtual FArchive& operator<<(FName& Value) override;
 
diff --git a/Engine/Source/Runtime/AssetRegistry/Private/AssetRegistryState.cpp b/Engine/Source/Runtime/AssetRegistry/Private/AssetRegistryState.cpp
--- a/Engine/Source/Runtime/AssetRegistry/Private/AssetRegistryState.cpp
+++ b/Engine/Source/Runtime/AssetRegistry/Private/AssetRegistryState.cpp
@@ -2352,14 +2352,24 @@ bool FAssetRegistryState::Save(FArchive& OriginalAr, const FAssetRegistrySeriali
 
 bool FAssetRegistryState::Load(FArchive& OriginalAr, const FAssetRegistryLoadOptions& Options,
 	FAssetRegistryVersion::Type* OutVersion)
+{
+	FAssetRegistryLoadResults Results;
+	bool bResult = Load(OriginalAr, Options, Results);
+	if (OutVersion)
+	{
+		*OutVersion = Results.Version;
+	}
+	return bResult;
+}
+
+bool FAssetRegistryState::Load(FArchive& OriginalAr, const FAssetRegistryLoadOptions& Options,
+	FAssetRegistryLoadResults& OutResults)
 {
 	LLM_SCOPE(ELLMTag::AssetRegistry);
 	FAssetRegistryHeader Header;
 	Header.SerializeHeader(OriginalAr);
-	if (OutVersion != nullptr)
-	{
-		*OutVersion = Header.Version;
-	}
+	OutResults.Version = Header.Version;
+	OutResults.Store = FixedTagPrivate::FWeakStorePtr();
 
 	FSoftObjectPathSerializationScope SerializationScope(NAME_None, NAME_None,
 		ESoftObjectPathCollectType::NonPackage, ESoftObjectPathSerializeType::AlwaysSerialize);
@@ -2376,7 +2386,8 @@ bool FAssetRegistryState::Load(FArchive& OriginalAr, const FAssetRegistryLoadOpt
 	}
 	else
 	{
-		FAssetRegistryReader Reader(OriginalAr, nullptr, Options.ParallelWorkers, Header);
+		FAssetRegistryReader Reader(OriginalAr, nullptr /* BackingFile */,
+			Options.Store, Options.ParallelWorkers, Header);
 
 		if (Reader.IsError())
 		{
@@ -2388,6 +2399,7 @@ bool FAssetRegistryState::Load(FArchive& OriginalAr, const FAssetRegistryLoadOpt
 		Load(Reader, Header, Options);
 
 		Reader.WaitForTasks();
+		OutResults.Store = FixedTagPrivate::FWeakStorePtr(Reader.GetTags());
 	}
 
 	return !OriginalAr.IsError();
diff --git a/Engine/Source/Runtime/AssetRegistry/Public/AssetRegistry/AssetRegistryState.h b/Engine/Source/Runtime/AssetRegistry/Public/AssetRegistry/AssetRegistryState.h
--- a/Engine/Source/Runtime/AssetRegistry/Public/AssetRegistry/AssetRegistryState.h
+++ b/Engine/Source/Runtime/AssetRegistry/Public/AssetRegistry/AssetRegistryState.h
@@ -4,6 +4,7 @@
 
 #include "AssetRegistry/AssetData.h"
 #include "AssetRegistry/AssetDataMap.h"
+#include "AssetRegistry/AssetRegistryStateSerialization.h"
 #include "AssetRegistry/IAssetRegistry.h"
 #include "Containers/Array.h"
 #include "Containers/ArrayView.h"
@@ -125,9 +126,24 @@ struct FAssetRegistryLoadOptions
 	FAssetRegistryLoadOptions() = default;
 	explicit FAssetRegistryLoadOptions(const FAssetRegistrySerializationOptions& Options);
 
+	/**
+	 * Optional pointer to the FStore previously created by Load, in cases where the same AssetRegistry.bin
+	 * file is being reloaded. Caution: If this Store is not identical to the one that would be loaded from disk,
+	 * chaotic crashes can occur during loading or afterwards during reading of AsssetDatas.
+	 * If non-null, this FStore will be used for all TagData lookups and the serialization of the bytes for it
+	 * on disk will be skipped.
+	 */
+	FixedTagPrivate::FWeakStorePtr Store;
+	int32 ParallelWorkers = 0;
 	bool bLoadDependencies = true;
 	bool bLoadPackageData = true;
-	int32 ParallelWorkers = 0;
+};
+
+/** Extra data about the load of an AssetRegistryState from disk that is optionally returned from Load. */
+struct FAssetRegistryLoadResults
+{
+	FAssetRegistryVersion::Type Version;
+	FixedTagPrivate::FWeakStorePtr Store;
 };
 
 struct FAssetRegistryPruneOptions
@@ -516,7 +532,7 @@ public:
 	
 	/**
 	 * Initializes a cache from an existing using a set of filters.
-	 * This is more efficient than calling InitalizeFromExisting and then PruneAssetData.
+	 * This is more efficient than calling InitializeFromExisting and then PruneAssetData.
 	 *
 	 * @param ExistingState State to use initialize from
 	 * @param RequiredPackages If set, only these packages will be maintained.
@@ -542,6 +558,9 @@ public:
 	ASSETREGISTRY_API bool Load(FArchive& Ar,
 		const FAssetRegistryLoadOptions& Options = FAssetRegistryLoadOptions(),
 		FAssetRegistryVersion::Type* OutVersion = nullptr);
+	ASSETREGISTRY_API bool Load(FArchive& Ar,
+		const FAssetRegistryLoadOptions& Options,
+		FAssetRegistryLoadResults& OutResults);
 
 	/** 
 	* Example Usage:
diff --git a/Engine/Source/Runtime/CoreUObject/Private/AssetRegistry/AssetDataTagMap.cpp b/Engine/Source/Runtime/CoreUObject/Private/AssetRegistry/AssetDataTagMap.cpp
--- a/Engine/Source/Runtime/CoreUObject/Private/AssetRegistry/AssetDataTagMap.cpp
+++ b/Engine/Source/Runtime/CoreUObject/Private/AssetRegistry/AssetDataTagMap.cpp
@@ -5,6 +5,7 @@
 #include "Algo/Sort.h"
 #include "AssetRegistry/AssetDataTagMapSerializationDetails.h"
 #include "AssetRegistry/AssetRegistrySerializationDetails.h"
+#include "AssetRegistry/AssetRegistryStateSerialization.h"
 #include "AutoRTFM.h"
 #include "HAL/PlatformFile.h"
 #include "HAL/PlatformFileManager.h"
@@ -454,11 +455,25 @@ namespace FixedTagPrivate
 			return nullptr;
 		}
 
-		void Unregister(FStore& Store)
+		/**
+		 * Create a placeholder FStore that is used to trigger SkipSerialize functions on TSerializer
+		 * but is otherwise unused and is deleted before serialization returns.
+		 */
+		FStore* CreateSkipBytesStore()
 		{
-			UE::TScopeLock Lock(Mutex);
+			return new FStore(Capacity);
+		}
 
+		void Unregister(FStore& Store)
+		{
+			if (Store.Index == Capacity)
+			{
+				// A placeholder store used in SkipBytes calls. These are never registered or unregistered.
+				return;
+			}
 			check(Store.Index < Capacity);
+
+			UE::TScopeLock Lock(Mutex);
 			check(Data[Store.Index] == &Store);
 
 			Data[Store.Index] = nullptr;
@@ -1024,16 +1039,6 @@ namespace FixedTagPrivate
 		}
 	}
 
-	void FStore::Release() const
-	{
-		UE_AUTORTFM_ONCOMMIT(this)
-		{
-			if (RefCount.Decrement() == 0)
-			{
-				delete this;
-			}
-		};
-	}
 
 	void* FStore::AllocateMemoryForView(SIZE_T NumBytes, SIZE_T AlignmentRequired)
 	{
@@ -1046,6 +1051,12 @@ namespace FixedTagPrivate
 		return Memory;
 	}
 
+	FStore::FStore(uint32 InIndex) 
+		: Index(InIndex)
+	{
+		checkf(!AutoRTFM::IsClosed() || !AutoRTFM::IsOnCurrentTransactionStack(this), TEXT("Not allowed to construct a stack local within a transaction."));
+		RefCountObject = new FStoreRefCount(this);
+	}
 
 	FStore::~FStore()
 	{
@@ -1068,6 +1079,187 @@ namespace FixedTagPrivate
 		}
 	}
 
+	FStoreRefCount::FStoreRefCount(FStore* InOwner)
+		: Owner(InOwner)
+	{
+		WeakAddRef();
+	}
+
+	void FStoreRefCount::StrongRelease()
+	{
+		UE_AUTORTFM_ONCOMMIT(this)
+		{
+			if (StrongRefCount.fetch_add(-1, std::memory_order_acq_rel) == 1)
+			{
+				delete Owner;
+				Owner = nullptr;
+				// The call to WeakRelease might delete this; do not access this after this call to WeakRelease.
+				// Note that it will not delete this if external code still holds a weak pointer.
+				WeakRelease();
+			}
+		};
+	}
+
+	void FStoreRefCount::WeakAddRef()
+	{
+		UE_AUTORTFM_OPEN
+		{
+			WeakRefCount.fetch_add(1, std::memory_order_relaxed);
+		};
+
+		UE_AUTORTFM_ONABORT(this)
+		{
+			WeakRefCount.fetch_add(-1, std::memory_order_relaxed);
+		};
+	}
+
+	void FStoreRefCount::WeakRelease()
+	{
+		UE_AUTORTFM_ONCOMMIT(this)
+		{
+			if (WeakRefCount.fetch_add(-1, std::memory_order_acq_rel) == 1)
+			{
+				delete this;
+			}
+		};
+	}
+
+	TRefCountPtr<const FStore> FStoreRefCount::Pin()
+	{
+		bool bSucceeded = false;
+
+		// This call to UE_AUTORTFM_OPEN is matched by a second UE_AUTORTFM_OPEN before we leave the function
+		// that undoes the memory change, so we do not need a matching UE_AUTORTFM_ONABORT.
+		UE_AUTORTFM_OPEN 
+		{
+			// See SharedPointerInternals::TReferenceControllerBase<Mode>::ConditionallyAddSharedReference
+			// for commenting about the std::atomic implementation details of this loop.
+			int32 OriginalCount = StrongRefCount.load(std::memory_order_relaxed);
+			for (;;)
+			{
+				if (OriginalCount == 0)
+				{
+					// Never add a shared reference if the pointer has already expired
+					break;
+				}
+				if (StrongRefCount.compare_exchange_weak(OriginalCount, OriginalCount + 1, std::memory_order_relaxed))
+				{
+					bSucceeded = true;
+					break;
+				}
+			}
+		};
+
+		if (bSucceeded)
+		{
+			// Now that we have a ref count, we can safely access the object. Create a TRefCountPtr based on it,
+			// which increments the refcount again, and then drop the refcount we added; we only needed it to
+			// be able to thread-safely access our Owner pointer.
+			TRefCountPtr<const FStore> Result(Owner);
+
+			UE_AUTORTFM_OPEN
+			{
+				StrongRefCount.fetch_add(-1, std::memory_order_relaxed);
+			};
+			return Result;
+		}
+		else
+		{
+			return TRefCountPtr<const FStore>();
+		}
+	}
+
+	FWeakStorePtr::FWeakStorePtr(const TRefCountPtr<const FStore>& Store)
+	{
+		if (Store)
+		{
+			RefCountObject = Store->RefCountObject;
+			RefCountObject->WeakAddRef();
+		}
+	}
+
+	FWeakStorePtr::FWeakStorePtr(const FWeakStorePtr& Other)
+	{
+		if (Other.RefCountObject)
+		{
+			RefCountObject = Other.RefCountObject;
+			RefCountObject->WeakAddRef();
+		}
+	}
+
+	FWeakStorePtr::FWeakStorePtr(FWeakStorePtr&& Other)
+	{
+		if (Other.RefCountObject)
+		{
+			RefCountObject = Other.RefCountObject;
+			Other.RefCountObject = nullptr;
+		}
+	}
+
+	FWeakStorePtr& FWeakStorePtr::operator=(const FWeakStorePtr& Other)
+	{
+		if (&Other == this)
+		{
+			return *this;
+		}
+		FStoreRefCount* OldObject = RefCountObject;
+		RefCountObject = Other.RefCountObject;
+		if (RefCountObject)
+		{
+			RefCountObject->WeakAddRef();
+		}
+		if (OldObject)
+		{
+			OldObject->WeakRelease();
+		}
+		return *this;
+	}
+
+	FWeakStorePtr& FWeakStorePtr::operator=(FWeakStorePtr&& Other)
+	{
+		if (&Other == this)
+		{
+			return *this;
+		}
+		FStoreRefCount* OldObject = RefCountObject;
+		RefCountObject = Other.RefCountObject;
+		if (RefCountObject)
+		{
+			Other.RefCountObject = nullptr;
+		}
+		if (OldObject)
+		{
+			OldObject->WeakRelease();
+		}
+		return *this;
+	}
+
+	FWeakStorePtr::~FWeakStorePtr()
+	{
+		if (RefCountObject)
+		{
+			RefCountObject->WeakRelease();
+		}
+	}
+
+	TRefCountPtr<const FStore> FWeakStorePtr::Pin() const
+	{
+		if (RefCountObject)
+		{
+			TRefCountPtr<const FStore> StrongPtr = RefCountObject->Pin();
+			if (!StrongPtr)
+			{
+				RefCountObject->WeakRelease();
+				RefCountObject = nullptr;
+			}
+			return StrongPtr;
+		}
+		else
+		{
+			return TRefCountPtr<const FStore>();
+		}
+	}
+
 	//////////////////////////////////////////////////////////////////////////
 	
 	template<class T>
@@ -1106,7 +1298,13 @@ namespace FixedTagPrivate
 			Ar << Item;
 			return Item;
 		}
-		
+
+		template<class T>
+		void SkipItem()
+		{
+			(void) LoadItem<T>();
+		}
+
 		void SaveItem(FDisplayNameEntryId NameId)
 		{
 			SaveItem(MakeNumberedName(NameId));
@@ -1116,7 +1314,12 @@ namespace FixedTagPrivate
 		{
 			return FDisplayNameEntryId(LoadItem<FName>());
 		}
-		
+
+		template<> void SkipItem<FDisplayNameEntryId>()
+		{
+			SkipItem<FName>();
+		}
+
 		void SaveItem(FNumberlessExportPath Path)
 		{
 			SaveItem(MakeNumberedPath(Path));
@@ -1140,7 +1343,12 @@ namespace FixedTagPrivate
 		{
 			return MakeNumberlessPath(LoadItem<FAssetRegistryExportPath>());
 		}
-		
+
+		template<> void SkipItem<FNumberlessExportPath>()
+		{
+			SkipItem<FAssetRegistryExportPath>();
+		}
+
 		void SaveItem(FValueId Value)
 		{
 			SaveItem(Value.ToInt());
@@ -1195,7 +1403,21 @@ namespace FixedTagPrivate
 				return FMarshalledText(ComplexString);
 			}
 		}
-		
+
+		template<> void SkipItem<FMarshalledText>()
+		{
+			if constexpr (Version >= FAssetRegistryVersion::MarshalledTextAsUTF8String)
+			{
+				FUtf8String ComplexString;
+				Ar << ComplexString;
+			}
+			else
+			{
+				FString ComplexString;
+				Ar << ComplexString;
+			}
+		}
+
 		template<typename T>
 		void SaveViewData(TArrayView<T> Items)
 		{
@@ -1227,6 +1449,20 @@ namespace FixedTagPrivate
 				}
 			}
 		}
+		template <typename T>
+		void SkipViewData(TArrayView<T> Items)
+		{
+			// If this function is too slow, we can change the FStore serialized format to include byte lengths for
+			// each view.
+			if (!Ar.IsError())
+			{
+				int32 NumItems = Items.Num();
+				for (int32 Index = 0; Index < NumItems; ++Index)
+				{
+					SkipItem<T>();
+				}
+			}
+		}
 
 		template<typename T>
 		void BulkSerialize(TArrayView<T> View)
@@ -1239,6 +1475,9 @@ namespace FixedTagPrivate
 		void LoadViewData(TArrayView<ANSICHAR> View) { BulkSerialize(View); }
 		void LoadViewData(TArrayView<WIDECHAR> View) { BulkSerialize(View); }
 		static_assert(PLATFORM_LITTLE_ENDIAN, "Byte-swapping WIDECHAR on load needed to load on big-endian platforms");
+		void SkipViewData(TArrayView<WIDECHAR> View) { Ar.Seek(Ar.Tell() + GetBytes(View)); }
+		void SkipViewData(TArrayView<ANSICHAR> View) { Ar.Seek(Ar.Tell() + GetBytes(View)); }
+		void SkipViewData(TArrayView<uint32> View) { Ar.Seek(Ar.Tell() + GetBytes(View)); };
 
 		template<class T>
 		inline void LoadViewFromMemoryWithoutAttemptingToMap(FStore& Store, TArrayView<T>& View)
@@ -1440,6 +1679,29 @@ namespace FixedTagPrivate
 			}
 		}
 
+		void SkipBytes()
+		{
+			TRefCountPtr<FStore> PlaceholderStore = GStores.CreateSkipBytesStore();
+			if (TOptional<ELoadOrder> Order = LoadHeader(/* Out */ *PlaceholderStore))
+			{
+				// Load view data
+				checkf(Order.GetValue() == ELoadOrder::TextFirst, TEXT("SkipBytes of old versions is not supported."));
+				uint32 TextDataBytes = LoadItem<uint32>();
+
+				VisitViews<EOrder::TextFirst>(*PlaceholderStore, [&](auto& View)
+					{
+						AlignPosInArchive();
+						SkipViewData(View);
+					});
+
+				if (LoadItem<uint32>() != EndMagic)
+				{
+					UE_LOG(LogAssetDataTags, Warning, TEXT("Bad end magic, archive '%s' is corrupt"), *Ar.GetArchiveName());
+					Ar.SetError();
+				}
+			}
+		}
+
 		TArray<uint8> ReadTextData()
 		{
 			uint32 TextDataBytes = LoadItem<uint32>();
@@ -1512,6 +1774,17 @@ namespace FixedTagPrivate
 		return TRefCountPtr<const FStore>(Store);
 	}
 
+	void SkipBytesStore(FArchive& Ar, FAssetRegistryVersion::Type Version /*= FAssetRegistryVersion::LatestVersion*/)
+	{
+		if (Ar.IsError())
+		{
+			return;
+		}
+
+		// We only support SkipBytes for the latest version.
+		TSerializer<FAssetRegistryVersion::LatestVersion>(Ar).SkipBytes();
+	}
+
 	FAsyncStoreLoader::FAsyncStoreLoader()
 		: Store(GStores.CreateAndRegister())
 	{}
diff --git a/Engine/Source/Runtime/CoreUObject/Public/AssetRegistry/AssetDataTagMapSerializationDetails.h b/Engine/Source/Runtime/CoreUObject/Public/AssetRegistry/AssetDataTagMapSerializationDetails.h
--- a/Engine/Source/Runtime/CoreUObject/Public/AssetRegistry/AssetDataTagMapSerializationDetails.h
+++ b/Engine/Source/Runtime/CoreUObject/Public/AssetRegistry/AssetDataTagMapSerializationDetails.h
@@ -4,11 +4,15 @@
 
 #include "AssetRegistry/AssetData.h"
 #include "AssetRegistry/AssetDataTagMap.h"
-#include "AssetRegistrySerializationDetails.h"
+#include "AssetRegistry/AssetRegistrySerializationDetails.h"
 #include "Async/Async.h"
 #include "AutoRTFM.h"
 #include "HAL/ThreadSafeCounter.h"
+#include "Misc/CoreMiscDefines.h"
 
+namespace FixedTagPrivate { class FStoreManager; }
+namespace FixedTagPrivate { class FWeakStorePtr; }
+namespace FixedTagPrivate { struct FStore; }
 struct FAssetRegistrySerializationOptions;
 
 namespace FixedTagPrivate
@@ -44,6 +48,51 @@ namespace FixedTagPrivate
 		FUtf8String String;
 	};
 
+	/**
+	 * A separate allocation used by FStore that stores the normal (aka strong) refcount alongside the weak refcount,
+	 * similar to the object used in TSharedPtr. When the strong refcount goes to zero, the FStore is deleted, but
+	 * this FStoreRefCount lives on to provide notification to any weak pointers, until all weak pointers
+	 * are dropped as well and this FStoreRefCount is deleted.
+	 * 
+	 * We use this custom implementation of WeakPointers rather than the general TSharedPtr/TWeakPtr solution for two
+	 * reasons:
+	 *    1. Independent AddRef+Release: We have references to FStores that require refcounting that come from
+	 *       16-bit FAssetDataTagMapSharedView.StoreIndex that we convert to a pointer via lookup in GStores. We need
+	 *       to be able to call AddRef and Release during constructor/destructor without having to hold a 64-bit
+	 *       TSharedPtr, to save memory on the hundreds of thousands of FAssetDataTagMapSharedView.
+	 *    2. Implementation hiding: We want to make FWeakStorePtr a public type to an implementation-hidden forward
+	 *       declare of FStore, and TWeakPtr<FStore> would require a public define of FStore.
+	 */
+	struct FStoreRefCount
+	{
+	public:
+		UE_NONCOPYABLE(FStoreRefCount);
+		FStoreRefCount(FStore* InOwner);
+
+		void StrongAddRef()
+		{
+			UE_AUTORTFM_OPEN
+			{
+				StrongRefCount.fetch_add(1, std::memory_order_relaxed);
+			};
+
+			UE_AUTORTFM_ONABORT(this)
+			{
+				StrongRefCount.fetch_add(-1, std::memory_order_relaxed);
+			};
+		}
+
+		COREUOBJECT_API void StrongRelease();
+		void WeakAddRef();
+		void WeakRelease();
+		TRefCountPtr<const FStore> Pin();
+
+	private:
+		std::atomic<int> StrongRefCount{ 0 };
+		std::atomic<int> WeakRefCount{ 0 };
+		FStore* Owner = nullptr;
+	};
+
 	/// Stores a fixed set of values and all the key-values maps used for lookup
 	struct FStore
 	{
@@ -76,33 +125,26 @@ namespace FixedTagPrivate
 
 		void AddRef() const
 		{
-			UE_AUTORTFM_OPEN
-			{
-				RefCount.Increment();
-			};
-
-			UE_AUTORTFM_ONABORT(this)
-			{
-				RefCount.Decrement();
-			};
+			RefCountObject->StrongAddRef();
 		}
 
-		COREUOBJECT_API void Release() const;
+		void Release() const
+		{
+			RefCountObject->StrongRelease();
+		}
 		
 		const ANSICHAR* GetAnsiString(uint32 Idx) const { return &AnsiStrings[AnsiStringOffsets[Idx]]; }
 		const WIDECHAR* GetWideString(uint32 Idx) const { return &WideStrings[WideStringOffsets[Idx]]; }
 
 	private:
-		friend class FStoreManager;
-
-		explicit FStore(uint32 InIndex) : Index(InIndex)
-		{
-			checkf(!AutoRTFM::IsClosed() || !AutoRTFM::IsOnCurrentTransactionStack(this), TEXT("Not allowed to construct a stack local within a transaction."));
-		}
+		friend FStoreManager;
+		friend FStoreRefCount;
+		friend FWeakStorePtr;
 
+		COREUOBJECT_API explicit FStore(uint32 InIndex);
 		~FStore();
 
-		mutable FThreadSafeCounter RefCount;
+		FStoreRefCount* RefCountObject = nullptr;
 
 		/** Guards access to the View pointers. No need for an R/O as we're only accessing the set when we're about to modify it. */
 		FCriticalSection ViewPointersCS;
@@ -271,6 +313,8 @@ namespace FixedTagPrivate
 
 	COREUOBJECT_API void SaveStore(const FStoreData& Store, FArchive& Ar);
 	COREUOBJECT_API TRefCountPtr<const FStore> LoadStore(FArchive& Ar, FAssetRegistryVersion::Type Version = FAssetRegistryVersion::LatestVersion, TSharedPtr<UE::AssetRegistry::FMemoryMappedFile> InBackingFile = nullptr);
+	/** Advance the offset of Ar to skip over the bytes for an FStore that starts at its current offset. */
+	COREUOBJECT_API void SkipBytesStore(FArchive& Ar, FAssetRegistryVersion::Type Version = FAssetRegistryVersion::LatestVersion);
 
 	/// Loads tag store with async creation of expensive tag values
 	///
diff --git a/Engine/Source/Runtime/CoreUObject/Public/AssetRegistry/AssetRegistryStateSerialization.h b/Engine/Source/Runtime/CoreUObject/Public/AssetRegistry/AssetRegistryStateSerialization.h
new file mode 100644
--- /dev/null
+++ b/Engine/Source/Runtime/CoreUObject/Public/AssetRegistry/AssetRegistryStateSerialization.h
@@ -0,0 +1,35 @@
+// Copyright Epic Games, Inc. All Rights Reserved.
+
+#pragma once
+
+#include "CoreTypes.h"
+
+namespace FixedTagPrivate { struct FStore; }
+namespace FixedTagPrivate { struct FStoreRefCount; }
+template<typename T> class TRefCountPtr;
+
+namespace FixedTagPrivate
+{
+
+/**
+ * A weak pointer to an FStore. When the normal refcount (e.g. via TRefCountPtr) on the FStore goes to zero, the
+ * FStore is deleted and calls to Pin on this weak pointer will return null. FStore is a private implementation
+ * detail; this class is public only to pass the opaque type between different Load functions.
+ */
+class FWeakStorePtr
+{
+public:
+	inline FWeakStorePtr() = default;
+	COREUOBJECT_API FWeakStorePtr(const TRefCountPtr<const FStore>& Store);
+	COREUOBJECT_API FWeakStorePtr(const FWeakStorePtr& Other);
+	COREUOBJECT_API FWeakStorePtr(FWeakStorePtr&& Other);
+	COREUOBJECT_API FWeakStorePtr& operator=(const FWeakStorePtr& Other);
+	COREUOBJECT_API FWeakStorePtr& operator=(FWeakStorePtr&& Other);
+	COREUOBJECT_API ~FWeakStorePtr();
+
+	COREUOBJECT_API TRefCountPtr<const FStore> Pin() const;
+private:
+	mutable FStoreRefCount* RefCountObject = nullptr;
+};
+
+} // namespace FixedTagPrivate
