[CoreUObject] Add GetName() accessor to FObjectInfo, to allow accessing the Name field.
This supports reference chain search functionality that needs to query object names during traversal.

#Synced-CL 50640962
#Synced-CL 50640962

#Synced-CL 50640962

diff --git a/Engine/Source/Runtime/CoreUObject/Public/UObject/GCObjectInfo.h b/Engine/Source/Runtime/CoreUObject/Public/UObject/GCObjectInfo.h
--- a/Engine/Source/Runtime/CoreUObject/Public/UObject/GCObjectInfo.h
+++ b/Engine/Source/Runtime/CoreUObject/Public/UObject/GCObjectInfo.h
@@ -67,6 +67,11 @@ public:
 	 */
 	COREUOBJECT_API UObject* TryResolveObject();
 
+	FName GetName() const
+	{
+		return Name;
+	}
+
 	FGCObjectInfo* GetClass() const
 	{
 		return Class;
