Fixed printf.

diff --git a/Engine/Source/Runtime/CoreUObject/Private/UObject/RemoteHeap.cpp b/Engine/Source/Runtime/CoreUObject/Private/UObject/RemoteHeap.cpp
--- a/Engine/Source/Runtime/CoreUObject/Private/UObject/RemoteHeap.cpp
+++ b/Engine/Source/Runtime/CoreUObject/Private/UObject/RemoteHeap.cpp
@@ -230,7 +230,7 @@ void FRemoteHeap::TickSubsystem()
 
 							*LineTextCursor++ = TEXT('\0');
 
-							UE_LOG(LogRemoteHeap, VeryVerbose, TEXT("TickSubsystem sending line %llx to %s %s | Line data: s"),
+							UE_LOG(LogRemoteHeap, VeryVerbose, TEXT("TickSubsystem sending line %llx to %s %s | Line data: %s"),
 								HighestPriorityRequest->Address,
 								*HighestPriorityRequest->DestinationServerId.ToString(),
 								*HighestPriorityRequest->RequestPriority.ToString(),
