Fix crash on copy/paste of prefabs
Changed FObjectPropertyBase::ExportText_Internal to only export using the Name if Temp is a subobject of Parent or a subobject of an object in Parent's outerchain. Otherwise we export with the path to avoid resolving to the wrong object on load.

#Synced-CL 50652007

diff --git a/Engine/Source/Runtime/CoreUObject/Private/UObject/PropertyBaseObject.cpp b/Engine/Source/Runtime/CoreUObject/Private/UObject/PropertyBaseObject.cpp
--- a/Engine/Source/Runtime/CoreUObject/Private/UObject/PropertyBaseObject.cpp
+++ b/Engine/Source/Runtime/CoreUObject/Private/UObject/PropertyBaseObject.cpp
@@ -305,8 +305,11 @@ void FObjectPropertyBase::ExportText_Internal(FString& ValueStr, TNotNull<const
 				 // be that it's inside the package we are currently in, which means the Temp pointer should be resolved 
 				 // already.  So don't move forward with the check unless it's resolved, we don't want to force a deferred
 				 // loaded asset if we don't have to with this check.
-		         // We also want to make sure the object is actually inside the package we are currently in
-				 && Temp.IsResolved() && Temp && Temp->IsDefaultSubobject() && Temp->IsIn(Parent->GetOutermostObject())
+				 && Temp.IsResolved() && Temp && Temp->IsDefaultSubobject() 
+				 // Use short name only if import will find it by searching Parent's outer chain.
+				 // Import searches upward for objects with matching name,
+				 // so Temp->GetOuter() must be Parent or an ancestor of Parent to avoid ambiguity.
+				 && (Temp->GetOuter() && (Temp->GetOuter() == Parent || Parent->IsInOuter(Temp->GetOuter())))
 				)
 		{
 			if (PortFlags & PPF_Delimited)
