Avoid copying of remote object data when queueing it up for transactional PostMigrate

#Synced-CL 50595067

diff --git a/Engine/Source/Runtime/CoreUObject/Private/UObject/RemoteObjectTransfer.cpp b/Engine/Source/Runtime/CoreUObject/Private/UObject/RemoteObjectTransfer.cpp
--- a/Engine/Source/Runtime/CoreUObject/Private/UObject/RemoteObjectTransfer.cpp
+++ b/Engine/Source/Runtime/CoreUObject/Private/UObject/RemoteObjectTransfer.cpp
@@ -51,7 +51,7 @@ static FAutoConsoleVariableRef CVarTransactionalMigrationEnabled(
 * @param MigrationContext Migration context
 * @returns Structure that holds the migrated data and the migration context
 */
-TRefCountPtr<FTransactionalMigrationData> StoreMigratedObjectData(FRemoteObjectData& ObjectData, const FUObjectMigrationContext& MigrationContext);
+TRefCountPtr<FTransactionalMigrationData> StoreMigratedObjectData(FRemoteObjectData& ObjectData, FUObjectMigrationContext& MigrationContext);
 
 /**
 * Removes the specified migrated data from the stubs of objects serialized inside of the migrated data
@@ -852,7 +852,7 @@ public:
 	* @param ObjectId Id of the object that was requested
 	* @param Data Object data. Data ownership is transferred to the receive request (if it exists)
 	*/
-	FORCENOINLINE void FulfillReceiveRequest(FRemoteObjectData& ObjectData, const FUObjectMigrationContext& MigrationContext)
+	FORCENOINLINE void FulfillReceiveRequest(FRemoteObjectData& ObjectData, FUObjectMigrationContext& MigrationContext)
 	{
 		using namespace UE::RemoteObject::Serialization;
 		using namespace UE::RemoteObject::Private;
@@ -934,7 +934,7 @@ public:
 
 FRemoteObjectTransferQueue GTransferQueue;
 
-TRefCountPtr<FTransactionalMigrationData> StoreMigratedObjectData(FRemoteObjectData& ObjectData, const FUObjectMigrationContext& MigrationContext)
+TRefCountPtr<FTransactionalMigrationData> StoreMigratedObjectData(FRemoteObjectData& ObjectData, FUObjectMigrationContext& MigrationContext)
 {
 	using namespace UE::RemoteObject::Serialization;
 	using namespace UE::RemoteObject::Private;
@@ -944,8 +944,8 @@ TRefCountPtr<FTransactionalMigrationData> StoreMigratedObjectData(FRemoteObjectD
 	FRemoteObjectStub* RootStub = FindOrAddRemoteObjectStub(MigrationContext.ObjectId, FRemoteServerId::GetLocalServerId());
 
 	TRefCountPtr<FTransactionalMigrationData> MigratedData = new FTransactionalMigrationData();
-	MigratedData->Data = ObjectData;
-	MigratedData->Context = MigrationContext;
+	MigratedData->Data = MoveTemp(ObjectData);
+	MigratedData->Context = MoveTemp(MigrationContext);
 
 	UE_LOG(LogRemoteObject, VeryVerbose, TEXT("%s received %d objects"), *FRemoteServerId::GetLocalServerId().ToString(), MigratedData->Data.SerializedObjectHeaders.Num());
 
