Move members in UClass to save 16 bytes (392 -> 376)
#rnx
#Synced-CL 50851879
#Synced-CL 50851879

diff --git a/Engine/Source/Runtime/CoreUObject/Public/UObject/Class.h b/Engine/Source/Runtime/CoreUObject/Public/UObject/Class.h
--- a/Engine/Source/Runtime/CoreUObject/Public/UObject/Class.h
+++ b/Engine/Source/Runtime/CoreUObject/Public/UObject/Class.h
@@ -3894,6 +3894,13 @@ public:
 
 	/** Used to check if the class layout is currently changing and therefore is not ready for a CDO to be created */
 	bool bLayoutChanging;
+	
+protected:
+	
+	/** TRUE if this class should not rely exclusively on instanced property semantics for subobject template instancing. */
+	bool bNeedsDynamicSubobjectInstancing = false;
+	
+public:
 
 	/** Class flags; See EClassFlags for more information */
 	EClassFlags ClassFlags;
@@ -4021,9 +4028,6 @@ protected:
 	/** The struct used to store sparse class data. */
 	TObjectPtr<UScriptStruct> SparseClassDataStruct;
 
-	/** TRUE if this class should not rely exclusively on instanced property semantics for subobject template instancing. */
-	bool bNeedsDynamicSubobjectInstancing = false;
-
 public:
 	/**
 	 * Returns a pointer to the sidecar data structure, based on the EGetSparseClassDataMethod.
@@ -4076,13 +4080,13 @@ private:
 
 	/** Scope lock to avoid the FuncMap being read and written to simultaneously on multiple threads. */
 	mutable FUClassFuncLock FuncMapLock;
+	
+	/** Scope lock to avoid the SuperFuncMap being read and written to simultaneously on multiple threads. */
+	mutable FUClassFuncLock AllFunctionsCacheLock;
 
 	/** A cache of all functions by name that exist in this class or a parent (superclass or interface) context */
 	mutable TMap<FName, UFunction*> AllFunctionsCache;
 
-	/** Scope lock to avoid the SuperFuncMap being read and written to simultaneously on multiple threads. */
-	mutable FUClassFuncLock AllFunctionsCacheLock;
-
 public:
 	union
 	{ 
@@ -4914,8 +4918,8 @@ consteval UClass::UClass(
 #endif
 	, FuncMap(ConstEval)
 	, FuncMapLock()
-	, AllFunctionsCache(ConstEval)
 	, AllFunctionsCacheLock()
+	, AllFunctionsCache(ConstEval)
 	, CompiledInInterfaces(Params.Interfaces)
 	, CompiledInNativeFunctions(Params.NativeFunctions)
 {
