Added new UE::CDynamicDelegate, UE::CDynamicMulticastDelegate and UE::CSparseDynamicMulticastDelegate concepts.
Added tests for the first two.

diff --git a/Engine/Source/Runtime/Core/Public/Concepts/DynamicDelegate.h b/Engine/Source/Runtime/Core/Public/Concepts/DynamicDelegate.h
new file mode 100644
--- /dev/null
+++ b/Engine/Source/Runtime/Core/Public/Concepts/DynamicDelegate.h
@@ -0,0 +1,30 @@
+// Copyright Epic Games, Inc. All Rights Reserved.
+
+#pragma once
+
+#include "Concepts/CompleteType.h"
+#include "Concepts/DerivedFrom.h"
+
+template <typename FuncType, typename ThreadSafetyMode>
+class TDynamicDelegate;
+
+namespace UE::Private::DynamicDelegate
+{
+	template <typename FuncType, typename ThreadSafetyMode>
+	void Resolve(const volatile TDynamicDelegate<FuncType, ThreadSafetyMode>*);
+
+	void Resolve(...) = delete;
+}
+
+namespace UE
+{
+	/**
+	 * Concept which describes a dynamic delegate.
+	 */
+	template <typename T>
+	concept CDynamicDelegate = UE::CCompleteType<T> &&
+		requires (T* Ptr)
+		{
+			UE::Private::DynamicDelegate::Resolve(Ptr);
+		};
+}
diff --git a/Engine/Source/Runtime/Core/Public/Concepts/DynamicMulticastDelegate.h b/Engine/Source/Runtime/Core/Public/Concepts/DynamicMulticastDelegate.h
new file mode 100644
--- /dev/null
+++ b/Engine/Source/Runtime/Core/Public/Concepts/DynamicMulticastDelegate.h
@@ -0,0 +1,30 @@
+// Copyright Epic Games, Inc. All Rights Reserved.
+
+#pragma once
+
+#include "Concepts/CompleteType.h"
+#include "Concepts/DerivedFrom.h"
+
+template <typename FuncType, typename ThreadSafetyMode>
+class TDynamicMulticastDelegate;
+
+namespace UE::Private::DynamicMulticastDelegate
+{
+	template <typename FuncType, typename ThreadSafetyMode>
+	void Resolve(const volatile TDynamicMulticastDelegate<FuncType, ThreadSafetyMode>*);
+
+	void Resolve(...) = delete;
+}
+
+namespace UE
+{
+	/**
+	 * Concept which describes a dynamic multicast delegate.
+	 */
+	template <typename T>
+	concept CDynamicMulticastDelegate = UE::CCompleteType<T> &&
+		requires (T* Ptr)
+		{
+			UE::Private::DynamicMulticastDelegate::Resolve(Ptr);
+		};
+}
diff --git a/Engine/Source/Runtime/Core/Public/Concepts/SparseDynamicMulticastDelegate.h b/Engine/Source/Runtime/Core/Public/Concepts/SparseDynamicMulticastDelegate.h
new file mode 100644
--- /dev/null
+++ b/Engine/Source/Runtime/Core/Public/Concepts/SparseDynamicMulticastDelegate.h
@@ -0,0 +1,30 @@
+// Copyright Epic Games, Inc. All Rights Reserved.
+
+#pragma once
+
+#include "Concepts/CompleteType.h"
+#include "Concepts/DerivedFrom.h"
+
+template <typename MulticastDelegate, typename OwningClass, typename DelegateInfoClass>
+struct TSparseDynamicDelegate;
+
+namespace UE::Private::SparseDynamicMulticastDelegate
+{
+	template <typename MulticastDelegate, typename OwningClass, typename DelegateInfoClass>
+	void Resolve(const volatile TSparseDynamicDelegate<MulticastDelegate, OwningClass, DelegateInfoClass>*);
+
+	void Resolve(...) = delete;
+}
+
+namespace UE
+{
+	/**
+	 * Concept which describes a sparse dynamic multicast delegate.
+	 */
+	template <typename T>
+	concept CSparseDynamicMulticastDelegate = UE::CCompleteType<T> &&
+		requires (T* Ptr)
+		{
+			UE::Private::SparseDynamicMulticastDelegate::Resolve(Ptr);
+		};
+}
diff --git a/Engine/Source/Runtime/CoreUObject/Public/UObject/DynamicDelegateTestObject.cpp b/Engine/Source/Runtime/CoreUObject/Public/UObject/DynamicDelegateTestObject.cpp
--- a/Engine/Source/Runtime/CoreUObject/Public/UObject/DynamicDelegateTestObject.cpp
+++ b/Engine/Source/Runtime/CoreUObject/Public/UObject/DynamicDelegateTestObject.cpp
@@ -4,6 +4,8 @@
 
 #include "DynamicDelegateTestObject.h"
 
+#include "Concepts/DynamicDelegate.h"
+#include "Concepts/DynamicMulticastDelegate.h"
 #include "Containers/AnsiString.h"
 #include "Tests/TestHarnessAdapter.h"
 
@@ -985,6 +987,84 @@ TEST_CASE_NAMED(DynamicDelegateTest, "UE::CoreUObject::DynamicDelegateTest", "[A
 	TestObj->RunTests();
 }
 
+static_assert(UE::CDynamicDelegate<FTestDynamicDelegate0>);
+static_assert(UE::CDynamicDelegate<FTestDynamicDelegate1>);
+static_assert(UE::CDynamicDelegate<FTestDynamicDelegate2>);
+static_assert(UE::CDynamicDelegate<FTestDynamicDelegate3>);
+static_assert(UE::CDynamicDelegate<FTestDynamicDelegate4>);
+static_assert(UE::CDynamicDelegate<FTestDynamicDelegate5>);
+static_assert(UE::CDynamicDelegate<FTestDynamicDelegate6>);
+static_assert(UE::CDynamicDelegate<FTestDynamicDelegateRetInt0>);
+static_assert(UE::CDynamicDelegate<FTestDynamicDelegateRetInt1>);
+static_assert(UE::CDynamicDelegate<FTestDynamicDelegateRetInt2>);
+static_assert(UE::CDynamicDelegate<FTestDynamicDelegateRetInt3>);
+static_assert(UE::CDynamicDelegate<FTestDynamicDelegateRetInt4>);
+static_assert(UE::CDynamicDelegate<FTestDynamicDelegateRetInt5>);
+static_assert(UE::CDynamicDelegate<FTestDynamicDelegateRetInt6>);
+static_assert(UE::CDynamicDelegate<FTestDynamicDelegateRetStr0>);
+static_assert(UE::CDynamicDelegate<FTestDynamicDelegateRetStr1>);
+static_assert(UE::CDynamicDelegate<FTestDynamicDelegateRetStr2>);
+static_assert(UE::CDynamicDelegate<FTestDynamicDelegateRetStr3>);
+static_assert(UE::CDynamicDelegate<FTestDynamicDelegateRetStr4>);
+static_assert(UE::CDynamicDelegate<FTestDynamicDelegateRetStr5>);
+static_assert(UE::CDynamicDelegate<FTestDynamicDelegateRetStr6>);
+static_assert(UE::CDynamicDelegate<FTestDynamicDelegateOut1>);
+static_assert(UE::CDynamicDelegate<FTestDynamicDelegateOut2>);
+static_assert(UE::CDynamicDelegate<FTestDynamicDelegateOut3>);
+static_assert(UE::CDynamicDelegate<FTestDynamicDelegateOut4>);
+static_assert(UE::CDynamicDelegate<FTestDynamicDelegateOut5>);
+static_assert(!UE::CDynamicMulticastDelegate<FTestDynamicDelegate0>);
+static_assert(!UE::CDynamicMulticastDelegate<FTestDynamicDelegate1>);
+static_assert(!UE::CDynamicMulticastDelegate<FTestDynamicDelegate2>);
+static_assert(!UE::CDynamicMulticastDelegate<FTestDynamicDelegate3>);
+static_assert(!UE::CDynamicMulticastDelegate<FTestDynamicDelegate4>);
+static_assert(!UE::CDynamicMulticastDelegate<FTestDynamicDelegate5>);
+static_assert(!UE::CDynamicMulticastDelegate<FTestDynamicDelegate6>);
+static_assert(!UE::CDynamicMulticastDelegate<FTestDynamicDelegateRetInt0>);
+static_assert(!UE::CDynamicMulticastDelegate<FTestDynamicDelegateRetInt1>);
+static_assert(!UE::CDynamicMulticastDelegate<FTestDynamicDelegateRetInt2>);
+static_assert(!UE::CDynamicMulticastDelegate<FTestDynamicDelegateRetInt3>);
+static_assert(!UE::CDynamicMulticastDelegate<FTestDynamicDelegateRetInt4>);
+static_assert(!UE::CDynamicMulticastDelegate<FTestDynamicDelegateRetInt5>);
+static_assert(!UE::CDynamicMulticastDelegate<FTestDynamicDelegateRetInt6>);
+static_assert(!UE::CDynamicMulticastDelegate<FTestDynamicDelegateRetStr0>);
+static_assert(!UE::CDynamicMulticastDelegate<FTestDynamicDelegateRetStr1>);
+static_assert(!UE::CDynamicMulticastDelegate<FTestDynamicDelegateRetStr2>);
+static_assert(!UE::CDynamicMulticastDelegate<FTestDynamicDelegateRetStr3>);
+static_assert(!UE::CDynamicMulticastDelegate<FTestDynamicDelegateRetStr4>);
+static_assert(!UE::CDynamicMulticastDelegate<FTestDynamicDelegateRetStr5>);
+static_assert(!UE::CDynamicMulticastDelegate<FTestDynamicDelegateRetStr6>);
+static_assert(!UE::CDynamicMulticastDelegate<FTestDynamicDelegateOut1>);
+static_assert(!UE::CDynamicMulticastDelegate<FTestDynamicDelegateOut2>);
+static_assert(!UE::CDynamicMulticastDelegate<FTestDynamicDelegateOut3>);
+static_assert(!UE::CDynamicMulticastDelegate<FTestDynamicDelegateOut4>);
+static_assert(!UE::CDynamicMulticastDelegate<FTestDynamicDelegateOut5>);
+
+static_assert(UE::CDynamicMulticastDelegate<FTestDynamicMulticastDelegate0>);
+static_assert(UE::CDynamicMulticastDelegate<FTestDynamicMulticastDelegate1>);
+static_assert(UE::CDynamicMulticastDelegate<FTestDynamicMulticastDelegate2>);
+static_assert(UE::CDynamicMulticastDelegate<FTestDynamicMulticastDelegate3>);
+static_assert(UE::CDynamicMulticastDelegate<FTestDynamicMulticastDelegate4>);
+static_assert(UE::CDynamicMulticastDelegate<FTestDynamicMulticastDelegate5>);
+static_assert(UE::CDynamicMulticastDelegate<FTestDynamicMulticastDelegate6>);
+static_assert(UE::CDynamicMulticastDelegate<FTestDynamicMulticastDelegateOut1>);
+static_assert(UE::CDynamicMulticastDelegate<FTestDynamicMulticastDelegateOut2>);
+static_assert(UE::CDynamicMulticastDelegate<FTestDynamicMulticastDelegateOut3>);
+static_assert(UE::CDynamicMulticastDelegate<FTestDynamicMulticastDelegateOut4>);
+static_assert(UE::CDynamicMulticastDelegate<FTestDynamicMulticastDelegateOut5>);
+static_assert(!UE::CDynamicDelegate<FTestDynamicMulticastDelegate0>);
+static_assert(!UE::CDynamicDelegate<FTestDynamicMulticastDelegate1>);
+static_assert(!UE::CDynamicDelegate<FTestDynamicMulticastDelegate2>);
+static_assert(!UE::CDynamicDelegate<FTestDynamicMulticastDelegate3>);
+static_assert(!UE::CDynamicDelegate<FTestDynamicMulticastDelegate4>);
+static_assert(!UE::CDynamicDelegate<FTestDynamicMulticastDelegate5>);
+static_assert(!UE::CDynamicDelegate<FTestDynamicMulticastDelegate6>);
+static_assert(!UE::CDynamicDelegate<FTestDynamicMulticastDelegateOut1>);
+static_assert(!UE::CDynamicDelegate<FTestDynamicMulticastDelegateOut2>);
+static_assert(!UE::CDynamicDelegate<FTestDynamicMulticastDelegateOut3>);
+static_assert(!UE::CDynamicDelegate<FTestDynamicMulticastDelegateOut4>);
+static_assert(!UE::CDynamicDelegate<FTestDynamicMulticastDelegateOut5>);
+
 #undef CHECK_ARG_VARIABLES
 #undef EXPAND_ARG_VARIABLES
 #undef DEFINE_ARG_VARIABLES
@@ -1014,4 +1094,4 @@ TEST_CASE_NAMED(DynamicDelegateTest, "UE::CoreUObject::DynamicDelegateTest", "[A
 #undef BIND_CALL_CHECK_DELEGATE_RETVAL
 #undef BIND_CALL_CHECK_DELEGATE
 
-#endif // #if WITH_TESTS
+#endif // #if WITH_TESTS
\ No newline at end of file
