Fix RF_HasDynamicImports to import objects within packages, not just packages.
Extended FDynamicImport structure to support importing specific objects within dynamically imported packages, assuming the caller is the outer of dynamic imports.
Entity system updated to provide full object details for dynamic imports.

#Synced-CL 50904874

diff --git a/Engine/Source/Runtime/CoreUObject/Private/UObject/LinkerLoad.cpp b/Engine/Source/Runtime/CoreUObject/Private/UObject/LinkerLoad.cpp
--- a/Engine/Source/Runtime/CoreUObject/Private/UObject/LinkerLoad.cpp
+++ b/Engine/Source/Runtime/CoreUObject/Private/UObject/LinkerLoad.cpp
@@ -2915,18 +2915,37 @@ FLinkerLoad::ELinkerStatus FLinkerLoad::AddDynamicImports()
 						DynamicImportsIndex = ImportMap.Num();
 					}
 
-					FObjectImport& Import = ImportMap.AddDefaulted_GetRef();
-					Import.ObjectName = DynamicImport.PackageName;
-					Import.ClassPackage = GLongCoreUObjectPackageName;
-					Import.ClassName = NAME_Package;
-					Import.PackageName = NAME_None;
-					Import.SourceIndex = -1;
-					Import.bImportOptional = false;
-					Import.bImportPackageHandled = false;
-					Import.bImportSearchedFor = false;
-					Import.bImportFailed = false;
-					Import.XObject = nullptr;
-					Import.SourceLinker = nullptr;
+					// Create package import
+					FObjectImport& PackageImport = ImportMap.AddDefaulted_GetRef();
+					PackageImport.ObjectName = DynamicImport.PackageName;
+					PackageImport.ClassPackage = GLongCoreUObjectPackageName;
+					PackageImport.ClassName = NAME_Package;
+					PackageImport.PackageName = NAME_None;
+					PackageImport.SourceIndex = -1;
+					PackageImport.bImportOptional = false;
+					PackageImport.bImportPackageHandled = false;
+					PackageImport.bImportSearchedFor = false;
+					PackageImport.bImportFailed = false;
+					PackageImport.XObject = nullptr;
+					PackageImport.SourceLinker = nullptr;
+
+					if (!DynamicImport.ObjectName.IsNone())
+					{
+						// Create object import
+						FObjectImport& ObjectImport = ImportMap.AddDefaulted_GetRef();
+						ObjectImport.ObjectName = DynamicImport.ObjectName;
+						ObjectImport.ClassName = DynamicImport.ClassName;
+						ObjectImport.ClassPackage = DynamicImport.ClassPackage;
+						ObjectImport.OuterIndex = FPackageIndex::FromExport(ExportIndex);
+						ObjectImport.PackageName = DynamicImport.PackageName;
+						ObjectImport.SourceIndex = -1;
+						ObjectImport.bImportOptional = false;
+						ObjectImport.bImportPackageHandled = false;
+						ObjectImport.bImportSearchedFor = false;
+						ObjectImport.bImportFailed = false;
+						ObjectImport.XObject = nullptr;
+						ObjectImport.SourceLinker = nullptr;
+					}
 
 					if (DynamicImport.InstanceName != NAME_None && DynamicImport.InstanceName != DynamicImport.PackageName)
 					{
diff --git a/Engine/Source/Runtime/CoreUObject/Public/UObject/ObjectResource.h b/Engine/Source/Runtime/CoreUObject/Public/UObject/ObjectResource.h
--- a/Engine/Source/Runtime/CoreUObject/Public/UObject/ObjectResource.h
+++ b/Engine/Source/Runtime/CoreUObject/Public/UObject/ObjectResource.h
@@ -643,4 +643,23 @@ struct FDynamicImport
 	 * Optional instance name to rename the package.
 	 */
 	FName InstanceName;
+
+	/**
+	 * Object name within the package (optional).
+	 * If NAME_None, imports the package itself.
+	 * If set, imports a specific object within the package.
+	 */
+	FName ObjectName;
+
+	/**
+	 * Class name for the object being imported (optional).
+	 * Required when ObjectName is set.
+	 */
+	FName ClassName;
+
+	/**
+	 * Class package for the object being imported (optional).
+	 * Required when ObjectName is set.
+	 */
+	FName ClassPackage;
 };
diff --git a/Engine/Source/Runtime/Engine/Private/OverriddenObjectsExternalPackageHelpers.cpp b/Engine/Source/Runtime/Engine/Private/OverriddenObjectsExternalPackageHelpers.cpp
--- a/Engine/Source/Runtime/Engine/Private/OverriddenObjectsExternalPackageHelpers.cpp
+++ b/Engine/Source/Runtime/Engine/Private/OverriddenObjectsExternalPackageHelpers.cpp
@@ -920,6 +920,10 @@ void FExternalObjectsOverrideHelper::ImportExternalOverrides(const FLinkerLoad&
 
 				FDynamicImport Import;
 				Import.PackageName = Asset.PackageName;
+				Import.ObjectName = Asset.AssetName;
+				Import.ClassName = Asset.AssetClassPath.GetAssetName();
+				Import.ClassPackage = Asset.AssetClassPath.GetPackageName();
+
 				if (bIsInstanced)
 				{
 					Import.InstanceName = InstancingContext.RemapPackage(Asset.PackageName);
@@ -931,7 +935,7 @@ void FExternalObjectsOverrideHelper::ImportExternalOverrides(const FLinkerLoad&
 					}
 
 					checkf(!Import.InstanceName.IsNone(), TEXT("Could find not the remapped package name for dynamic import package %s, in memory name %s"), *Asset.PackageName.ToString(), *PackageNameInMemory.ToString());
-				} 
+				}
 
 				OutDynamicImports.Add(Import);
 			}
