Remove unused FObjectImport flags
- bImportPackageHandled
- bImportSearchedFor
- bImportFailed

#Synced-CL 50904874

diff --git a/Engine/Source/Runtime/CoreUObject/Private/UObject/LinkerLoad.cpp b/Engine/Source/Runtime/CoreUObject/Private/UObject/LinkerLoad.cpp
--- a/Engine/Source/Runtime/CoreUObject/Private/UObject/LinkerLoad.cpp
+++ b/Engine/Source/Runtime/CoreUObject/Private/UObject/LinkerLoad.cpp
@@ -2923,9 +2923,6 @@ FLinkerLoad::ELinkerStatus FLinkerLoad::AddDynamicImports()
 					PackageImport.PackageName = NAME_None;
 					PackageImport.SourceIndex = -1;
 					PackageImport.bImportOptional = false;
-					PackageImport.bImportPackageHandled = false;
-					PackageImport.bImportSearchedFor = false;
-					PackageImport.bImportFailed = false;
 					PackageImport.XObject = nullptr;
 					PackageImport.SourceLinker = nullptr;
 
@@ -2940,9 +2937,6 @@ FLinkerLoad::ELinkerStatus FLinkerLoad::AddDynamicImports()
 						ObjectImport.PackageName = DynamicImport.PackageName;
 						ObjectImport.SourceIndex = -1;
 						ObjectImport.bImportOptional = false;
-						ObjectImport.bImportPackageHandled = false;
-						ObjectImport.bImportSearchedFor = false;
-						ObjectImport.bImportFailed = false;
 						ObjectImport.XObject = nullptr;
 						ObjectImport.SourceLinker = nullptr;
 					}
diff --git a/Engine/Source/Runtime/CoreUObject/Private/UObject/LinkerLoadImportBehavior.cpp b/Engine/Source/Runtime/CoreUObject/Private/UObject/LinkerLoadImportBehavior.cpp
--- a/Engine/Source/Runtime/CoreUObject/Private/UObject/LinkerLoadImportBehavior.cpp
+++ b/Engine/Source/Runtime/CoreUObject/Private/UObject/LinkerLoadImportBehavior.cpp
@@ -49,7 +49,7 @@ EImportBehavior FindLoadBehavior(const UClass& Class)
 
 EImportBehavior GetPropertyImportLoadBehavior(const FObjectImport& Import, const FLinkerLoad& LinkerLoad)
 {
-	if (Import.bImportSearchedFor || Import.XObject)
+	if (Import.XObject)
 	{
 		// If it was something that's been searched for, we've already attempted a resolve, might as well use it
 		return EImportBehavior::Eager;
diff --git a/Engine/Source/Runtime/CoreUObject/Private/UObject/ObjectResource.cpp b/Engine/Source/Runtime/CoreUObject/Private/UObject/ObjectResource.cpp
--- a/Engine/Source/Runtime/CoreUObject/Private/UObject/ObjectResource.cpp
+++ b/Engine/Source/Runtime/CoreUObject/Private/UObject/ObjectResource.cpp
@@ -310,9 +310,6 @@ FObjectImport::FObjectImport()
 	: FObjectResource()
 	, SourceIndex(INDEX_NONE)
 	, bImportOptional(false)
-	, bImportPackageHandled(false)
-	, bImportSearchedFor(false)
-	, bImportFailed(false)
 	, XObject(nullptr)
 	, SourceLinker(nullptr)
 {
@@ -324,9 +321,6 @@ FObjectImport::FObjectImport(UObject* InObject)
 	, ClassName(InObject ? InObject->GetClass()->GetFName() : NAME_None)
 	, SourceIndex(INDEX_NONE)
 	, bImportOptional(false)
-	, bImportPackageHandled(false)
-	, bImportSearchedFor(false)
-	, bImportFailed(false)
 	, XObject(InObject)
 	, SourceLinker(nullptr)
 {
@@ -338,9 +332,6 @@ FObjectImport::FObjectImport(UObject* InObject, UClass* InClass)
 	, ClassName((InObject && InClass) ? InClass->GetFName() : NAME_None)
 	, SourceIndex(INDEX_NONE)
 	, bImportOptional(false)
-	, bImportPackageHandled(false)
-	, bImportSearchedFor(false)
-	, bImportFailed(false)
 	, XObject(InObject)
 	, SourceLinker(nullptr)
 {
diff --git a/Engine/Source/Runtime/CoreUObject/Public/UObject/ObjectResource.h b/Engine/Source/Runtime/CoreUObject/Public/UObject/ObjectResource.h
--- a/Engine/Source/Runtime/CoreUObject/Public/UObject/ObjectResource.h
+++ b/Engine/Source/Runtime/CoreUObject/Public/UObject/ObjectResource.h
@@ -472,10 +472,6 @@ struct FObjectImport : public FObjectResource
 	 */
 	bool			bImportOptional;
 
-	bool			bImportPackageHandled;
-	bool			bImportSearchedFor;
-	bool			bImportFailed;
-
 	/**
 	 * The UObject represented by this resource.  Assigned the first time CreateImport is called for this import.
 	 * Transient
