[UserDefinedStruct] Add struct pathnames to user defined struct ensures so build health can better identify the problem.

#Synced-CL 50911512

diff --git a/Engine/Source/Runtime/CoreUObject/Private/StructUtils/UserDefinedStruct.cpp b/Engine/Source/Runtime/CoreUObject/Private/StructUtils/UserDefinedStruct.cpp
--- a/Engine/Source/Runtime/CoreUObject/Private/StructUtils/UserDefinedStruct.cpp
+++ b/Engine/Source/Runtime/CoreUObject/Private/StructUtils/UserDefinedStruct.cpp
@@ -502,7 +502,12 @@ uint32 UUserDefinedStruct::GetUserDefinedStructTypeHash(const void* Src, const U
 
 const uint8* UUserDefinedStruct::GetDefaultInstance() const
 {
-	ensure(DefaultStructInstance.IsValid() && DefaultStructInstance.GetStruct() == this);
+	if (ensureMsgf(DefaultStructInstance.IsValid(), TEXT("DefaultStructInstance is invalid for '%s'."), *GetPathName()))
+	{
+		const UStruct* DefaultStruct = DefaultStructInstance.GetStruct();
+		ensureMsgf(DefaultStruct == this, TEXT("DefaultStructInstance struct mismatch: '%s' vs. '%s'"), *GetPathNameSafe(DefaultStruct),
+			*GetPathName());
+	}
 	return DefaultStructInstance.GetStructMemory();
 }
 
