CoreUObject: Deprecate all getters and setters for FField::FlagsPrivate of type EObjectFlags in preperation for a complete deprecation of this member and its use in all constructors too.
The FlagsPrivate member in FField is a remnant from UProperty with no meaningful or reliable value for any engine system today.

Changes that "should not change any behavior"(TM):

A1. Remove checking the always true condition !HasAnyFlags(RF_BeginDestroyed | RF_FinishDestroyed).
Files: MovieSceneLiveLinkStructPropertyBindings.h, NiagaraRendererComponents.h, TrackInstancePropertyBindings.h

A2. Remove the irrelevant use of SetFlags(RF_LoadCompleted) that is never checked in any loader code.
The removed logic dates back to the 8 year old cl 4101170.
Files: AnimBlueprintCompiler.cpp, KismetCompiler.cpp, UserDefinedStructureCompilerUtils.cpp

A3. Remove checkSlow(HasAllFlags(RF_Transient)).
Files: Property.cpp

A4. Remove the irrelevant use of SetFlags(RF_Transient).
The BlueprintCopmpilationManager call dates back to the 9 year old cl 3381178.
Add suggestion to consider using SetPropertyFlags(CPF_Transient) for OptimusNode_ConstantValue.
Files: OptimusNode_ConstantValue.cpp, BlueprintCompilationManager.cpp, RCPropertyContainer.cpp, DataprepParameterization.cpp

A5. Deprecate and change TMVVMFieldVariant::HasAnyFlags(InFlags) to always return false for all FProperties.
Update its single use case to do GetFunction()->HasAnyFlags(RF_NewerVersionExists) after checking IsFunction().
Files: MVVMFieldVariant.h, SMVVMFieldEntry.cpp

A6. Remove the redundant code that was calling GetFields().
Files: Rotation.cpp

A7. Remove the USTruct::TagSubObjects()/SetFlags(RF_TagGarbageTemp) implementation that was just wasting cycles tagging FProperties with RF_TagGarbageTemp.
Files: Class.h, Class.cpp

A8. Remove logic in UStruct::Link() that was just wasting cycles looking for changes in HasAllFlags(RF_Transient) and GetFName() before/after calling Property->Link(Ar).
The removed logic dates back to the 12 year old cl 2065450.
There are no FProperty LinkInternal implementations that touches the RF_Transient flag nor the NamePrivate field.
The most compex LinkInternal implementation is FStructProperty::LinkInternal that does PreloadInnerStructMembers() that may preload all its UFunctions.
Files: Class.cpp

A9. Remove the json export of the "Flags" field populated with the result of GetFlags().
Reading the exported json data is expected to be robust.
Files: JsonUClassExporter.cpp

Changes that may have some minor impact on behavior:

B1. Remove the use of the unreliable HasAllFlags(RF_Public) result and assume it returns true for all FProperties.
The FCustomPropertyOptionalPinManager::CanTreatPropertyAsOptional() code is already relying on the more relevant !HasMetaData(FBlueprintMetadata::MD_Private).
The BlueprintActionFilterImpl::IsGloballyAccessible(() code does not make much sense at all. Try and document the current behavior without changing it.
Note that e.g. the non-static struct property FPyTestClassSparseData::IntFromSparseData shows up as a blueprint variable in the BP action menu in any context both before and after this change.
Files: BlueprintActionFilter.cpp, CustomPropertyOptionalPinManager.cpp

#tests BuildCookRun Lyra, OverridableSerialization automation in EngineTest, Manual blueprint action menu testing in Lyra

#Synced-CL 50902884

diff --git a/Engine/Plugins/Animation/DeformerGraph/Source/OptimusCore/Private/Nodes/OptimusNode_ConstantValue.cpp b/Engine/Plugins/Animation/DeformerGraph/Source/OptimusCore/Private/Nodes/OptimusNode_ConstantValue.cpp
--- a/Engine/Plugins/Animation/DeformerGraph/Source/OptimusCore/Private/Nodes/OptimusNode_ConstantValue.cpp
+++ b/Engine/Plugins/Animation/DeformerGraph/Source/OptimusCore/Private/Nodes/OptimusNode_ConstantValue.cpp
@@ -52,7 +52,7 @@ UClass* UOptimusNode_ConstantValueGeneratorClass::GetClassForType(UPackage* InPa
 
 		// Out value doesn't need storage or saving.
 		FProperty* OutputValueProp = InDataType->CreateProperty(TypeClass, "Out");
-		OutputValueProp->SetFlags(RF_Transient);
+		// OutputValueProp->SetFlags(RF_Transient); // todo: SetFlags(RF_Transient) is a deprecated no-op, consider using SetPropertyFlags(CPF_Transient)
 #if WITH_EDITOR
 		OutputValueProp->SetMetaData(UOptimusNode::PropertyMeta::Output, TEXT("1"));
 #endif
diff --git a/Engine/Plugins/Animation/LiveLink/Source/LiveLinkMovieScene/Public/MovieScene/MovieSceneLiveLinkStructPropertyBindings.h b/Engine/Plugins/Animation/LiveLink/Source/LiveLinkMovieScene/Public/MovieScene/MovieSceneLiveLinkStructPropertyBindings.h
--- a/Engine/Plugins/Animation/LiveLink/Source/LiveLinkMovieScene/Public/MovieScene/MovieSceneLiveLinkStructPropertyBindings.h
+++ b/Engine/Plugins/Animation/LiveLink/Source/LiveLinkMovieScene/Public/MovieScene/MovieSceneLiveLinkStructPropertyBindings.h
@@ -218,12 +218,7 @@ private:
 
 		FProperty* GetProperty() const
 		{
-			FProperty* PropertyPtr = Property.Get();
-			if (PropertyPtr && !PropertyPtr->HasAnyFlags(RF_BeginDestroyed | RF_FinishDestroyed))
-			{
-				return PropertyPtr;
-			}
-			return nullptr;
+			return Property.Get();
 		}
 
 		template<typename ValueType>
diff --git a/Engine/Plugins/Enterprise/DataprepEditor/Source/DataprepCore/Private/Parameterization/DataprepParameterization.cpp b/Engine/Plugins/Enterprise/DataprepEditor/Source/DataprepCore/Private/Parameterization/DataprepParameterization.cpp
--- a/Engine/Plugins/Enterprise/DataprepEditor/Source/DataprepCore/Private/Parameterization/DataprepParameterization.cpp
+++ b/Engine/Plugins/Enterprise/DataprepEditor/Source/DataprepCore/Private/Parameterization/DataprepParameterization.cpp
@@ -1255,7 +1255,6 @@ void UDataprepParameterization::GenerateClass()
 		for ( TPair<FName, FProperty*>& Pair : NameToParameterizationProperty )
 		{
 			FProperty* NewProperty = CastField<FProperty>(FField::Duplicate(Pair.Value, CustomContainerClass, Pair.Key));
-			NewProperty->SetFlags( RF_Transient );
 			NewProperty->PropertyFlags = CPF_Edit;
 
 			// Need to manually call Link to fix-up some data (such as the C++ property flags) that are only set during Link
@@ -1442,7 +1441,6 @@ FProperty* UDataprepParameterization::AddPropertyToClass(const FName& Parameteri
 	if ( !NameToParameterizationProperty.Find( ParameterisationPropertyName ) )
 	{
 		FProperty* NewProperty = CastFieldChecked<FProperty>(FField::Duplicate(&Property, CustomContainerClass, ParameterisationPropertyName));
-		NewProperty->SetFlags( RF_Transient );
 		NewProperty->PropertyFlags = CPF_Edit | CPF_NonTransactional;
 
 		if ( !bAddFullProperty )
diff --git a/Engine/Plugins/FX/Niagara/Source/Niagara/Public/NiagaraRendererComponents.h b/Engine/Plugins/FX/Niagara/Source/Niagara/Public/NiagaraRendererComponents.h
--- a/Engine/Plugins/FX/Niagara/Source/Niagara/Public/NiagaraRendererComponents.h
+++ b/Engine/Plugins/FX/Niagara/Source/Niagara/Public/NiagaraRendererComponents.h
@@ -39,7 +39,7 @@ private:
 		FProperty* GetProperty() const
 		{
 			FProperty* PropertyPtr = Property.Get();
-			if (PropertyPtr && Address && !PropertyPtr->HasAnyFlags(RF_BeginDestroyed | RF_FinishDestroyed))
+			if (PropertyPtr && Address)
 			{
 				return PropertyPtr;
 			}
diff --git a/Engine/Plugins/Runtime/ModelViewViewModel/Source/ModelViewViewModel/Public/Types/MVVMFieldVariant.h b/Engine/Plugins/Runtime/ModelViewViewModel/Source/ModelViewViewModel/Public/Types/MVVMFieldVariant.h
--- a/Engine/Plugins/Runtime/ModelViewViewModel/Source/ModelViewViewModel/Public/Types/MVVMFieldVariant.h
+++ b/Engine/Plugins/Runtime/ModelViewViewModel/Source/ModelViewViewModel/Public/Types/MVVMFieldVariant.h
@@ -121,13 +121,10 @@ namespace UE::MVVM
 			return nullptr;
 		}
 
+		UE_DEPRECATED(5.8, "TMVVMFieldVariant::HasAnyFlags is deprecated, it is only relevant for functions. Consider using GetFunction()->HasAnyFlags() after checking IsFunction() instead.")
 		[[nodiscard]] bool HasAnyFlags(EObjectFlags InFlags) const
 		{
-			if (Type == EVariantType::Property)
-			{
-				return Binding.Property ? Binding.Property->HasAnyFlags(InFlags) : false;
-			}
-			else if (Type == EVariantType::Function)
+			if (Type == EVariantType::Function)
 			{
 				return Binding.Function ? Binding.Function->HasAnyFlags(InFlags) : false;
 			}
diff --git a/Engine/Plugins/Runtime/ModelViewViewModel/Source/ModelViewViewModelEditor/Private/Widgets/SMVVMFieldEntry.cpp b/Engine/Plugins/Runtime/ModelViewViewModel/Source/ModelViewViewModelEditor/Private/Widgets/SMVVMFieldEntry.cpp
--- a/Engine/Plugins/Runtime/ModelViewViewModel/Source/ModelViewViewModelEditor/Private/Widgets/SMVVMFieldEntry.cpp
+++ b/Engine/Plugins/Runtime/ModelViewViewModel/Source/ModelViewViewModelEditor/Private/Widgets/SMVVMFieldEntry.cpp
@@ -150,7 +150,7 @@ void SFieldPaths::SetFieldPaths(TArrayView<UE::MVVM::FMVVMConstFieldVariant> InP
 	for (; Index < InPropertyPath.Num(); ++Index)
 	{
 		bool bIsValid = InPropertyPath[Index].IsValid();
-		bool bIsFieldValid = bIsValid && !InPropertyPath[Index].HasAnyFlags(RF_NewerVersionExists);
+		bool bIsFieldValid = (bIsValid && InPropertyPath[Index].IsFunction()) ? !InPropertyPath[Index].GetFunction()->HasAnyFlags(RF_NewerVersionExists) : bIsValid;
 		UStruct* Struct = InPropertyPath[Index].GetOwner();
 		bool bIsOwnerValid = bIsValid && Struct && !Struct->HasAnyFlags(RF_NewerVersionExists);
 
diff --git a/Engine/Plugins/VirtualProduction/RemoteControl/Source/RemoteControlCommon/Private/RCPropertyContainer.cpp b/Engine/Plugins/VirtualProduction/RemoteControl/Source/RemoteControlCommon/Private/RCPropertyContainer.cpp
--- a/Engine/Plugins/VirtualProduction/RemoteControl/Source/RemoteControlCommon/Private/RCPropertyContainer.cpp
+++ b/Engine/Plugins/VirtualProduction/RemoteControl/Source/RemoteControlCommon/Private/RCPropertyContainer.cpp
@@ -527,7 +527,6 @@ TSubclassOf<URCPropertyContainerBase>& URCPropertyContainerRegistry::FindOrAddCo
 #if WITH_EDITORONLY_DATA
 	FField::CopyMetaData(InValueSrcProperty, ValueProperty);
 #endif
-	ValueProperty->SetFlags(RF_Transient);
 	ValueProperty->PropertyFlags |= CPF_Edit; // add this flag(s)
 	//ValueProperty->PropertyFlags &= ~CPF_BlueprintVisible; // clear this flag(s)
 
diff --git a/Engine/Source/Editor/AnimGraph/Private/AnimBlueprintCompiler.cpp b/Engine/Source/Editor/AnimGraph/Private/AnimBlueprintCompiler.cpp
--- a/Engine/Source/Editor/AnimGraph/Private/AnimBlueprintCompiler.cpp
+++ b/Engine/Source/Editor/AnimGraph/Private/AnimBlueprintCompiler.cpp
@@ -1840,12 +1840,6 @@ FProperty* FAnimBlueprintCompilerContext::CreateStructVariable(UScriptStruct* In
 	FProperty* NewProperty = FKismetCompilerUtilities::CreatePropertyOnScope(InStruct, VarName, VarType, nullptr, CPF_None, Schema, MessageLog);
 	if (NewProperty != nullptr)
 	{
-		// This fixes a rare bug involving asynchronous loading of BPs in editor builds. The pattern was established
-		// in FKismetCompilerContext::CompileFunctions where we do this for the uber graph function. By setting
-		// the RF_LoadCompleted we prevent the linker from overwriting our regenerated property, although the
-		// circumstances under which this occurs are murky. More testing of BPs loading asynchronously in the editor
-		// needs to be added:
-		NewProperty->SetFlags(RF_LoadCompleted);
 		FKismetCompilerUtilities::LinkAddedProperty(InStruct, NewProperty);
 	}
 	else
diff --git a/Engine/Source/Editor/AnimGraph/Private/CustomPropertyOptionalPinManager.cpp b/Engine/Source/Editor/AnimGraph/Private/CustomPropertyOptionalPinManager.cpp
--- a/Engine/Source/Editor/AnimGraph/Private/CustomPropertyOptionalPinManager.cpp
+++ b/Engine/Source/Editor/AnimGraph/Private/CustomPropertyOptionalPinManager.cpp
@@ -37,7 +37,6 @@ bool FCustomPropertyOptionalPinManager::CanTreatPropertyAsOptional(FProperty* Te
 
 	FEdGraphPinType PinType;
 	return (TestProperty->HasAllPropertyFlags(CPF_Edit | CPF_BlueprintVisible) &&
-			TestProperty->HasAllFlags(RF_Public) &&
 			K2Schema->ConvertPropertyToPinType(TestProperty, PinType) &&
 			!TestProperty->HasMetaData(FBlueprintMetadata::MD_Private));
 }
diff --git a/Engine/Source/Editor/BlueprintGraph/Private/BlueprintActionFilter.cpp b/Engine/Source/Editor/BlueprintGraph/Private/BlueprintActionFilter.cpp
--- a/Engine/Source/Editor/BlueprintGraph/Private/BlueprintActionFilter.cpp
+++ b/Engine/Source/Editor/BlueprintGraph/Private/BlueprintActionFilter.cpp
@@ -681,10 +681,14 @@ static bool BlueprintActionFilterImpl::IsLatent(FBlueprintActionInfo& BlueprintA
 static bool BlueprintActionFilterImpl::IsGloballyAccessible(FFieldVariant Field)
 {
 	bool bHasPersistentExtents = false; // is global or static
-	bool bIsPublic = Field.IsUObject() ? Field.ToUObject()->HasAnyFlags(RF_Public) : Field.ToField()->HasAnyFlags(RF_Public);
+	// todo: Consider using !HasAnyFlags(CPF_NativeAccessSpecifierProtected|CPF_NativeAccessSpecifierPrivate) for non-UObject fields
+	// (i.e. FProperties) to match the documentation of this function.
+	bool bIsPublic = Field.IsUObject() ? Field.ToUObject()->HasAnyFlags(RF_Public) : true; 
 	
 	UClass* ClassOuter = Cast<UClass>(Field.IsUObject() ? Field.ToUObject()->GetOuter() : Field.ToField()->GetOwner<UObject>());
-	// outer is probably a UPackage (for like global enums, structs, etc.)
+	// For UObjects like global enums, structs, etc., the outer may be a UPackage, treat these a globally accessible.
+	// For FProperties the outer may be a UClass or a UStruct, and we are currently
+	// treating all UStruct properties (including private and non-static ones) as globally accessible.
 	if (ClassOuter == nullptr)
 	{
 		bHasPersistentExtents = true;
@@ -2527,4 +2531,4 @@ void FBlueprintActionFilter::EndFilterTestProfileEvent(FFilterTestProfileRecord*
 		}
 	}
 }
-#endif	// ENABLE_BLUEPRINT_ACTION_FILTER_PROFILING
\ No newline at end of file
+#endif	// ENABLE_BLUEPRINT_ACTION_FILTER_PROFILING
diff --git a/Engine/Source/Editor/Kismet/Private/BlueprintCompilationManager.cpp b/Engine/Source/Editor/Kismet/Private/BlueprintCompilationManager.cpp
--- a/Engine/Source/Editor/Kismet/Private/BlueprintCompilationManager.cpp
+++ b/Engine/Source/Editor/Kismet/Private/BlueprintCompilationManager.cpp
@@ -3523,7 +3523,6 @@ PRAGMA_ENABLE_DEPRECATION_WARNINGS
 					FProperty* Param = FKismetCompilerUtilities::CreatePropertyOnScope(NewFunction, Pin->PinName, Pin->PinType, Ret, CPF_BlueprintVisible|CPF_BlueprintReadOnly, Schema, MessageLog);
 					if(Param)
 					{
-						Param->SetFlags(RF_Transient);
 						Param->PropertyFlags |= CPF_Parm;
 						if(Pin->PinType.bIsReference)
 						{
@@ -3600,7 +3599,6 @@ PRAGMA_ENABLE_DEPRECATION_WARNINGS
 								FProperty* Param = FKismetCompilerUtilities::CreatePropertyOnScope(NewFunction, Pin->PinName, Pin->PinType, Ret, CPF_None, Schema, MessageLog);
 								if(Param)
 								{
-									Param->SetFlags(RF_Transient);
 									// we only tag things as CPF_ReturnParm if the value is named ReturnValue.... this is *terrible* behavior:
 									if(Param->GetFName() == UEdGraphSchema_K2::PN_ReturnValue)
 									{
@@ -3663,10 +3661,7 @@ PRAGMA_ENABLE_DEPRECATION_WARNINGS
 					// locals:
 					for( const FBPVariableDescription& BPVD : EntryNode->LocalVariables )
 					{
-						if(FProperty* LocalVariable = FKismetCompilerContext::CreateUserDefinedLocalVariableForFunction(BPVD, NewFunction, Ret, PropertyListBuilder, Schema, MessageLog) )
-						{
-							LocalVariable->SetFlags(RF_Transient);
-						}
+						FKismetCompilerContext::CreateUserDefinedLocalVariableForFunction(BPVD, NewFunction, Ret, PropertyListBuilder, Schema, MessageLog);
 					}
 
 					// __WorldContext:
@@ -3678,7 +3673,6 @@ PRAGMA_ENABLE_DEPRECATION_WARNINGS
 							FProperty* Param = FKismetCompilerUtilities::CreatePropertyOnScope(NewFunction, TEXT("__WorldContext"), WorldContextPinType, Ret, CPF_None, Schema, MessageLog);
 							if(Param)
 							{
-								Param->SetFlags(RF_Transient);
 								Param->PropertyFlags |= CPF_Parm;
 								PropertyListBuilder.AppendNoTerminate(*Param);
 							}
diff --git a/Engine/Source/Editor/KismetCompiler/Private/KismetCompiler.cpp b/Engine/Source/Editor/KismetCompiler/Private/KismetCompiler.cpp
--- a/Engine/Source/Editor/KismetCompiler/Private/KismetCompiler.cpp
+++ b/Engine/Source/Editor/KismetCompiler/Private/KismetCompiler.cpp
@@ -965,9 +965,6 @@ FProperty* FKismetCompilerContext::CreateVariable(const FName VarName, const FEd
 	FProperty* NewProperty = FKismetCompilerUtilities::CreatePropertyOnScope(NewClass, VarName, VarType, NewClass, CPF_None, Schema, MessageLog);
 	if (NewProperty != nullptr)
 	{
-		// This is necessary to stop async loading and partial reloads from corrupting the compiled data
-		// without marking the owning class for recompilation
-		NewProperty->SetFlags(RF_LoadCompleted);
 		FKismetCompilerUtilities::LinkAddedProperty(NewClass, NewProperty);
 	}
 	else
diff --git a/Engine/Source/Editor/KismetCompiler/Private/UserDefinedStructureCompilerUtils.cpp b/Engine/Source/Editor/KismetCompiler/Private/UserDefinedStructureCompilerUtils.cpp
--- a/Engine/Source/Editor/KismetCompiler/Private/UserDefinedStructureCompilerUtils.cpp
+++ b/Engine/Source/Editor/KismetCompiler/Private/UserDefinedStructureCompilerUtils.cpp
@@ -298,7 +298,6 @@ struct FUserDefinedStructureCompilerInner
 			}
 			else if (bIsNewVariable)
 			{
-				VarProperty->SetFlags(RF_LoadCompleted);
 				FKismetCompilerUtilities::LinkAddedProperty(Struct, VarProperty);
 			}
 
diff --git a/Engine/Source/Runtime/CoreUObject/Private/UObject/Class.cpp b/Engine/Source/Runtime/CoreUObject/Private/UObject/Class.cpp
--- a/Engine/Source/Runtime/CoreUObject/Private/UObject/Class.cpp
+++ b/Engine/Source/Runtime/CoreUObject/Private/UObject/Class.cpp
@@ -900,62 +900,45 @@ void UStruct::Link(FArchive& Ar, bool bRelinkExistingProperties)
 #endif // WITH_EDITORONLY_DATA
 		}
 
-		int32 LoopNum = 1;
-		for (int32 LoopIter = 0; LoopIter < LoopNum; LoopIter++)
+	#if WITH_EDITORONLY_DATA
+		TotalFieldCount = 0;
+	#endif
+		PropertiesSize = 0;
+		MinAlignment = 1;
+
+		if (InheritanceSuper)
 		{
 		#if WITH_EDITORONLY_DATA
-			TotalFieldCount = 0;
+			TotalFieldCount = InheritanceSuper->TotalFieldCount;
 		#endif
-			PropertiesSize = 0;
-			MinAlignment = 1;
+			PropertiesSize = InheritanceSuper->GetPropertiesSize();
+			MinAlignment = IntCastChecked<int16>(InheritanceSuper->GetMinAlignment());
+		}
 
-			if (InheritanceSuper)
+		for (FField* Field = ChildProperties; Field; Field = Field->Next)
+		{
+			if (Field->GetOwner<UObject>() != this)
 			{
-			#if WITH_EDITORONLY_DATA
-				TotalFieldCount = InheritanceSuper->TotalFieldCount;
-			#endif
-				PropertiesSize = InheritanceSuper->GetPropertiesSize();
-				MinAlignment = IntCastChecked<int16>(InheritanceSuper->GetMinAlignment());
+				break;
 			}
 
-			for (FField* Field = ChildProperties; Field; Field = Field->Next)
+			if (FProperty* Property = CastField<FProperty>(Field))
 			{
-				if (Field->GetOwner<UObject>() != this)
-				{
-					break;
-				}
-
-				if (FProperty* Property = CastField<FProperty>(Field))
-				{
-#if !WITH_EDITORONLY_DATA
-					// If we don't have the editor, make sure we aren't trying to link properties that are editor only.
-					check(!Property->IsEditorOnlyProperty());
-#endif // WITH_EDITORONLY_DATA
-					ensureMsgf(Property->GetOwner<UObject>() == this, TEXT("Linking '%s'. Property '%s' has outer '%s'"),
-						*GetFullName(), *Property->GetName(), *Property->GetOwnerVariant().GetFullName());
-
-					// Linking a property can cause a recompilation of the struct.
-					// When the property was changed, the struct should be relinked again, to be sure, the PropertiesSize is actual.
-					const bool bPropertyIsTransient = Property->HasAllFlags(RF_Transient);
-					const FName PropertyName = Property->GetFName();
-
-					PropertiesSize = Property->Link(Ar);
+			#if !WITH_EDITORONLY_DATA
+				// If we don't have the editor, make sure we aren't trying to link properties that are editor only.
+				check(!Property->IsEditorOnlyProperty());
+			#endif // WITH_EDITORONLY_DATA
+				ensureMsgf(Property->GetOwner<UObject>() == this, TEXT("Linking '%s'. Property '%s' has outer '%s'"),
+					*GetFullName(), *Property->GetName(), *Property->GetOwnerVariant().GetFullName());
 
-				#if WITH_EDITORONLY_DATA
-					Property->SetIndexInOwner(TotalFieldCount);
-					TotalFieldCount += Property->ArrayDim;
-				#endif
+				PropertiesSize = Property->Link(Ar);
 
-					if ((bPropertyIsTransient != Property->HasAllFlags(RF_Transient)) || (PropertyName != Property->GetFName()))
-					{
-						LoopNum++;
-						const int32 MaxLoopLimit = 64;
-						ensure(LoopNum < MaxLoopLimit);
-						break;
-					}
+			#if WITH_EDITORONLY_DATA
+				Property->SetIndexInOwner(TotalFieldCount);
+				TotalFieldCount += Property->ArrayDim;
+			#endif
 
-					MinAlignment = IntCastChecked<int16>(FMath::Max(MinAlignment, Property->GetMinAlignment()));
-				}
+				MinAlignment = IntCastChecked<int16>(FMath::Max(MinAlignment, Property->GetMinAlignment()));
 			}
 		}
 
@@ -2764,21 +2747,6 @@ void UStruct::InstanceSubobjectTemplates(TNotNull<void*> Data, const void* Defau
 	}
 }
 
-void UStruct::TagSubobjects(EObjectFlags NewFlags)
-{
-	Super::TagSubobjects(NewFlags);
-
-	// Tag our properties
-	for (TFieldIterator<FProperty> It(this, EFieldIteratorFlags::ExcludeSuper); It; ++It)
-	{
-		FProperty* Property = *It;
-		if (Property && !Property->HasAnyFlags(GARBAGE_COLLECTION_KEEPFLAGS) && !Property->IsRooted())
-		{
-			Property->SetFlags(NewFlags);
-		}
-	}
-}
-
 /**
 * @return	true if this object is of the specified type.
 */
diff --git a/Engine/Source/Runtime/CoreUObject/Private/UObject/Field.cpp b/Engine/Source/Runtime/CoreUObject/Private/UObject/Field.cpp
--- a/Engine/Source/Runtime/CoreUObject/Private/UObject/Field.cpp
+++ b/Engine/Source/Runtime/CoreUObject/Private/UObject/Field.cpp
@@ -341,7 +341,9 @@ FField::FField(UField* InField)
 {
 	check(InField);
 	InField->ConditionalPreload();
+	PRAGMA_DISABLE_DEPRECATION_WARNINGS
 	FlagsPrivate = InField->GetFlags();
+	PRAGMA_ENABLE_DEPRECATION_WARNINGS
 	// Associate this FField with the UField we're constructing from so that next time something tries to convert it,
 	// it can already grab the cached new FField
 	InField->SetAssociatedFField(this);
@@ -847,7 +849,9 @@ void FField::PostDuplicate(const FField& InField)
 FField* FField::Duplicate(const FField* InField, FFieldVariant DestOwner, const FName DestName, EObjectFlags FlagMask, EInternalObjectFlags InternalFlagsMask)
 {
 	check(InField);
+	PRAGMA_DISABLE_DEPRECATION_WARNINGS
 	FField* NewField = InField->GetClass()->Construct(DestOwner, (DestName == NAME_None) ? InField->GetFName() : DestName, InField->GetFlags() & FlagMask);
+	PRAGMA_ENABLE_DEPRECATION_WARNINGS
 	NewField->PostDuplicate(*InField);
 	return NewField;
 }
diff --git a/Engine/Source/Runtime/CoreUObject/Private/UObject/Property.cpp b/Engine/Source/Runtime/CoreUObject/Private/UObject/Property.cpp
--- a/Engine/Source/Runtime/CoreUObject/Private/UObject/Property.cpp
+++ b/Engine/Source/Runtime/CoreUObject/Private/UObject/Property.cpp
@@ -802,7 +802,6 @@ void FProperty::Init()
 	//	checkSlow(!HasAnyPropertyFlags(CPF_EditorOnly));
 #endif // WITH_EDITORONLY_DATA
 	checkSlow(GetOwnerUField()->HasAllFlags(RF_Transient));
-	checkSlow(HasAllFlags(RF_Transient));
 
 	if (GetOwner<UObject>())
 	{
diff --git a/Engine/Source/Runtime/CoreUObject/Public/UObject/Class.h b/Engine/Source/Runtime/CoreUObject/Public/UObject/Class.h
--- a/Engine/Source/Runtime/CoreUObject/Public/UObject/Class.h
+++ b/Engine/Source/Runtime/CoreUObject/Public/UObject/Class.h
@@ -656,7 +656,6 @@ public:
 #endif // !UE_WITH_CONSTINIT_UOBJECT
 	static COREUOBJECT_API void AddReferencedObjects(UObject* InThis, FReferenceCollector& Collector);
 	COREUOBJECT_API virtual void GetPreloadDependencies(TArray<UObject*>& OutDeps) override;
-	COREUOBJECT_API virtual void TagSubobjects(EObjectFlags NewFlags) override;
 
 	// UField interface.
 	COREUOBJECT_API virtual void AddCppProperty(FProperty* Property) override;
diff --git a/Engine/Source/Runtime/CoreUObject/Public/UObject/Field.h b/Engine/Source/Runtime/CoreUObject/Public/UObject/Field.h
--- a/Engine/Source/Runtime/CoreUObject/Public/UObject/Field.h
+++ b/Engine/Source/Runtime/CoreUObject/Public/UObject/Field.h
@@ -714,57 +714,48 @@ public:
 #endif
 
 protected:
-	/**
-	* Set the object flags directly
-	**/
+	UE_DEPRECATED(5.8, "FField::SetFlagsTo is deprecated, for FProperties consider using PropertyFlags instead.")
 	void SetFlagsTo(EObjectFlags NewFlags)
 	{
-		checkfSlow((NewFlags & ~RF_AllFlags) == 0, TEXT("%s flagged as 0x%x but is trying to set flags to RF_AllFlags"), *GetFName().ToString(), (int32)FlagsPrivate);
 		FlagsPrivate = NewFlags;
 	}
 public:
-	/**
-	* Retrieve the object flags directly
-	*
-	* @return Flags for this object
-	**/
+	UE_DEPRECATED(5.8, "FField::GetFlags is deprecated, for FProperties consider using PropertyFlags instead.")
 	EObjectFlags GetFlags() const
 	{
-		checkfSlow((FlagsPrivate & ~RF_AllFlags) == 0, TEXT("%s flagged as RF_AllFlags"), *GetFName().ToString());
 		return FlagsPrivate;
 	}
 
+	UE_DEPRECATED(5.8, "FField::SetFlags is deprecated, for FProperties consider using PropertyFlags instead.")
 	void SetFlags(EObjectFlags NewFlags)
 	{
-		checkSlow(!(NewFlags & (RF_MarkAsNative | RF_MarkAsRootSet))); // These flags can't be used outside of constructors / internal code
+		PRAGMA_DISABLE_DEPRECATION_WARNINGS
 		SetFlagsTo(GetFlags() | NewFlags);
+		PRAGMA_ENABLE_DEPRECATION_WARNINGS
+
 	}
+	UE_DEPRECATED(5.8, "FField::ClearFlags is deprecated, for FProperties consider using PropertyFlags instead.")
 	void ClearFlags(EObjectFlags NewFlags)
 	{
-		checkSlow(!(NewFlags & (RF_MarkAsNative | RF_MarkAsRootSet)) || NewFlags == RF_AllFlags); // These flags can't be used outside of constructors / internal code
+		PRAGMA_DISABLE_DEPRECATION_WARNINGS
 		SetFlagsTo(GetFlags() & ~NewFlags);
+		PRAGMA_ENABLE_DEPRECATION_WARNINGS
 	}
-	/**
-	* Used to safely check whether any of the passed in flags are set.
-	*
-	* @param FlagsToCheck	Object flags to check for.
-	* @return				true if any of the passed in flags are set, false otherwise  (including no flags passed in).
-	*/
+
+	UE_DEPRECATED(5.8, "FField::HasAnyFlags is deprecated, for FProperties consider using PropertyFlags instead.")
 	bool HasAnyFlags(EObjectFlags FlagsToCheck) const
 	{
-		checkSlow(!(FlagsToCheck & (RF_MarkAsNative | RF_MarkAsRootSet)) || FlagsToCheck == RF_AllFlags); // These flags can't be used outside of constructors / internal code
+		PRAGMA_DISABLE_DEPRECATION_WARNINGS
 		return (GetFlags() & FlagsToCheck) != 0;
+		PRAGMA_ENABLE_DEPRECATION_WARNINGS
 	}
-	/**
-	* Used to safely check whether all of the passed in flags are set.
-	*
-	* @param FlagsToCheck	Object flags to check for
-	* @return true if all of the passed in flags are set (including no flags passed in), false otherwise
-	*/
+
+	UE_DEPRECATED(5.8, "FField::HasAllFlags is deprecated, for FProperties consider using PropertyFlags instead.")
 	bool HasAllFlags(EObjectFlags FlagsToCheck) const
 	{
-		checkSlow(!(FlagsToCheck & (RF_MarkAsNative | RF_MarkAsRootSet)) || FlagsToCheck == RF_AllFlags); // These flags can't be used outside of constructors / internal code
+		PRAGMA_DISABLE_DEPRECATION_WARNINGS
 		return ((GetFlags() & FlagsToCheck) == FlagsToCheck);
+		PRAGMA_ENABLE_DEPRECATION_WARNINGS
 	}
 
 	inline FFieldClass* GetClass() const
diff --git a/Engine/Source/Runtime/JsonUtilities/Private/JsonUClassExporter.cpp b/Engine/Source/Runtime/JsonUtilities/Private/JsonUClassExporter.cpp
--- a/Engine/Source/Runtime/JsonUtilities/Private/JsonUClassExporter.cpp
+++ b/Engine/Source/Runtime/JsonUtilities/Private/JsonUClassExporter.cpp
@@ -230,7 +230,6 @@ void WriteFField(const FField* Field, FJsonUETypeWriter& Writer)
 
 	Writer.WriteValue(TEXT("Name"), Field->GetName());
 	Writer.WriteValue(TEXT("ClassName"), Field->GetClass()->GetName());
-	Writer.WriteValue(TEXT("Flags"), LexToString(Field->GetFlags()));
 
 	if (const FProperty* Property = CastField<const FProperty>(Field))
 	{
diff --git a/Engine/Source/Runtime/MovieScene/Public/TrackInstancePropertyBindings.h b/Engine/Source/Runtime/MovieScene/Public/TrackInstancePropertyBindings.h
--- a/Engine/Source/Runtime/MovieScene/Public/TrackInstancePropertyBindings.h
+++ b/Engine/Source/Runtime/MovieScene/Public/TrackInstancePropertyBindings.h
@@ -30,7 +30,7 @@ struct FCachedProperty
 	FProperty* GetValidProperty() const
 	{
 		FProperty* PropertyPtr = Property.Get();
-		if (ContainerAddress && PropertyPtr && !PropertyPtr->HasAnyFlags(RF_BeginDestroyed | RF_FinishDestroyed))
+		if (ContainerAddress && PropertyPtr)
 		{
 			return PropertyPtr;
 		}
@@ -109,7 +109,7 @@ struct FVolatileProperty
 	FProperty* GetValidProperty() const
 	{
 		FProperty* PropertyPtr = LeafProperty.Get();
-		if (RootContainer && PropertyPtr && !PropertyPtr->HasAnyFlags(RF_BeginDestroyed | RF_FinishDestroyed))
+		if (RootContainer && PropertyPtr)
 		{
 			return PropertyPtr;
 		}
