Fixed TConstStructView::operator->() constness.

diff --git a/Engine/Source/Runtime/CoreUObject/Public/StructUtils/StructView.h b/Engine/Source/Runtime/CoreUObject/Public/StructUtils/StructView.h
--- a/Engine/Source/Runtime/CoreUObject/Public/StructUtils/StructView.h
+++ b/Engine/Source/Runtime/CoreUObject/Public/StructUtils/StructView.h
@@ -385,7 +385,7 @@ struct [[nodiscard]] TConstStructView : FConstStructView
 	}
 
 	/** Simplified accessor for members of the guaranteed minimum structure type */
-	BaseStructT* operator->() const
+	const BaseStructT* operator->() const
 	{
 		check(IsValid());
 		return GetPtr();
