Add LexicalLess comparison helper to FPrimaryAssetId for deterministic sorting in containers.

#rnx
#Synced-CL 50906023

diff --git a/Engine/Source/Runtime/CoreUObject/Public/UObject/PrimaryAssetId.h b/Engine/Source/Runtime/CoreUObject/Public/UObject/PrimaryAssetId.h
--- a/Engine/Source/Runtime/CoreUObject/Public/UObject/PrimaryAssetId.h
+++ b/Engine/Source/Runtime/CoreUObject/Public/UObject/PrimaryAssetId.h
@@ -45,29 +45,29 @@ struct FPrimaryAssetType
 	operator const FName&() const { return Name; }
 
 	/** Returns internal Name explicitly, not normally needed */
-	FName GetName() const
+	[[nodiscard]] FName GetName() const
 	{
 		return Name;
 	}
 
-	bool operator==(const FPrimaryAssetType& Other) const
+	[[nodiscard]] bool operator==(const FPrimaryAssetType& Other) const
 	{
 		return Name == Other.Name;
 	}
 
-	bool operator!=(const FPrimaryAssetType& Other) const
+	[[nodiscard]] bool operator!=(const FPrimaryAssetType& Other) const
 	{
 		return Name != Other.Name;
 	}
 
 	/** Returns true if this is a valid Type */
-	bool IsValid() const
+	[[nodiscard]] bool IsValid() const
 	{
 		return Name != NAME_None;
 	}
 
 	/** Returns string version of this Type */
-	FString ToString() const
+	[[nodiscard]] FString ToString() const
 	{
 		return Name.ToString();
 	}
@@ -83,17 +83,25 @@ struct FPrimaryAssetType
 	bool ImportTextItem(const TCHAR*& Buffer, int32 PortFlags, UObject* Parent, FOutputDevice* ErrorText);
 	bool SerializeFromMismatchedTag(struct FPropertyTag const& Tag, FStructuredArchive::FSlot Slot);
 
-	friend inline uint32 GetTypeHash(const FPrimaryAssetType& Key)
+	[[nodiscard]] friend inline uint32 GetTypeHash(const FPrimaryAssetType& Key)
 	{
 		return GetTypeHash(Key.Name);
 	}
 
-	bool LexicalLess(const FPrimaryAssetType& Other) const
+	/** Compares FPrimaryAssetTypes lexically based on internal FName. Sort is alphabetical ascending, ignores case. */
+	[[nodiscard]] int32 CompareLexical(const FPrimaryAssetType& Other) const
+	{
+		return Name.Compare(Other.Name);
+	}
+
+	/** Slow alphabetical order that is stable / deterministic over process runs, ignores case. */
+	[[nodiscard]] bool LexicalLess(const FPrimaryAssetType& Other) const
 	{
 		return Name.LexicalLess(Other.Name);
 	}
 
-	bool FastLess(const FPrimaryAssetType& Other) const
+	/** Fast non-alphabetical order that is only stable during this process' lifetime. Uses FName comparison index. */
+	[[nodiscard]] bool FastLess(const FPrimaryAssetType& Other) const
 	{
 		return Name.FastLess(Other.Name);
 	}
@@ -146,9 +154,9 @@ struct FPrimaryAssetId
 		: PrimaryAssetType(InAssetType), PrimaryAssetName(InAssetName)
 	{}
 
-	static COREUOBJECT_API FPrimaryAssetId ParseTypeAndName(const TCHAR* TypeAndName, uint32 Len);
-	static COREUOBJECT_API FPrimaryAssetId ParseTypeAndName(FName TypeAndName);
-	static FPrimaryAssetId ParseTypeAndName(const FString& TypeAndName)
+	[[nodiscard]] static COREUOBJECT_API FPrimaryAssetId ParseTypeAndName(const TCHAR* TypeAndName, uint32 Len);
+	[[nodiscard]] static COREUOBJECT_API FPrimaryAssetId ParseTypeAndName(FName TypeAndName);
+	[[nodiscard]] static FPrimaryAssetId ParseTypeAndName(const FString& TypeAndName)
 	{
 		return ParseTypeAndName(*TypeAndName, static_cast<uint32>(TypeAndName.Len()));
 	}
@@ -163,13 +171,13 @@ struct FPrimaryAssetId
 	FPrimaryAssetId& operator=(FPrimaryAssetId&&) = default;
 
 	/** Returns true if this is a valid identifier */
-	bool IsValid() const
+	[[nodiscard]] bool IsValid() const
 	{
 		return PrimaryAssetType != NAME_None && PrimaryAssetName != NAME_None;
 	}
 
 	/** Returns string version of this identifier in Type:Name format */
-	FString ToString() const
+	[[nodiscard]] FString ToString() const
 	{
 		TStringBuilder<256> Builder;
 		AppendString(Builder);
@@ -188,17 +196,17 @@ struct FPrimaryAssetId
 	}
 
 	/** Converts from Type:Name format */
-	static FPrimaryAssetId FromString(const FString& String)
+	[[nodiscard]] static FPrimaryAssetId FromString(const FString& String)
 	{
 		return FPrimaryAssetId(String);
 	}
 
-	bool operator==(const FPrimaryAssetId& Other) const
+	[[nodiscard]] bool operator==(const FPrimaryAssetId& Other) const
 	{
 		return PrimaryAssetType == Other.PrimaryAssetType && PrimaryAssetName == Other.PrimaryAssetName;
 	}
 
-	bool operator!=(const FPrimaryAssetId& Other) const
+	[[nodiscard]] bool operator!=(const FPrimaryAssetId& Other) const
 	{
 		return PrimaryAssetType != Other.PrimaryAssetType || PrimaryAssetName != Other.PrimaryAssetName;
 	}
@@ -215,7 +223,7 @@ struct FPrimaryAssetId
 	bool ImportTextItem(const TCHAR*& Buffer, int32 PortFlags, UObject* Parent, FOutputDevice* ErrorText);
 	bool SerializeFromMismatchedTag(struct FPropertyTag const& Tag, FStructuredArchive::FSlot Slot);
 
-	friend inline uint32 GetTypeHash(const FPrimaryAssetId& Key)
+	[[nodiscard]] friend inline uint32 GetTypeHash(const FPrimaryAssetId& Key)
 	{
 		uint32 Hash = 0;
 
@@ -224,6 +232,24 @@ struct FPrimaryAssetId
 		return Hash;
 	}
 
+	/** Compare lexical ordering by PrimaryAssetType first followed by PrimaryAssetName, ignores case. */
+	[[nodiscard]] int32 CompareLexical(const FPrimaryAssetId& Other) const
+	{
+		const int32 AssetTypeCompare = PrimaryAssetType.CompareLexical(Other.PrimaryAssetType);
+		if (AssetTypeCompare != 0)
+		{
+			return AssetTypeCompare;
+		}
+
+		return PrimaryAssetName.Compare(Other.PrimaryAssetName);
+	}
+
+	/** Slow alphabetical order that is stable / deterministic over process runs, ignores case. */
+	[[nodiscard]] bool LexicalLess(const FPrimaryAssetId& Other) const
+	{
+		return CompareLexical(Other) < 0;
+	}
+
 private:
 	 COREUOBJECT_API void SerializeForLog(FCbWriter& Writer, const FPrimaryAssetId& Value);
 };
