Add a cvar s.AsyncCreatePackagesFromQueueMaxTimeLimitMs to limit the time spent in CreateAsyncPackagesFromQueue when running on the async loading thread. Default to 2.5ms
This can improve when s.AllowMultithreadedLoading is enabled, by allowing I/O and multithreaded loader work to kick off earlier.
#Synced-CL 50901525

diff --git a/Engine/Source/Runtime/CoreUObject/Private/Serialization/AsyncLoading2.cpp b/Engine/Source/Runtime/CoreUObject/Private/Serialization/AsyncLoading2.cpp
--- a/Engine/Source/Runtime/CoreUObject/Private/Serialization/AsyncLoading2.cpp
+++ b/Engine/Source/Runtime/CoreUObject/Private/Serialization/AsyncLoading2.cpp
@@ -541,6 +541,16 @@ static FAutoConsoleVariableRef CVarGAllowMultithreadedLoading(
 	ECVF_Default
 );
 
+static float GAsyncCreatePackagesFromQueueMaxTimeLimitMs = 2.5f;
+static FAutoConsoleVariableRef CVarAsyncCreatePackagesFromQueueMaxTimeLimitMs(
+	TEXT("s.AsyncCreatePackagesFromQueueMaxTimeLimitMs"),
+	GAsyncCreatePackagesFromQueueMaxTimeLimitMs,
+	TEXT("Limit the time spent in CreateAsyncPackagesFromQueue when running on the ALT.\n")
+	TEXT("This can improve throughput when s.AllowMultithreadedLoading is enabled by allowing MT work to kick off earlier."),
+	ECVF_Default
+);
+
+
 static bool GOnlyProcessRequiredPackagesWhenSyncLoading = true;
 static FAutoConsoleVariableRef CVarGOnlyProcessRequiredPackagesWhenSyncLoading(
 	TEXT("s.OnlyProcessRequiredPackagesWhenSyncLoading"),
@@ -11327,6 +11337,11 @@ uint32 FAsyncLoadingThread2::Run()
 
 					if (QueuedPackagesCounter || !PendingPackages.IsEmpty())
 					{
+						TRACE_CPUPROFILER_EVENT_SCOPE_CONDITIONAL(TimeLimitedCreateAsyncPackages, GAsyncCreatePackagesFromQueueMaxTimeLimitMs > .0f);
+						// Use a time limit scope to restore the time limit to the old values when exiting the scope
+						// This is required to ensure a reentrant call here doesn't overwrite time limits permanently.
+						FAsyncLoadingThreadState2::FTimeLimitScope TimeLimitScope(ThreadState, GAsyncCreatePackagesFromQueueMaxTimeLimitMs > .0f, double(GAsyncCreatePackagesFromQueueMaxTimeLimitMs/1000.0f));
+
 						if (CreateAsyncPackagesFromQueue(ThreadState, GAsyncLoading2_AllowPreemptingPackagesDuringGC))
 						{
 							// Fall through to FAsyncLoadEventQueue2 processing unless we need to suspend
