[StructUtils] Generate the correct underlying type of FEnumProperty. When the enum form is EnumClass, UBT generates the correct underlying type. The underlying type is now known. It can be any of the supported integer types.
Add support for int8, int16, and uint16.
Refactor and expose GetPropertyValueAsInt64 and SetPropertyValueFromInt64 to be reusable between different modules. There is no need to test the type, because FNumericalProperty already handles it.

#Synced-CL 50933965

diff --git a/Engine/Plugins/Experimental/UAF/UAF/Source/UAF/Private/Variables/UAFInstanceVariableContainer.cpp b/Engine/Plugins/Experimental/UAF/UAF/Source/UAF/Private/Variables/UAFInstanceVariableContainer.cpp
--- a/Engine/Plugins/Experimental/UAF/UAF/Source/UAF/Private/Variables/UAFInstanceVariableContainer.cpp
+++ b/Engine/Plugins/Experimental/UAF/UAF/Source/UAF/Private/Variables/UAFInstanceVariableContainer.cpp
@@ -21,239 +21,35 @@ namespace UE::UAF::Private
 
 	EPropertyBagResult GetPropertyAsInt64(const FProperty* InProperty, EPropertyBagPropertyType InValueType, const void* Address, int64& OutValue)
 	{
-		check(Address != nullptr);
-		check(InProperty != nullptr);
-
-		switch (InValueType)
-		{
-		case EPropertyBagPropertyType::Bool:
-		{
-			const FBoolProperty* Property = CastFieldChecked<FBoolProperty>(InProperty);
-			OutValue = Property->GetPropertyValue(Address) ? 1 : 0;
-			return EPropertyBagResult::Success;
-		}
-		case EPropertyBagPropertyType::Byte:
-		{
-			const FByteProperty* Property = CastFieldChecked<FByteProperty>(InProperty);
-			OutValue = Property->GetPropertyValue(Address);
-			return EPropertyBagResult::Success;
-		}
-		case EPropertyBagPropertyType::Int32:
-		{
-			const FIntProperty* Property = CastFieldChecked<FIntProperty>(InProperty);
-			OutValue = Property->GetPropertyValue(Address);
-			return EPropertyBagResult::Success;
-		}
-		case EPropertyBagPropertyType::UInt32:
-		{
-			const FUInt32Property* Property = CastFieldChecked<FUInt32Property>(InProperty);
-			OutValue = static_cast<uint32>(Property->GetPropertyValue(Address));
-			return EPropertyBagResult::Success;
-		}
-		case EPropertyBagPropertyType::Int64:
-		{
-			const FInt64Property* Property = CastFieldChecked<FInt64Property>(InProperty);
-			OutValue = Property->GetPropertyValue(Address);
-			return EPropertyBagResult::Success;
-		}
-		case EPropertyBagPropertyType::UInt64:
-		{
-			const FUInt64Property* Property = CastFieldChecked<FUInt64Property>(InProperty);
-			OutValue = static_cast<int64>(Property->GetPropertyValue(Address));
-			return EPropertyBagResult::Success;
-		}
-		case EPropertyBagPropertyType::Float:
+		TValueOrError<int64, void> Result = UE::StructUtils::GetPropertyValueAsInt64(InProperty, Address);
+		if (Result.HasValue())
 		{
-			const FFloatProperty* Property = CastFieldChecked<FFloatProperty>(InProperty);
-			OutValue = static_cast<int64>(Property->GetPropertyValue(Address));
+			OutValue = Result.GetValue();
 			return EPropertyBagResult::Success;
 		}
-		case EPropertyBagPropertyType::Double:
-		{
-			const FDoubleProperty* Property = CastFieldChecked<FDoubleProperty>(InProperty);
-			OutValue = static_cast<int64>(Property->GetPropertyValue(Address));
-			return EPropertyBagResult::Success;
-		}
-		case EPropertyBagPropertyType::Enum:
-		{
-			if (const FEnumProperty* EnumProperty = CastField<FEnumProperty>(InProperty))
-			{
-				const FNumericProperty* UnderlyingProperty = EnumProperty->GetUnderlyingProperty();
-				check(UnderlyingProperty);
-				OutValue = UnderlyingProperty->GetSignedIntPropertyValue(Address);
-				return EPropertyBagResult::Success;
-			}
-			else if (const FByteProperty* ByteProperty = CastField<FByteProperty>(InProperty))
-			{
-				if (ByteProperty->Enum != nullptr)
-				{
-					OutValue = ByteProperty->GetPropertyValue(Address);
-					return EPropertyBagResult::Success;
-				}
-			}
-			return EPropertyBagResult::TypeMismatch;
-		}
-		default:
-			return EPropertyBagResult::TypeMismatch;
-		}
+		return EPropertyBagResult::TypeMismatch;
 	}
 
 	EPropertyBagResult GetPropertyAsUInt64(const FProperty* InProperty, EPropertyBagPropertyType InValueType, const void* Address, uint64& OutValue)
 	{
-		check(Address != nullptr);
-		check(InProperty != nullptr);
-
-		switch (InValueType)
+		TValueOrError<uint64, void> Result = UE::StructUtils::GetPropertyValueAsUInt64(InProperty, Address);
+		if (Result.HasValue())
 		{
-		case EPropertyBagPropertyType::Bool:
-		{
-			const FBoolProperty* Property = CastFieldChecked<FBoolProperty>(InProperty);
-			OutValue = Property->GetPropertyValue(Address) ? 1 : 0;
+			OutValue = Result.GetValue();
 			return EPropertyBagResult::Success;
 		}
-		case EPropertyBagPropertyType::Byte:
-		{
-			const FByteProperty* Property = CastFieldChecked<FByteProperty>(InProperty);
-			OutValue = Property->GetPropertyValue(Address);
-			return EPropertyBagResult::Success;
-		}
-		case EPropertyBagPropertyType::Int32:
-		{
-			const FIntProperty* Property = CastFieldChecked<FIntProperty>(InProperty);
-			OutValue = static_cast<uint32>(Property->GetPropertyValue(Address));
-			return EPropertyBagResult::Success;
-		}
-		case EPropertyBagPropertyType::UInt32:
-		{
-			const FUInt32Property* Property = CastFieldChecked<FUInt32Property>(InProperty);
-			OutValue = Property->GetPropertyValue(Address);
-			return EPropertyBagResult::Success;
-		}
-		case EPropertyBagPropertyType::Int64:
-		{
-			const FInt64Property* Property = CastFieldChecked<FInt64Property>(InProperty);
-			OutValue = static_cast<uint64>(Property->GetPropertyValue(Address));
-			return EPropertyBagResult::Success;
-		}
-		case EPropertyBagPropertyType::UInt64:
-		{
-			const FUInt64Property* Property = CastFieldChecked<FUInt64Property>(InProperty);
-			OutValue = Property->GetPropertyValue(Address);
-			return EPropertyBagResult::Success;
-		}
-		case EPropertyBagPropertyType::Float:
-		{
-			const FFloatProperty* Property = CastFieldChecked<FFloatProperty>(InProperty);
-			OutValue = static_cast<uint64>(Property->GetPropertyValue(Address));
-			return EPropertyBagResult::Success;
-		}
-		case EPropertyBagPropertyType::Double:
-		{
-			const FDoubleProperty* Property = CastFieldChecked<FDoubleProperty>(InProperty);
-			OutValue = static_cast<uint64>(Property->GetPropertyValue(Address));
-			return EPropertyBagResult::Success;
-		}
-		case EPropertyBagPropertyType::Enum:
-		{
-			if (const FEnumProperty* EnumProperty = CastField<FEnumProperty>(InProperty))
-			{
-				const FNumericProperty* UnderlyingProperty = EnumProperty->GetUnderlyingProperty();
-				check(UnderlyingProperty);
-				OutValue = UnderlyingProperty->GetUnsignedIntPropertyValue(Address);
-				return EPropertyBagResult::Success;
-			}
-			else if (const FByteProperty* ByteProperty = CastField<FByteProperty>(InProperty))
-			{
-				if (ByteProperty->Enum != nullptr)
-				{
-					OutValue = ByteProperty->GetPropertyValue(Address);
-					return EPropertyBagResult::Success;
-				}
-			}
-			return EPropertyBagResult::TypeMismatch;
-		}
-		default:
-			return EPropertyBagResult::TypeMismatch;
-		}
+		return EPropertyBagResult::TypeMismatch;
 	}
 
 	EPropertyBagResult GetPropertyAsDouble(const FProperty* InProperty, EPropertyBagPropertyType InValueType, const void* Address, double& OutValue)
 	{
-		check(Address != nullptr);
-		check(InProperty != nullptr);
-
-		switch (InValueType)
-		{
-		case EPropertyBagPropertyType::Bool:
-		{
-			const FBoolProperty* Property = CastFieldChecked<FBoolProperty>(InProperty);
-			OutValue = Property->GetPropertyValue(Address) ? 1.0 : 0.0;
-			return EPropertyBagResult::Success;
-		}
-		case EPropertyBagPropertyType::Byte:
-		{
-			const FByteProperty* Property = CastFieldChecked<FByteProperty>(InProperty);
-			OutValue = Property->GetPropertyValue(Address);
-			return EPropertyBagResult::Success;
-		}
-		case EPropertyBagPropertyType::Int32:
-		{
-			const FIntProperty* Property = CastFieldChecked<FIntProperty>(InProperty);
-			OutValue = Property->GetPropertyValue(Address);
-			return EPropertyBagResult::Success;
-		}
-		case EPropertyBagPropertyType::UInt32:
-		{
-			const FUInt32Property* Property = CastFieldChecked<FUInt32Property>(InProperty);
-			OutValue = Property->GetPropertyValue(Address);
-			return EPropertyBagResult::Success;
-		}
-		case EPropertyBagPropertyType::Int64:
+		TValueOrError<double, void> Result = UE::StructUtils::GetPropertyValueAsDouble(InProperty, Address);
+		if (Result.HasValue())
 		{
-			const FInt64Property* Property = CastFieldChecked<FInt64Property>(InProperty);
-			OutValue = static_cast<double>(Property->GetPropertyValue(Address));
+			OutValue = Result.GetValue();
 			return EPropertyBagResult::Success;
 		}
-		case EPropertyBagPropertyType::UInt64:
-		{
-			const FUInt64Property* Property = CastFieldChecked<FUInt64Property>(InProperty);
-			OutValue = static_cast<double>(Property->GetPropertyValue(Address));
-			return EPropertyBagResult::Success;
-		}
-		case EPropertyBagPropertyType::Float:
-		{
-			const FFloatProperty* Property = CastFieldChecked<FFloatProperty>(InProperty);
-			OutValue = Property->GetPropertyValue(Address);
-			return EPropertyBagResult::Success;
-		}
-		case EPropertyBagPropertyType::Double:
-		{
-			const FDoubleProperty* Property = CastFieldChecked<FDoubleProperty>(InProperty);
-			OutValue = Property->GetPropertyValue(Address);
-			return EPropertyBagResult::Success;
-		}
-		case EPropertyBagPropertyType::Enum:
-		{
-			if (const FEnumProperty* EnumProperty = CastField<FEnumProperty>(InProperty))
-			{
-				const FNumericProperty* UnderlyingProperty = EnumProperty->GetUnderlyingProperty();
-				check(UnderlyingProperty);
-				OutValue = static_cast<double>(UnderlyingProperty->GetSignedIntPropertyValue(Address));
-				return EPropertyBagResult::Success;
-			}
-			else if (const FByteProperty* ByteProperty = CastField<FByteProperty>(InProperty))
-			{
-				if (ByteProperty->Enum != nullptr)
-				{
-					OutValue = ByteProperty->GetPropertyValue(Address);
-					return EPropertyBagResult::Success;
-				}
-			}
-			return EPropertyBagResult::TypeMismatch;
-		}
-		default:
-			return EPropertyBagResult::TypeMismatch;
-		}
+		return EPropertyBagResult::TypeMismatch;
 	}
 
 	// Generic property getter. Used for FName, FString, FText. 
diff --git a/Engine/Plugins/MovieScene/MovieRenderPipeline/Source/MovieRenderPipelineCore/Private/Graph/MovieGraphValueUtils.cpp b/Engine/Plugins/MovieScene/MovieRenderPipeline/Source/MovieRenderPipelineCore/Private/Graph/MovieGraphValueUtils.cpp
--- a/Engine/Plugins/MovieScene/MovieRenderPipeline/Source/MovieRenderPipelineCore/Private/Graph/MovieGraphValueUtils.cpp
+++ b/Engine/Plugins/MovieScene/MovieRenderPipeline/Source/MovieRenderPipelineCore/Private/Graph/MovieGraphValueUtils.cpp
@@ -1,4 +1,4 @@
-﻿// Copyright Epic Games, Inc. All Rights Reserved.
+// Copyright Epic Games, Inc. All Rights Reserved.
 
 #include "Graph/MovieGraphValueUtils.h"
 
@@ -103,6 +103,18 @@ bool UE::MovieGraph::ValueUtils::GetTypesForFProperty(
 			OutValueTypeObject = CastField<FSoftClassProperty>(InProperty)->MetaClass;
 		}
 	}
+	else if (InProperty->IsA<FInt8Property>())
+	{
+		OutValueType = EMovieGraphValueType::Int8;
+	}
+	else if (InProperty->IsA<FInt16Property>())
+	{
+		OutValueType = EMovieGraphValueType::Int16;
+	}
+	else if (InProperty->IsA<FUInt16Property>())
+	{
+		OutValueType = EMovieGraphValueType::UInt16;
+	}
 	else if (InProperty->IsA<FUInt32Property>())
 	{
 		OutValueType = EMovieGraphValueType::UInt32;
diff --git a/Engine/Plugins/MovieScene/MovieRenderPipeline/Source/MovieRenderPipelineCore/Public/Graph/MovieGraphCommon.h b/Engine/Plugins/MovieScene/MovieRenderPipeline/Source/MovieRenderPipelineCore/Public/Graph/MovieGraphCommon.h
--- a/Engine/Plugins/MovieScene/MovieRenderPipeline/Source/MovieRenderPipelineCore/Public/Graph/MovieGraphCommon.h
+++ b/Engine/Plugins/MovieScene/MovieRenderPipeline/Source/MovieRenderPipelineCore/Public/Graph/MovieGraphCommon.h
@@ -28,6 +28,9 @@ enum class EMovieGraphValueType : uint8
 	SoftObject,
 	Class,
 	SoftClass,
+	Int8,	// Type not fully supported at UI, will work with restrictions to type editing
+	Int16,	// Type not fully supported at UI, will work with restrictions to type editing
+	UInt16,	// Type not fully supported at UI, will work with restrictions to type editing
 	UInt32,	// Type not fully supported at UI, will work with restrictions to type editing
 	UInt64, // Type not fully supported at UI, will work with restrictions to type editing
 
diff --git a/Engine/Source/Editor/StructUtilsEditor/Private/Customizations/PropertyBagDetails.cpp b/Engine/Source/Editor/StructUtilsEditor/Private/Customizations/PropertyBagDetails.cpp
--- a/Engine/Source/Editor/StructUtilsEditor/Private/Customizations/PropertyBagDetails.cpp
+++ b/Engine/Source/Editor/StructUtilsEditor/Private/Customizations/PropertyBagDetails.cpp
@@ -257,6 +257,7 @@ FEdGraphPinType GetPropertyDescAsPin(const FPropertyBagPropertyDesc& Desc)
 			OutCategory = UEdGraphSchema_K2::PC_Text;
 			break;
 		case EPropertyBagPropertyType::Enum:
+			//Always use the byte type. The compiler will generate the proper FByteProperty or FEnumProperty
 			OutCategory = UEdGraphSchema_K2::PC_Byte;
 			OutSubCategoryObject = const_cast<UObject*>(TypeObject.Get());
 			break;
@@ -280,10 +281,19 @@ FEdGraphPinType GetPropertyDescAsPin(const FPropertyBagPropertyDesc& Desc)
 			OutCategory = UEdGraphSchema_K2::PC_SoftClass;
 			OutSubCategoryObject = const_cast<UObject*>(TypeObject.Get());
 			break;
-		case EPropertyBagPropertyType::UInt32:	// Warning : Type only partially supported (Blueprint does not support unsigned type)
+		case EPropertyBagPropertyType::Int8:	// Warning : Type only partially supported (Blueprint does not support the int8 type)
 			OutCategory = UEdGraphSchema_K2::PC_Int;
 			break;
-		case EPropertyBagPropertyType::UInt64:	// Warning : Type only partially supported (Blueprint does not support unsigned type)
+		case EPropertyBagPropertyType::Int16:	// Warning : Type only partially supported (Blueprint does not support the int16 type)
+			OutCategory = UEdGraphSchema_K2::PC_Int;
+			break;
+		case EPropertyBagPropertyType::UInt16:	// Warning : Type only partially supported (Blueprint does not support the uint16 type)
+			OutCategory = UEdGraphSchema_K2::PC_Int;
+			break;
+		case EPropertyBagPropertyType::UInt32:	// Warning : Type only partially supported (Blueprint does not support the uint32 type)
+			OutCategory = UEdGraphSchema_K2::PC_Int;
+			break;
+		case EPropertyBagPropertyType::UInt64:	// Warning : Type only partially supported (Blueprint does not support the uint64 type)
 			OutCategory = UEdGraphSchema_K2::PC_Int64;
 			break;
 		default:
@@ -305,6 +315,42 @@ FEdGraphPinType GetPropertyDescAsPin(const FPropertyBagPropertyDesc& Desc)
 	return PinType;
 }
 
+bool CanCreateValidGraphPinTypeForPropertyDesc(const FPropertyBagPropertyDesc& Desc)
+{
+	auto IsValidType = [](EPropertyBagPropertyType Type, const TObjectPtr<const UObject>& TypeObject)
+		{
+			switch (Type)
+			{
+			case EPropertyBagPropertyType::Int8:
+			case EPropertyBagPropertyType::Int16:
+			case EPropertyBagPropertyType::UInt16:
+			case EPropertyBagPropertyType::UInt32:
+				return false;
+			case EPropertyBagPropertyType::Enum:
+			case EPropertyBagPropertyType::Struct:
+			case EPropertyBagPropertyType::Object:
+			case EPropertyBagPropertyType::SoftObject:
+			case EPropertyBagPropertyType::Class:
+			case EPropertyBagPropertyType::SoftClass:
+				return TypeObject != nullptr;
+			default:
+				return true;
+			}
+		};
+
+	if (Desc.ContainerTypes.Num() > 1)
+	{
+		return false;
+	}
+
+	if (Desc.KeyType != EPropertyBagPropertyType::None && !IsValidType(Desc.KeyType, Desc.KeyTypeObject))
+	{
+		return false;
+	}
+
+	return IsValidType(Desc.ValueType, Desc.ValueTypeObject);
+}
+
 namespace Private
 {
 /** @return true if the property is one of the known missing types. */
diff --git a/Engine/Source/Editor/StructUtilsEditor/Public/PropertyBagDetails.h b/Engine/Source/Editor/StructUtilsEditor/Public/PropertyBagDetails.h
--- a/Engine/Source/Editor/StructUtilsEditor/Public/PropertyBagDetails.h
+++ b/Engine/Source/Editor/StructUtilsEditor/Public/PropertyBagDetails.h
@@ -81,6 +81,7 @@ namespace UE::StructUtils
 {
 	STRUCTUTILSEDITOR_API void SetPropertyDescFromPin(FPropertyBagPropertyDesc& Desc, const FEdGraphPinType& PinType);
 	STRUCTUTILSEDITOR_API FEdGraphPinType GetPropertyDescAsPin(const FPropertyBagPropertyDesc& Desc);
+	STRUCTUTILSEDITOR_API bool CanCreateValidGraphPinTypeForPropertyDesc(const FPropertyBagPropertyDesc& Desc);
 	/** Creates type selection pill widget. */
 	STRUCTUTILSEDITOR_API TSharedRef<SWidget> CreateTypeSelectionWidget(TSharedPtr<IPropertyHandle> ChildPropertyHandle, const TSharedPtr<IPropertyHandle>& InBagStructProperty, const TSharedPtr<IPropertyUtilities>& InPropUtils, ESelectorType SelectorType, bool bAllowContainers = true, TSubclassOf<UEdGraphSchema> PropertyBagSchemaClass = nullptr);
 
diff --git a/Engine/Source/Runtime/CoreUObject/Private/StructUtils/PropertyBag.cpp b/Engine/Source/Runtime/CoreUObject/Private/StructUtils/PropertyBag.cpp
--- a/Engine/Source/Runtime/CoreUObject/Private/StructUtils/PropertyBag.cpp
+++ b/Engine/Source/Runtime/CoreUObject/Private/StructUtils/PropertyBag.cpp
@@ -88,6 +88,22 @@ namespace UE::StructUtils::Private
 		}
 	}
 
+	bool IsNumericIntegralType(const EPropertyBagPropertyType ValueType)
+	{
+		switch (ValueType)
+		{
+		case EPropertyBagPropertyType::Int8: return true;
+		case EPropertyBagPropertyType::Byte: return true;
+		case EPropertyBagPropertyType::Int16: return true;
+		case EPropertyBagPropertyType::UInt16: return true;
+		case EPropertyBagPropertyType::Int32: return true;
+		case EPropertyBagPropertyType::UInt32: return true;
+		case EPropertyBagPropertyType::Int64: return true;
+		case EPropertyBagPropertyType::UInt64: return true;
+		default: return false;
+		}
+	}
+
 	FPropertyBagContainerTypes GetContainerTypesFromProperty(const FProperty* InSourceProperty)
 	{
 		FPropertyBagContainerTypes ContainerTypes;
@@ -143,33 +159,49 @@ namespace UE::StructUtils::Private
 		{
 			return EPropertyBagPropertyType::Bool;
 		}
-		if (const FByteProperty* ByteProperty = CastField<FByteProperty>(InSourceProperty))
-		{
-			return ByteProperty->IsEnum() ? EPropertyBagPropertyType::Enum : EPropertyBagPropertyType::Byte;
-		}
-		if (CastField<FIntProperty>(InSourceProperty))
-		{
-			return EPropertyBagPropertyType::Int32;
-		}
-		if (CastField<FUInt32Property>(InSourceProperty))
-		{
-			return EPropertyBagPropertyType::UInt32;
-		}
-		if (CastField<FInt64Property>(InSourceProperty))
+		if (CastField<FNumericProperty>(InSourceProperty))
 		{
-			return EPropertyBagPropertyType::Int64;
-		}
-		if (CastField<FUInt64Property>(InSourceProperty))
-		{
-			return EPropertyBagPropertyType::UInt64;
-		}
-		if (CastField<FFloatProperty>(InSourceProperty))
-		{
-			return EPropertyBagPropertyType::Float;
-		}
-		if (CastField<FDoubleProperty>(InSourceProperty))
-		{
-			return EPropertyBagPropertyType::Double;
+			if (const FInt8Property* Int8Property = CastField<FInt8Property>(InSourceProperty))
+			{
+				return EPropertyBagPropertyType::Int8;
+			}
+			if (const FByteProperty* ByteProperty = CastField<FByteProperty>(InSourceProperty))
+			{
+				return ByteProperty->IsEnum() ? EPropertyBagPropertyType::Enum : EPropertyBagPropertyType::Byte;
+			}
+			if (CastField<FInt16Property>(InSourceProperty))
+			{
+				return EPropertyBagPropertyType::Int16;
+			}
+			if (CastField<FUInt16Property>(InSourceProperty))
+			{
+				return EPropertyBagPropertyType::UInt16;
+			}
+			if (CastField<FIntProperty>(InSourceProperty))
+			{
+				return EPropertyBagPropertyType::Int32;
+			}
+			if (CastField<FUInt32Property>(InSourceProperty))
+			{
+				return EPropertyBagPropertyType::UInt32;
+			}
+			if (CastField<FInt64Property>(InSourceProperty))
+			{
+				return EPropertyBagPropertyType::Int64;
+			}
+			if (CastField<FUInt64Property>(InSourceProperty))
+			{
+				return EPropertyBagPropertyType::UInt64;
+			}
+			if (CastField<FFloatProperty>(InSourceProperty))
+			{
+				return EPropertyBagPropertyType::Float;
+			}
+			if (CastField<FDoubleProperty>(InSourceProperty))
+			{
+				return EPropertyBagPropertyType::Double;
+			}
+			return EPropertyBagPropertyType::None;
 		}
 		if (CastField<FNameProperty>(InSourceProperty))
 		{
@@ -304,6 +336,60 @@ namespace UE::StructUtils::Private
 		return nullptr;
 	}
 
+	EPropertyBagPropertyType EnumUnderlyingTypeToPropertyBagPropertyType(const UEnum::EUnderlyingType UnderlyingType)
+	{
+		switch (UnderlyingType)
+		{
+			case UEnum::EUnderlyingType::uint8: return EPropertyBagPropertyType::Byte;
+			case UEnum::EUnderlyingType::int8: return EPropertyBagPropertyType::Int8;
+			case UEnum::EUnderlyingType::uint16: return EPropertyBagPropertyType::UInt16;
+			case UEnum::EUnderlyingType::int16: return EPropertyBagPropertyType::Int16;
+			case UEnum::EUnderlyingType::uint32: return EPropertyBagPropertyType::UInt32;
+			case UEnum::EUnderlyingType::int32: return EPropertyBagPropertyType::Int32;
+			case UEnum::EUnderlyingType::uint64: return EPropertyBagPropertyType::UInt64;
+			case UEnum::EUnderlyingType::int64: return EPropertyBagPropertyType::Int64;
+		}
+		return EPropertyBagPropertyType::None;
+	}
+
+	FNumericProperty* CreateIntegerPropertyFromDesc(const EPropertyBagPropertyType ValueType, const FFieldVariant PropertyScope, const FName PropertyName)
+	{
+		FNumericProperty* Prop = nullptr;
+		switch (ValueType)
+		{
+		case EPropertyBagPropertyType::Int8:
+			Prop = new FInt8Property(PropertyScope, PropertyName, RF_Public);
+			break;
+		case EPropertyBagPropertyType::Byte:
+			Prop = new FByteProperty(PropertyScope, PropertyName, RF_Public);
+			break;
+		case EPropertyBagPropertyType::Int16:
+			Prop = new FInt16Property(PropertyScope, PropertyName, RF_Public);
+			break;
+		case EPropertyBagPropertyType::UInt16:
+			Prop = new FUInt16Property(PropertyScope, PropertyName, RF_Public);
+			break;
+		case EPropertyBagPropertyType::Int32:
+			Prop = new FIntProperty(PropertyScope, PropertyName, RF_Public);
+			break;
+		case EPropertyBagPropertyType::UInt32:
+			Prop = new FUInt32Property(PropertyScope, PropertyName, RF_Public);
+			break;
+		case EPropertyBagPropertyType::Int64:
+			Prop = new FInt64Property(PropertyScope, PropertyName, RF_Public);
+			break;
+		case EPropertyBagPropertyType::UInt64:
+			Prop = new FUInt64Property(PropertyScope, PropertyName, RF_Public);
+			break;
+		}
+
+		if (Prop != nullptr)
+		{
+			Prop->SetPropertyFlags(CPF_HasGetValueTypeHash);
+		}
+		return Prop;
+	}
+
 	FProperty* CreatePropertyFromDesc(const FPropertyBagPropertyDesc& Desc, const FFieldVariant PropertyScope)
 	{
 		// Handle array and nested containers properties
@@ -375,37 +461,16 @@ namespace UE::StructUtils::Private
 				FBoolProperty* Prop = new FBoolProperty(PropertyScope, Desc.Name, RF_Public);
 				Prop->SetBoolSize(sizeof(bool), true); // Enable native access (init the whole byte, rather than just first bit)
 				return Prop;
-			}
+		}
+		case EPropertyBagPropertyType::Int8:
 		case EPropertyBagPropertyType::Byte:
-			{
-				FByteProperty* Prop = new FByteProperty(PropertyScope, Desc.Name, RF_Public);
-				Prop->SetPropertyFlags(CPF_HasGetValueTypeHash);
-				return Prop;
-			}
+		case EPropertyBagPropertyType::Int16:
+		case EPropertyBagPropertyType::UInt16:
 		case EPropertyBagPropertyType::Int32:
-			{
-				FIntProperty* Prop = new FIntProperty(PropertyScope, Desc.Name, RF_Public);
-				Prop->SetPropertyFlags(CPF_HasGetValueTypeHash);
-				return Prop;
-			}
 		case EPropertyBagPropertyType::UInt32:
-			{
-				FUInt32Property* Prop = new FUInt32Property(PropertyScope, Desc.Name, RF_Public);
-				Prop->SetPropertyFlags(CPF_HasGetValueTypeHash);
-				return Prop;
-			}
 		case EPropertyBagPropertyType::Int64:
-			{
-				FInt64Property* Prop = new FInt64Property(PropertyScope, Desc.Name, RF_Public);
-				Prop->SetPropertyFlags(CPF_HasGetValueTypeHash);
-				return Prop;
-			}
 		case EPropertyBagPropertyType::UInt64:
-			{
-				FUInt64Property* Prop = new FUInt64Property(PropertyScope, Desc.Name, RF_Public);
-				Prop->SetPropertyFlags(CPF_HasGetValueTypeHash);
-				return Prop;
-			}
+			return CreateIntegerPropertyFromDesc(Desc.ValueType, PropertyScope, Desc.Name);
 		case EPropertyBagPropertyType::Float:
 			{
 				FFloatProperty* Prop = new FFloatProperty(PropertyScope, Desc.Name, RF_Public);
@@ -441,12 +506,18 @@ namespace UE::StructUtils::Private
 				// Only generate FEnumProperty for cpp enum class like UHT and Blueprint.
 				if (Enum->GetCppForm() == UEnum::ECppForm::EnumClass)
 				{
-					FEnumProperty* Prop = new FEnumProperty(PropertyScope, Desc.Name, RF_Public);
-					Prop->SetPropertyFlags(CPF_HasGetValueTypeHash);
-					FNumericProperty* UnderlyingProp = new FByteProperty(Prop, "UnderlyingType", RF_Public); // HACK: Hardwire to byte property for now for BP compatibility
-					Prop->SetEnum(const_cast<UEnum*>(Enum));
-					Prop->AddCppProperty(UnderlyingProp);
-					return Prop;
+					const EPropertyBagPropertyType UnderlyingType = EnumUnderlyingTypeToPropertyBagPropertyType(Enum->GetUnderlyingType());
+					if (IsNumericIntegralType(UnderlyingType))
+					{
+						FEnumProperty* Prop = new FEnumProperty(PropertyScope, Desc.Name, RF_Public);
+						Prop->SetPropertyFlags(CPF_HasGetValueTypeHash);
+						Prop->SetEnum(const_cast<UEnum*>(Enum));
+						static FName NAME_UnderlyingType = "UnderlyingType";
+						FNumericProperty* UnderlyingProp = CreateIntegerPropertyFromDesc(UnderlyingType, Prop, NAME_UnderlyingType);
+						checkf(UnderlyingProp, TEXT("The type was checked in IsNumericIntegralType."));
+						Prop->AddCppProperty(UnderlyingProp);
+						return Prop;
+					}
 				}
 				else
 				{
@@ -558,78 +629,13 @@ namespace UE::StructUtils::Private
 			return EPropertyBagResult::TypeMismatch;
 		}
 
-		switch(Desc->ValueType)
+		TValueOrError<int64, void> Result = UE::StructUtils::GetPropertyValueAsInt64(Desc->CachedProperty, Address);
+		if (Result.HasValue())
 		{
-		case EPropertyBagPropertyType::Bool:
-			{
-				const FBoolProperty* Property = CastFieldChecked<FBoolProperty>(Desc->CachedProperty);
-				OutValue = Property->GetPropertyValue(Address) ? 1 : 0;
-				return EPropertyBagResult::Success;
-			}
-		case EPropertyBagPropertyType::Byte:
-			{
-				const FByteProperty* Property = CastFieldChecked<FByteProperty>(Desc->CachedProperty);
-				OutValue = Property->GetPropertyValue(Address);
-				return EPropertyBagResult::Success;
-			}
-		case EPropertyBagPropertyType::Int32:
-			{
-				const FIntProperty* Property = CastFieldChecked<FIntProperty>(Desc->CachedProperty);
-				OutValue = Property->GetPropertyValue(Address);
-				return EPropertyBagResult::Success;
-			}
-		case EPropertyBagPropertyType::UInt32:
-			{
-				const FUInt32Property* Property = CastFieldChecked<FUInt32Property>(Desc->CachedProperty);
-				OutValue = static_cast<uint32>(Property->GetPropertyValue(Address));
-				return EPropertyBagResult::Success;
-			}
-		case EPropertyBagPropertyType::Int64:
-			{
-				const FInt64Property* Property = CastFieldChecked<FInt64Property>(Desc->CachedProperty);
-				OutValue = Property->GetPropertyValue(Address);
-				return EPropertyBagResult::Success;
-			}
-		case EPropertyBagPropertyType::UInt64:
-			{
-				const FUInt64Property* Property = CastFieldChecked<FUInt64Property>(Desc->CachedProperty);
-				OutValue = static_cast<uint64>(Property->GetPropertyValue(Address));
-				return EPropertyBagResult::Success;
-			}
-		case EPropertyBagPropertyType::Float:
-			{
-				const FFloatProperty* Property = CastFieldChecked<FFloatProperty>(Desc->CachedProperty);
-				OutValue = static_cast<int64>(Property->GetPropertyValue(Address));
-				return EPropertyBagResult::Success;
-			}
-		case EPropertyBagPropertyType::Double:
-			{
-				const FDoubleProperty* Property = CastFieldChecked<FDoubleProperty>(Desc->CachedProperty);
-				OutValue = static_cast<int64>(Property->GetPropertyValue(Address));
-				return EPropertyBagResult::Success;
-			}
-		case EPropertyBagPropertyType::Enum:
-			{
-				if (const FEnumProperty* EnumProperty = CastField<FEnumProperty>(Desc->CachedProperty))
-				{
-					const FNumericProperty* UnderlyingProperty = EnumProperty->GetUnderlyingProperty();
-					check(UnderlyingProperty);
-					OutValue = UnderlyingProperty->GetSignedIntPropertyValue(Address);
-					return EPropertyBagResult::Success;
-				}
-				else if (const FByteProperty* ByteProperty = CastField<FByteProperty>(Desc->CachedProperty))
-				{
-					if (ByteProperty->Enum != nullptr)
-					{
-						OutValue = ByteProperty->GetPropertyValue(Address);
-						return EPropertyBagResult::Success;
-					}
-				}
-				return EPropertyBagResult::TypeMismatch;
-			}
-		default:
-			return EPropertyBagResult::TypeMismatch;
+			OutValue = Result.GetValue();
+			return EPropertyBagResult::Success;
 		}
+		return EPropertyBagResult::TypeMismatch;
 	}
 
 	EPropertyBagResult GetPropertyAsUInt64(const FPropertyBagPropertyDesc* Desc, const void* Address, uint64& OutValue)
@@ -647,78 +653,13 @@ namespace UE::StructUtils::Private
 			return EPropertyBagResult::TypeMismatch;
 		}
 
-		switch (Desc->ValueType)
-		{
-		case EPropertyBagPropertyType::Bool:
-		{
-			const FBoolProperty* Property = CastFieldChecked<FBoolProperty>(Desc->CachedProperty);
-			OutValue = Property->GetPropertyValue(Address) ? 1 : 0;
-			return EPropertyBagResult::Success;
-		}
-		case EPropertyBagPropertyType::Byte:
-		{
-			const FByteProperty* Property = CastFieldChecked<FByteProperty>(Desc->CachedProperty);
-			OutValue = Property->GetPropertyValue(Address);
-			return EPropertyBagResult::Success;
-		}
-		case EPropertyBagPropertyType::Int32:
+		TValueOrError<uint64, void> Result = UE::StructUtils::GetPropertyValueAsUInt64(Desc->CachedProperty, Address);
+		if (Result.HasValue())
 		{
-			const FIntProperty* Property = CastFieldChecked<FIntProperty>(Desc->CachedProperty);
-			OutValue = static_cast<uint32>(Property->GetPropertyValue(Address));
+			OutValue = Result.GetValue();
 			return EPropertyBagResult::Success;
 		}
-		case EPropertyBagPropertyType::UInt32:
-		{
-			const FUInt32Property* Property = CastFieldChecked<FUInt32Property>(Desc->CachedProperty);
-			OutValue = Property->GetPropertyValue(Address);
-			return EPropertyBagResult::Success;
-		}
-		case EPropertyBagPropertyType::Int64:
-		{
-			const FInt64Property* Property = CastFieldChecked<FInt64Property>(Desc->CachedProperty);
-			OutValue = static_cast<uint64>(Property->GetPropertyValue(Address));
-			return EPropertyBagResult::Success;
-		}
-		case EPropertyBagPropertyType::UInt64:
-		{
-			const FUInt64Property* Property = CastFieldChecked<FUInt64Property>(Desc->CachedProperty);
-			OutValue = Property->GetPropertyValue(Address);
-			return EPropertyBagResult::Success;
-		}
-		case EPropertyBagPropertyType::Float:
-		{
-			const FFloatProperty* Property = CastFieldChecked<FFloatProperty>(Desc->CachedProperty);
-			OutValue = static_cast<uint64>(Property->GetPropertyValue(Address));
-			return EPropertyBagResult::Success;
-		}
-		case EPropertyBagPropertyType::Double:
-		{
-			const FDoubleProperty* Property = CastFieldChecked<FDoubleProperty>(Desc->CachedProperty);
-			OutValue = static_cast<uint64>(Property->GetPropertyValue(Address));
-			return EPropertyBagResult::Success;
-		}
-		case EPropertyBagPropertyType::Enum:
-		{
-			if (const FEnumProperty* EnumProperty = CastField<FEnumProperty>(Desc->CachedProperty))
-			{
-				const FNumericProperty* UnderlyingProperty = EnumProperty->GetUnderlyingProperty();
-				check(UnderlyingProperty);
-				OutValue = UnderlyingProperty->GetUnsignedIntPropertyValue(Address);
-				return EPropertyBagResult::Success;
-			}
-			else if (const FByteProperty* ByteProperty = CastField<FByteProperty>(Desc->CachedProperty))
-			{
-				if (ByteProperty->Enum != nullptr)
-				{
-					OutValue = ByteProperty->GetPropertyValue(Address);
-					return EPropertyBagResult::Success;
-				}
-			}
-			return EPropertyBagResult::TypeMismatch;
-		}
-		default:
-			return EPropertyBagResult::TypeMismatch;
-		}
+		return EPropertyBagResult::TypeMismatch;
 	}
 
 	EPropertyBagResult GetPropertyAsDouble(const FPropertyBagPropertyDesc* Desc, const void* Address, double& OutValue)
@@ -736,78 +677,13 @@ namespace UE::StructUtils::Private
 			return EPropertyBagResult::TypeMismatch;
 		}
 
-		switch(Desc->ValueType)
+		TValueOrError<double, void> Result = UE::StructUtils::GetPropertyValueAsDouble(Desc->CachedProperty, Address);
+		if (Result.HasValue())
 		{
-		case EPropertyBagPropertyType::Bool:
-			{
-				const FBoolProperty* Property = CastFieldChecked<FBoolProperty>(Desc->CachedProperty);
-				OutValue = Property->GetPropertyValue(Address) ? 1.0 : 0.0;
-				return EPropertyBagResult::Success;
-			}
-		case EPropertyBagPropertyType::Byte:
-			{
-				const FByteProperty* Property = CastFieldChecked<FByteProperty>(Desc->CachedProperty);
-				OutValue = Property->GetPropertyValue(Address);
-				return EPropertyBagResult::Success;
-			}
-		case EPropertyBagPropertyType::Int32:
-			{
-				const FIntProperty* Property = CastFieldChecked<FIntProperty>(Desc->CachedProperty);
-				OutValue = Property->GetPropertyValue(Address);
-				return EPropertyBagResult::Success;
-			}
-		case EPropertyBagPropertyType::UInt32:
-			{
-				const FUInt32Property* Property = CastFieldChecked<FUInt32Property>(Desc->CachedProperty);
-				OutValue = Property->GetPropertyValue(Address);
-				return EPropertyBagResult::Success;
-			}
-		case EPropertyBagPropertyType::Int64:
-			{
-				const FInt64Property* Property = CastFieldChecked<FInt64Property>(Desc->CachedProperty);
-				OutValue = static_cast<double>(Property->GetPropertyValue(Address));
-				return EPropertyBagResult::Success;
-			}
-		case EPropertyBagPropertyType::UInt64:
-			{
-				const FUInt64Property* Property = CastFieldChecked<FUInt64Property>(Desc->CachedProperty);
-				OutValue = static_cast<double>(Property->GetPropertyValue(Address));
-				return EPropertyBagResult::Success;
-			}
-		case EPropertyBagPropertyType::Float:
-			{
-				const FFloatProperty* Property = CastFieldChecked<FFloatProperty>(Desc->CachedProperty);
-				OutValue = Property->GetPropertyValue(Address);
-				return EPropertyBagResult::Success;
-			}
-		case EPropertyBagPropertyType::Double:
-			{
-				const FDoubleProperty* Property = CastFieldChecked<FDoubleProperty>(Desc->CachedProperty);
-				OutValue = Property->GetPropertyValue(Address);
-				return EPropertyBagResult::Success;
-			}
-		case EPropertyBagPropertyType::Enum:
-			{
-				if (const FEnumProperty* EnumProperty = CastField<FEnumProperty>(Desc->CachedProperty))
-				{
-					const FNumericProperty* UnderlyingProperty = EnumProperty->GetUnderlyingProperty();
-					check(UnderlyingProperty);
-					OutValue = static_cast<double>(UnderlyingProperty->GetSignedIntPropertyValue(Address));
-					return EPropertyBagResult::Success;
-				}
-				else if (const FByteProperty* ByteProperty = CastField<FByteProperty>(Desc->CachedProperty))
-				{
-					if (ByteProperty->Enum != nullptr)
-					{
-						OutValue = ByteProperty->GetPropertyValue(Address);
-						return EPropertyBagResult::Success;
-					}
-				}
-				return EPropertyBagResult::TypeMismatch;
-			}
-		default:
-			return EPropertyBagResult::TypeMismatch;
+			OutValue = Result.GetValue();
+			return EPropertyBagResult::Success;
 		}
+		return EPropertyBagResult::TypeMismatch;
 	}
 
 	// Generic property getter. Used for FName, FString, FText. 
@@ -838,7 +714,7 @@ namespace UE::StructUtils::Private
 		return EPropertyBagResult::Success;
 	}
 
-	EPropertyBagResult GetPropertyValueAsEnum(const FPropertyBagPropertyDesc* Desc, const UEnum* RequestedEnum, const void* Address, uint8& OutValue)
+	EPropertyBagResult GetPropertyValueAsEnum(const FPropertyBagPropertyDesc* Desc, const UEnum* RequestedEnum, const void* Address, int64& OutValue)
 	{
 		if (Desc == nullptr || Desc->CachedProperty == nullptr)
     	{
@@ -867,7 +743,7 @@ namespace UE::StructUtils::Private
 				return EPropertyBagResult::TypeMismatch;
 			}
 	
-			OutValue = static_cast<uint8>(UnderlyingProperty->GetUnsignedIntPropertyValue(Address));
+			OutValue = UnderlyingProperty->GetSignedIntPropertyValue(Address);
 		}
 		else if (const FByteProperty* ByteProperty = CastField<FByteProperty>(Desc->CachedProperty))
 		{
@@ -999,78 +875,8 @@ namespace UE::StructUtils::Private
 			return EPropertyBagResult::TypeMismatch;
 		}
 		
-		switch(Desc->ValueType)
-		{
-		case EPropertyBagPropertyType::Bool:
-			{
-				const FBoolProperty* Property = CastFieldChecked<FBoolProperty>(Desc->CachedProperty);
-				Property->SetPropertyValue(Address, InValue != 0);
-				return EPropertyBagResult::Success;
-			}
-		case EPropertyBagPropertyType::Byte:
-			{
-				const FByteProperty* Property = CastFieldChecked<FByteProperty>(Desc->CachedProperty);
-				Property->SetPropertyValue(Address, static_cast<uint8>(InValue));
-				return EPropertyBagResult::Success;
-			}
-		case EPropertyBagPropertyType::Int32:
-			{
-				const FIntProperty* Property = CastFieldChecked<FIntProperty>(Desc->CachedProperty);
-				Property->SetPropertyValue(Address, static_cast<int32>(InValue));
-				return EPropertyBagResult::Success;
-			}
-		case EPropertyBagPropertyType::UInt32:
-			{
-				const FUInt32Property* Property = CastFieldChecked<FUInt32Property>(Desc->CachedProperty);
-				Property->SetPropertyValue(Address, static_cast<uint32>(InValue));
-				return EPropertyBagResult::Success;
-			}
-		case EPropertyBagPropertyType::Int64:
-			{
-				const FInt64Property* Property = CastFieldChecked<FInt64Property>(Desc->CachedProperty);
-				Property->SetPropertyValue(Address, InValue);
-				return EPropertyBagResult::Success;
-			}
-		case EPropertyBagPropertyType::UInt64:
-			{
-				const FUInt64Property* Property = CastFieldChecked<FUInt64Property>(Desc->CachedProperty);
-				Property->SetPropertyValue(Address, static_cast<uint64>(InValue));
-				return EPropertyBagResult::Success;
-			}
-		case EPropertyBagPropertyType::Float:
-			{
-				const FFloatProperty* Property = CastFieldChecked<FFloatProperty>(Desc->CachedProperty);
-				Property->SetPropertyValue(Address, static_cast<float>(InValue));
-				return EPropertyBagResult::Success;
-			}
-		case EPropertyBagPropertyType::Double:
-			{
-				const FDoubleProperty* Property = CastFieldChecked<FDoubleProperty>(Desc->CachedProperty);
-				Property->SetPropertyValue(Address, static_cast<double>(InValue));
-				return EPropertyBagResult::Success;
-			}
-		case EPropertyBagPropertyType::Enum:
-			{
-				if (const FEnumProperty* EnumProperty = CastField<FEnumProperty>(Desc->CachedProperty))
-				{
-					const FNumericProperty* UnderlyingProperty = EnumProperty->GetUnderlyingProperty();
-					check(UnderlyingProperty);
-					UnderlyingProperty->SetIntPropertyValue(Address, static_cast<uint64>(InValue));
-					return EPropertyBagResult::Success;
-				}
-				else if (const FByteProperty* ByteProperty = CastField<FByteProperty>(Desc->CachedProperty))
-				{
-					if (ByteProperty->Enum)
-					{
-						ByteProperty->SetPropertyValue(Address, static_cast<uint8>(InValue));
-						return EPropertyBagResult::Success;
-					}
-				}
-				return EPropertyBagResult::TypeMismatch;
-			}
-		default:
-			return EPropertyBagResult::TypeMismatch;
-		}
+		const bool bResult = UE::StructUtils::SetPropertyValueFromInt64(Desc->CachedProperty, Address, InValue);
+		return bResult ? EPropertyBagResult::Success : EPropertyBagResult::TypeMismatch;
 	}
 
 	EPropertyBagResult SetPropertyFromUInt64(const FPropertyBagPropertyDesc* Desc, void* Address, const uint64 InValue)
@@ -1088,78 +894,8 @@ namespace UE::StructUtils::Private
 			return EPropertyBagResult::TypeMismatch;
 		}
 
-		switch (Desc->ValueType)
-		{
-		case EPropertyBagPropertyType::Bool:
-		{
-			const FBoolProperty* Property = CastFieldChecked<FBoolProperty>(Desc->CachedProperty);
-			Property->SetPropertyValue(Address, InValue != 0);
-			return EPropertyBagResult::Success;
-		}
-		case EPropertyBagPropertyType::Byte:
-		{
-			const FByteProperty* Property = CastFieldChecked<FByteProperty>(Desc->CachedProperty);
-			Property->SetPropertyValue(Address, static_cast<uint8>(InValue));
-			return EPropertyBagResult::Success;
-		}
-		case EPropertyBagPropertyType::Int32:
-		{
-			const FIntProperty* Property = CastFieldChecked<FIntProperty>(Desc->CachedProperty);
-			Property->SetPropertyValue(Address, static_cast<int32>(InValue));
-			return EPropertyBagResult::Success;
-		}
-		case EPropertyBagPropertyType::UInt32:
-		{
-			const FUInt32Property* Property = CastFieldChecked<FUInt32Property>(Desc->CachedProperty);
-			Property->SetPropertyValue(Address, static_cast<uint32>(InValue));
-			return EPropertyBagResult::Success;
-		}
-		case EPropertyBagPropertyType::Int64:
-		{
-			const FInt64Property* Property = CastFieldChecked<FInt64Property>(Desc->CachedProperty);
-			Property->SetPropertyValue(Address, static_cast<int64>(InValue));
-			return EPropertyBagResult::Success;
-		}
-		case EPropertyBagPropertyType::UInt64:
-		{
-			const FUInt64Property* Property = CastFieldChecked<FUInt64Property>(Desc->CachedProperty);
-			Property->SetPropertyValue(Address, InValue);
-			return EPropertyBagResult::Success;
-		}
-		case EPropertyBagPropertyType::Float:
-		{
-			const FFloatProperty* Property = CastFieldChecked<FFloatProperty>(Desc->CachedProperty);
-			Property->SetPropertyValue(Address, static_cast<float>(InValue));
-			return EPropertyBagResult::Success;
-		}
-		case EPropertyBagPropertyType::Double:
-		{
-			const FDoubleProperty* Property = CastFieldChecked<FDoubleProperty>(Desc->CachedProperty);
-			Property->SetPropertyValue(Address, static_cast<double>(InValue));
-			return EPropertyBagResult::Success;
-		}
-		case EPropertyBagPropertyType::Enum:
-		{
-			if (const FEnumProperty* EnumProperty = CastField<FEnumProperty>(Desc->CachedProperty))
-			{
-				const FNumericProperty* UnderlyingProperty = EnumProperty->GetUnderlyingProperty();
-				check(UnderlyingProperty);
-				UnderlyingProperty->SetIntPropertyValue(Address, InValue);
-				return EPropertyBagResult::Success;
-			}
-			else if (const FByteProperty* ByteProperty = CastField<FByteProperty>(Desc->CachedProperty))
-			{
-				if (ByteProperty->Enum)
-				{
-					ByteProperty->SetPropertyValue(Address, static_cast<uint8>(InValue));
-					return EPropertyBagResult::Success;
-				}
-			}
-			return EPropertyBagResult::TypeMismatch;
-		}
-		default:
-			return EPropertyBagResult::TypeMismatch;
-		}
+		const bool bResult = UE::StructUtils::SetPropertyValueFromUInt64(Desc->CachedProperty, Address, InValue);
+		return bResult ? EPropertyBagResult::Success : EPropertyBagResult::TypeMismatch;
 	}
 
 	EPropertyBagResult SetPropertyFromDouble(const FPropertyBagPropertyDesc* Desc, void* Address, const double InValue)
@@ -1177,78 +913,8 @@ namespace UE::StructUtils::Private
 			return EPropertyBagResult::TypeMismatch;
 		}
 
-		switch(Desc->ValueType)
-		{
-		case EPropertyBagPropertyType::Bool:
-			{
-				const FBoolProperty* Property = CastFieldChecked<FBoolProperty>(Desc->CachedProperty);
-				Property->SetPropertyValue(Address, !(FMath::IsNearlyZero(InValue)));
-				return EPropertyBagResult::Success;
-			}
-		case EPropertyBagPropertyType::Byte:
-			{
-				const FByteProperty* Property = CastFieldChecked<FByteProperty>(Desc->CachedProperty);
-				Property->SetPropertyValue(Address, static_cast<uint8>(FMath::RoundToInt32(InValue)));
-				return EPropertyBagResult::Success;
-			}
-		case EPropertyBagPropertyType::Int32:
-			{
-				const FIntProperty* Property = CastFieldChecked<FIntProperty>(Desc->CachedProperty);
-				Property->SetPropertyValue(Address, FMath::RoundToInt32(InValue));
-				return EPropertyBagResult::Success;
-			}
-		case EPropertyBagPropertyType::UInt32:
-			{
-				const FUInt32Property* Property = CastFieldChecked<FUInt32Property>(Desc->CachedProperty);
-				Property->SetPropertyValue(Address, FMath::RoundToInt32(InValue));
-				return EPropertyBagResult::Success;
-			}
-		case EPropertyBagPropertyType::Int64:
-			{
-				const FInt64Property* Property = CastFieldChecked<FInt64Property>(Desc->CachedProperty);
-				Property->SetPropertyValue(Address, FMath::RoundToInt64(InValue));
-				return EPropertyBagResult::Success;
-			}
-		case EPropertyBagPropertyType::UInt64:
-			{
-				const FUInt64Property* Property = CastFieldChecked<FUInt64Property>(Desc->CachedProperty);
-				Property->SetPropertyValue(Address, FMath::RoundToInt64(InValue));
-				return EPropertyBagResult::Success;
-			}
-		case EPropertyBagPropertyType::Float:
-			{
-				const FFloatProperty* Property = CastFieldChecked<FFloatProperty>(Desc->CachedProperty);
-				Property->SetPropertyValue(Address, static_cast<float>(InValue));
-				return EPropertyBagResult::Success;
-			}
-		case EPropertyBagPropertyType::Double:
-			{
-				const FDoubleProperty* Property = CastFieldChecked<FDoubleProperty>(Desc->CachedProperty);
-				Property->SetPropertyValue(Address, InValue);
-				return EPropertyBagResult::Success;
-			}
-		case EPropertyBagPropertyType::Enum:
-			{
-				if (const FEnumProperty* EnumProperty = CastField<FEnumProperty>(Desc->CachedProperty))
-				{
-					const FNumericProperty* UnderlyingProperty = EnumProperty->GetUnderlyingProperty();
-					check(UnderlyingProperty);
-					UnderlyingProperty->SetIntPropertyValue(Address, static_cast<uint64>(InValue));
-					return EPropertyBagResult::Success;
-				}
-				else if (const FByteProperty* ByteProperty = CastField<FByteProperty>(Desc->CachedProperty))
-				{
-					if (ByteProperty->Enum)
-					{
-						ByteProperty->SetPropertyValue(Address, static_cast<uint8>(FMath::RoundToInt32(InValue)));
-						return EPropertyBagResult::Success;
-					}
-				}
-				return EPropertyBagResult::TypeMismatch;
-			}
-		default:
-			return EPropertyBagResult::TypeMismatch;
-		}
+		const bool bResult = UE::StructUtils::SetPropertyValueFromDouble(Desc->CachedProperty, Address, InValue);
+		return bResult ? EPropertyBagResult::Success : EPropertyBagResult::TypeMismatch;
 	}
 
 	// Generic property setter. Used for FName, FString, FText 
@@ -1279,7 +945,7 @@ namespace UE::StructUtils::Private
 		return EPropertyBagResult::Success;
 	}
 
-	EPropertyBagResult SetPropertyValueAsEnum(const FPropertyBagPropertyDesc* Desc, void* Address, const uint8 InValue, const UEnum* Enum)
+	EPropertyBagResult SetPropertyValueAsEnum(const FPropertyBagPropertyDesc* Desc, void* Address, const int64 InValue, const UEnum* Enum)
 	{
 		if (Desc == nullptr || Desc->CachedProperty == nullptr)
 		{
@@ -1308,7 +974,7 @@ namespace UE::StructUtils::Private
 				return EPropertyBagResult::TypeMismatch;
 			}
 	
-			UnderlyingProperty->SetIntPropertyValue(Address, static_cast<uint64>(InValue));
+			UnderlyingProperty->SetIntPropertyValue(Address, InValue);
 	
 			return EPropertyBagResult::Success;
 		}
@@ -1318,7 +984,7 @@ namespace UE::StructUtils::Private
 			{
 				return EPropertyBagResult::TypeMismatch;
 			}
-			ByteProperty->SetPropertyValue(Address, InValue);
+			ByteProperty->SetPropertyValue(Address, static_cast<uint8>(InValue));
 			return EPropertyBagResult::Success;
 		}
 		else
@@ -1883,7 +1549,10 @@ bool FPropertyBagPropertyDesc::IsNumericType() const
 	switch (ValueType)
 	{
 	case EPropertyBagPropertyType::Bool: return true;
+	case EPropertyBagPropertyType::Int8: return true;
 	case EPropertyBagPropertyType::Byte: return true;
+	case EPropertyBagPropertyType::Int16: return true;
+	case EPropertyBagPropertyType::UInt16: return true;
 	case EPropertyBagPropertyType::Int32: return true;
 	case EPropertyBagPropertyType::UInt32: return true;
 	case EPropertyBagPropertyType::Int64: return true;
@@ -1899,12 +1568,18 @@ bool FPropertyBagPropertyDesc::IsUnsignedNumericType() const
 	switch (ValueType)
 	{
 	case EPropertyBagPropertyType::Byte: return true;
+	case EPropertyBagPropertyType::UInt16: return true;
 	case EPropertyBagPropertyType::UInt32: return true;
 	case EPropertyBagPropertyType::UInt64: return true;
 	default: return false;
 	}
 }
 
+bool FPropertyBagPropertyDesc::IsNumericIntegralType() const
+{
+	return UE::StructUtils::Private::IsNumericIntegralType(ValueType);
+}
+
 bool FPropertyBagPropertyDesc::IsNumericFloatType() const
 {
 	switch (ValueType)
@@ -2447,88 +2122,90 @@ TValueOrError<bool, EPropertyBagResult> FInstancedPropertyBag::GetValueBool(cons
 	return MakeValue(ReturnValue != 0);
 }
 
-TValueOrError<uint8, EPropertyBagResult> FInstancedPropertyBag::GetValueByte(const FName Name) const
+template<typename T>
+TValueOrError<T, EPropertyBagResult> FInstancedPropertyBag::GetNumericalValue(const FName Name) const
 {
 	const FPropertyBagPropertyDesc* Desc = FindPropertyDescByName(Name);
-	uint64 ReturnValue = 0;
-	const EPropertyBagResult Result = UE::StructUtils::Private::GetPropertyAsUInt64(Desc, GetValueAddress(Desc), ReturnValue);
-	if (Result != EPropertyBagResult::Success)
+	if constexpr (std::is_floating_point_v<T>)
 	{
-		return MakeError(Result);
+		double ReturnValue = 0;
+		const EPropertyBagResult Result = UE::StructUtils::Private::GetPropertyAsDouble(Desc, GetValueAddress(Desc), ReturnValue);
+		if (Result != EPropertyBagResult::Success)
+		{
+			return MakeError(Result);
+		}
+		return MakeValue(static_cast<T>(ReturnValue));
+	}
+	else if constexpr (std::is_unsigned_v<T>)
+	{
+		uint64 ReturnValue = 0;
+		const EPropertyBagResult Result = UE::StructUtils::Private::GetPropertyAsUInt64(Desc, GetValueAddress(Desc), ReturnValue);
+		if (Result != EPropertyBagResult::Success)
+		{
+			return MakeError(Result);
+		}
+		return MakeValue(static_cast<T>(ReturnValue));
+	}
+	else
+	{
+		int64 ReturnValue = 0;
+		const EPropertyBagResult Result = UE::StructUtils::Private::GetPropertyAsInt64(Desc, GetValueAddress(Desc), ReturnValue);
+		if (Result != EPropertyBagResult::Success)
+		{
+			return MakeError(Result);
+		}
+		return MakeValue(static_cast<T>(ReturnValue));
 	}
-	return MakeValue(static_cast<uint8>(ReturnValue));
 }
 
-TValueOrError<int32, EPropertyBagResult> FInstancedPropertyBag::GetValueInt32(const FName Name) const
+TValueOrError<int8, EPropertyBagResult> FInstancedPropertyBag::GetValueInt8(const FName Name) const
 {
-	const FPropertyBagPropertyDesc* Desc = FindPropertyDescByName(Name);
-	int64 ReturnValue = 0;
-	const EPropertyBagResult Result = UE::StructUtils::Private::GetPropertyAsInt64(Desc, GetValueAddress(Desc), ReturnValue);
-	if (Result != EPropertyBagResult::Success)
-	{
-		return MakeError(Result);
-	}
-	return MakeValue(static_cast<int32>(ReturnValue));
+	return GetNumericalValue<int8>(Name);
+}
+
+TValueOrError<uint8, EPropertyBagResult> FInstancedPropertyBag::GetValueByte(const FName Name) const
+{
+	return GetNumericalValue<uint8>(Name);
+}
+
+TValueOrError<int16, EPropertyBagResult> FInstancedPropertyBag::GetValueInt16(const FName Name) const
+{
+	return GetNumericalValue<int16>(Name);
+}
+
+TValueOrError<uint16, EPropertyBagResult> FInstancedPropertyBag::GetValueUInt16(const FName Name) const
+{
+	return GetNumericalValue<uint16>(Name);
+}
+
+TValueOrError<int32, EPropertyBagResult> FInstancedPropertyBag::GetValueInt32(const FName Name) const
+{
+	return GetNumericalValue<int32>(Name);
 }
 
 TValueOrError<uint32, EPropertyBagResult> FInstancedPropertyBag::GetValueUInt32(const FName Name) const
 {
-	const FPropertyBagPropertyDesc* Desc = FindPropertyDescByName(Name);
-	uint64 ReturnValue = 0;
-	const EPropertyBagResult Result = UE::StructUtils::Private::GetPropertyAsUInt64(Desc, GetValueAddress(Desc), ReturnValue);
-	if (Result != EPropertyBagResult::Success)
-	{
-		return MakeError(Result);
-	}
-	return MakeValue(static_cast<uint32>(ReturnValue));
+	return GetNumericalValue<uint32>(Name);
 }
 
 TValueOrError<int64, EPropertyBagResult> FInstancedPropertyBag::GetValueInt64(const FName Name) const
 {
-	const FPropertyBagPropertyDesc* Desc = FindPropertyDescByName(Name);
-	int64 ReturnValue = 0;
-	const EPropertyBagResult Result = UE::StructUtils::Private::GetPropertyAsInt64(Desc, GetValueAddress(Desc), ReturnValue);
-	if (Result != EPropertyBagResult::Success)
-	{
-		return MakeError(Result);
-	}
-	return MakeValue(ReturnValue);
+	return GetNumericalValue<int64>(Name);
 }
 
 TValueOrError<uint64, EPropertyBagResult> FInstancedPropertyBag::GetValueUInt64(const FName Name) const
 {
-	const FPropertyBagPropertyDesc* Desc = FindPropertyDescByName(Name);
-	uint64 ReturnValue = 0;
-	const EPropertyBagResult Result = UE::StructUtils::Private::GetPropertyAsUInt64(Desc, GetValueAddress(Desc), ReturnValue);
-	if (Result != EPropertyBagResult::Success)
-	{
-		return MakeError(Result);
-	}
-	return MakeValue((uint64)ReturnValue);
+	return GetNumericalValue<uint64>(Name);
 }
 
 TValueOrError<float, EPropertyBagResult> FInstancedPropertyBag::GetValueFloat(const FName Name) const
 {
-	const FPropertyBagPropertyDesc* Desc = FindPropertyDescByName(Name);
-	double ReturnValue = 0;
-	const EPropertyBagResult Result = UE::StructUtils::Private::GetPropertyAsDouble(Desc, GetValueAddress(Desc), ReturnValue);
-	if (Result != EPropertyBagResult::Success)
-	{
-		return MakeError(Result);
-	}
-	return MakeValue(static_cast<float>(ReturnValue));
+	return GetNumericalValue<float>(Name);
 }
 
 TValueOrError<double, EPropertyBagResult> FInstancedPropertyBag::GetValueDouble(const FName Name) const
 {
-	const FPropertyBagPropertyDesc* Desc = FindPropertyDescByName(Name);
-	double ReturnValue = 0;
-	const EPropertyBagResult Result = UE::StructUtils::Private::GetPropertyAsDouble(Desc, GetValueAddress(Desc), ReturnValue);
-	if (Result != EPropertyBagResult::Success)
-	{
-		return MakeError(Result);
-	}
-	return MakeValue(ReturnValue);
+	return GetNumericalValue<double>(Name);
 }
 
 TValueOrError<FName, EPropertyBagResult> FInstancedPropertyBag::GetValueName(const FName Name) const
@@ -2568,9 +2245,19 @@ TValueOrError<FText, EPropertyBagResult> FInstancedPropertyBag::GetValueText(con
 }
 
 TValueOrError<uint8, EPropertyBagResult> FInstancedPropertyBag::GetValueEnum(const FName Name, const UEnum* RequestedEnum) const
+{
+	TValueOrError<int64, EPropertyBagResult> Result = GetValueEnumInt64(Name, RequestedEnum);
+	if (Result.HasValue())
+	{
+		return MakeValue(static_cast<uint8>(Result.GetValue()));
+	}
+	return MakeError(Result.GetError());
+}
+
+TValueOrError<int64, EPropertyBagResult> FInstancedPropertyBag::GetValueEnumInt64(const FName Name, const UEnum* RequestedEnum) const
 {
 	const FPropertyBagPropertyDesc* Desc = FindPropertyDescByName(Name);
-	uint8 ReturnValue;
+	int64 ReturnValue;
 	const EPropertyBagResult Result = UE::StructUtils::Private::GetPropertyValueAsEnum(Desc, RequestedEnum, GetValueAddress(Desc), ReturnValue);
 	if (Result != EPropertyBagResult::Success)
 	{
@@ -2671,116 +2358,95 @@ TValueOrError<bool, EPropertyBagResult> FInstancedPropertyBag::GetValueBool(cons
 	return MakeValue(ReturnValue != 0);
 }
 
-TValueOrError<uint8, EPropertyBagResult> FInstancedPropertyBag::GetValueByte(const FPropertyBagPropertyDesc& Desc) const
+template<typename T>
+TValueOrError<T, EPropertyBagResult> FInstancedPropertyBag::GetNumericalValue(const FPropertyBagPropertyDesc& Desc) const
 {
 	if (!OwnsPropertyDesc(Desc))
 	{
 		return MakeError(EPropertyBagResult::PropertyNotFound);
 	}
 
-	uint64 ReturnValue = 0;
-	const EPropertyBagResult Result = UE::StructUtils::Private::GetPropertyAsUInt64(&Desc, GetValueAddress(&Desc), ReturnValue);
-	if (Result != EPropertyBagResult::Success)
+	if constexpr (std::is_floating_point<T>::value)
 	{
-		return MakeError(Result);
+		double ReturnValue = 0.0;
+		const EPropertyBagResult Result = UE::StructUtils::Private::GetPropertyAsDouble(&Desc, GetValueAddress(&Desc), ReturnValue);
+		if (Result != EPropertyBagResult::Success)
+		{
+			return MakeError(Result);
+		}
+		return MakeValue(static_cast<T>(ReturnValue));
 	}
-	return MakeValue(static_cast<uint8>(ReturnValue));
-}
-
-TValueOrError<int32, EPropertyBagResult> FInstancedPropertyBag::GetValueInt32(const FPropertyBagPropertyDesc& Desc) const
-{
-	if (!OwnsPropertyDesc(Desc))
+	else if constexpr (std::is_unsigned<T>::value)
 	{
-		return MakeError(EPropertyBagResult::PropertyNotFound);
+		uint64 ReturnValue = 0;
+		const EPropertyBagResult Result = UE::StructUtils::Private::GetPropertyAsUInt64(&Desc, GetValueAddress(&Desc), ReturnValue);
+		if (Result != EPropertyBagResult::Success)
+		{
+			return MakeError(Result);
+		}
+		return MakeValue(static_cast<T>(ReturnValue));
 	}
-
-	int64 ReturnValue = 0;
-	const EPropertyBagResult Result = UE::StructUtils::Private::GetPropertyAsInt64(&Desc, GetValueAddress(&Desc), ReturnValue);
-	if (Result != EPropertyBagResult::Success)
+	else
 	{
-		return MakeError(Result);
+		int64 ReturnValue = 0;
+		const EPropertyBagResult Result = UE::StructUtils::Private::GetPropertyAsInt64(&Desc, GetValueAddress(&Desc), ReturnValue);
+		if (Result != EPropertyBagResult::Success)
+		{
+			return MakeError(Result);
+		}
+		return MakeValue(static_cast<T>(ReturnValue));
 	}
-	return MakeValue(static_cast<int32>(ReturnValue));
 }
 
-TValueOrError<uint32, EPropertyBagResult> FInstancedPropertyBag::GetValueUInt32(const FPropertyBagPropertyDesc& Desc) const
+TValueOrError<int8, EPropertyBagResult> FInstancedPropertyBag::GetValueInt8(const FPropertyBagPropertyDesc& Desc) const
 {
-	if (!OwnsPropertyDesc(Desc))
-	{
-		return MakeError(EPropertyBagResult::PropertyNotFound);
-	}
+	return GetNumericalValue<int8>(Desc);
+}
 
-	uint64 ReturnValue = 0;
-	const EPropertyBagResult Result = UE::StructUtils::Private::GetPropertyAsUInt64(&Desc, GetValueAddress(&Desc), ReturnValue);
-	if (Result != EPropertyBagResult::Success)
-	{
-		return MakeError(Result);
-	}
-	return MakeValue(static_cast<uint32>(ReturnValue));
+TValueOrError<uint8, EPropertyBagResult> FInstancedPropertyBag::GetValueByte(const FPropertyBagPropertyDesc& Desc) const
+{
+	return GetNumericalValue<uint8>(Desc);
+}
+
+TValueOrError<int16, EPropertyBagResult> FInstancedPropertyBag::GetValueInt16(const FPropertyBagPropertyDesc& Desc) const
+{
+	return GetNumericalValue<int16>(Desc);
+}
+
+TValueOrError<uint16, EPropertyBagResult> FInstancedPropertyBag::GetValueUInt16(const FPropertyBagPropertyDesc& Desc) const
+{
+	return GetNumericalValue<uint16>(Desc);
+}
+
+TValueOrError<int32, EPropertyBagResult> FInstancedPropertyBag::GetValueInt32(const FPropertyBagPropertyDesc& Desc) const
+{
+	return GetNumericalValue<int32>(Desc);
+}
+
+TValueOrError<uint32, EPropertyBagResult> FInstancedPropertyBag::GetValueUInt32(const FPropertyBagPropertyDesc& Desc) const
+{
+	return GetNumericalValue<uint32>(Desc);
 }
 
 TValueOrError<int64, EPropertyBagResult> FInstancedPropertyBag::GetValueInt64(const FPropertyBagPropertyDesc& Desc) const
 {
-	if (!OwnsPropertyDesc(Desc))
-	{
-		return MakeError(EPropertyBagResult::PropertyNotFound);
-	}
 
-	int64 ReturnValue = 0;
-	const EPropertyBagResult Result = UE::StructUtils::Private::GetPropertyAsInt64(&Desc, GetValueAddress(&Desc), ReturnValue);
-	if (Result != EPropertyBagResult::Success)
-	{
-		return MakeError(Result);
-	}
-	return MakeValue(ReturnValue);
+	return GetNumericalValue<int64>(Desc);
 }
 
 TValueOrError<uint64, EPropertyBagResult> FInstancedPropertyBag::GetValueUInt64(const FPropertyBagPropertyDesc& Desc) const
 {
-	if (!OwnsPropertyDesc(Desc))
-	{
-		return MakeError(EPropertyBagResult::PropertyNotFound);
-	}
-
-	uint64 ReturnValue = 0;
-	const EPropertyBagResult Result = UE::StructUtils::Private::GetPropertyAsUInt64(&Desc, GetValueAddress(&Desc), ReturnValue);
-	if (Result != EPropertyBagResult::Success)
-	{
-		return MakeError(Result);
-	}
-	return MakeValue((uint64)ReturnValue);
+	return GetNumericalValue<uint64>(Desc);
 }
 
 TValueOrError<float, EPropertyBagResult> FInstancedPropertyBag::GetValueFloat(const FPropertyBagPropertyDesc& Desc) const
 {
-	if (!OwnsPropertyDesc(Desc))
-	{
-		return MakeError(EPropertyBagResult::PropertyNotFound);
-	}
-
-	double ReturnValue = 0;
-	const EPropertyBagResult Result = UE::StructUtils::Private::GetPropertyAsDouble(&Desc, GetValueAddress(&Desc), ReturnValue);
-	if (Result != EPropertyBagResult::Success)
-	{
-		return MakeError(Result);
-	}
-	return MakeValue(static_cast<float>(ReturnValue));
+	return GetNumericalValue<float>(Desc);
 }
 
 TValueOrError<double, EPropertyBagResult> FInstancedPropertyBag::GetValueDouble(const FPropertyBagPropertyDesc& Desc) const
 {
-	if (!OwnsPropertyDesc(Desc))
-	{
-		return MakeError(EPropertyBagResult::PropertyNotFound);
-	}
-
-	double ReturnValue = 0;
-	const EPropertyBagResult Result = UE::StructUtils::Private::GetPropertyAsDouble(&Desc, GetValueAddress(&Desc), ReturnValue);
-	if (Result != EPropertyBagResult::Success)
-	{
-		return MakeError(Result);
-	}
-	return MakeValue(ReturnValue);
+	return GetNumericalValue<double>(Desc);
 }
 
 TValueOrError<FName, EPropertyBagResult> FInstancedPropertyBag::GetValueName(const FPropertyBagPropertyDesc& Desc) const
@@ -2832,13 +2498,23 @@ TValueOrError<FText, EPropertyBagResult> FInstancedPropertyBag::GetValueText(con
 }
 
 TValueOrError<uint8, EPropertyBagResult> FInstancedPropertyBag::GetValueEnum(const FPropertyBagPropertyDesc& Desc, const UEnum* RequestedEnum) const
+{
+	TValueOrError<int64, EPropertyBagResult> Result = GetValueEnumInt64(Desc, RequestedEnum);
+	if (Result.HasValue())
+	{
+		return MakeValue(static_cast<uint8>(Result.GetValue()));
+	}
+	return MakeError(Result.GetError());
+}
+
+TValueOrError<int64, EPropertyBagResult> FInstancedPropertyBag::GetValueEnumInt64(const FPropertyBagPropertyDesc& Desc, const UEnum* RequestedEnum) const
 {
 	if (!OwnsPropertyDesc(Desc))
 	{
 		return MakeError(EPropertyBagResult::PropertyNotFound);
 	}
 
-	uint8 ReturnValue;
+	int64 ReturnValue;
 	const EPropertyBagResult Result = UE::StructUtils::Private::GetPropertyValueAsEnum(&Desc, RequestedEnum, GetValueAddress(&Desc), ReturnValue);
 	if (Result != EPropertyBagResult::Success)
 	{
@@ -3055,12 +2731,30 @@ EPropertyBagResult FInstancedPropertyBag::SetValueBool(const FName Name, const b
 	return UE::StructUtils::Private::SetPropertyFromInt64(Desc, GetMutableValueAddress(Desc), bInValue ? 1 : 0);
 }
 
+EPropertyBagResult FInstancedPropertyBag::SetValueInt8(const FName Name, const int8 InValue)
+{
+	const FPropertyBagPropertyDesc* Desc = FindPropertyDescByName(Name);
+	return UE::StructUtils::Private::SetPropertyFromInt64(Desc, GetMutableValueAddress(Desc), InValue);
+}
+
 EPropertyBagResult FInstancedPropertyBag::SetValueByte(const FName Name, const uint8 InValue)
 {
 	const FPropertyBagPropertyDesc* Desc = FindPropertyDescByName(Name);
 	return UE::StructUtils::Private::SetPropertyFromUInt64(Desc, GetMutableValueAddress(Desc), InValue);
 }
 
+EPropertyBagResult FInstancedPropertyBag::SetValueInt16(const FName Name, const int16 InValue)
+{
+	const FPropertyBagPropertyDesc* Desc = FindPropertyDescByName(Name);
+	return UE::StructUtils::Private::SetPropertyFromInt64(Desc, GetMutableValueAddress(Desc), InValue);
+}
+
+EPropertyBagResult FInstancedPropertyBag::SetValueUInt16(const FName Name, const uint16 InValue)
+{
+	const FPropertyBagPropertyDesc* Desc = FindPropertyDescByName(Name);
+	return UE::StructUtils::Private::SetPropertyFromUInt64(Desc, GetMutableValueAddress(Desc), InValue);
+}
+
 EPropertyBagResult FInstancedPropertyBag::SetValueInt32(const FName Name, const int32 InValue)
 {
 	const FPropertyBagPropertyDesc* Desc = FindPropertyDescByName(Name);
@@ -3070,7 +2764,7 @@ EPropertyBagResult FInstancedPropertyBag::SetValueInt32(const FName Name, const
 EPropertyBagResult FInstancedPropertyBag::SetValueUInt32(const FName Name, const uint32 InValue)
 {
 	const FPropertyBagPropertyDesc* Desc = FindPropertyDescByName(Name);
-	return UE::StructUtils::Private::SetPropertyFromUInt64(Desc, GetMutableValueAddress(Desc), (uint64)InValue);
+	return UE::StructUtils::Private::SetPropertyFromUInt64(Desc, GetMutableValueAddress(Desc), InValue);
 }
 
 EPropertyBagResult FInstancedPropertyBag::SetValueInt64(const FName Name, const int64 InValue)
@@ -3115,7 +2809,7 @@ EPropertyBagResult FInstancedPropertyBag::SetValueText(const FName Name, const F
 	return UE::StructUtils::Private::SetPropertyValue<FText, FTextProperty>(Desc, GetMutableValueAddress(Desc), InValue);
 }
 
-EPropertyBagResult FInstancedPropertyBag::SetValueEnum(const FName Name, const uint8 InValue, const UEnum* Enum)
+EPropertyBagResult FInstancedPropertyBag::SetValueEnum(const FName Name, const int64 InValue, const UEnum* Enum)
 {
 	const FPropertyBagPropertyDesc* Desc = FindPropertyDescByName(Name);
 	return UE::StructUtils::Private::SetPropertyValueAsEnum(Desc, GetMutableValueAddress(Desc), InValue, Enum);
@@ -3204,6 +2898,15 @@ EPropertyBagResult FInstancedPropertyBag::SetValueBool(const FPropertyBagPropert
 	return UE::StructUtils::Private::SetPropertyFromInt64(&Desc, GetMutableValueAddress(&Desc), bInValue ? 1 : 0);
 }
 
+EPropertyBagResult FInstancedPropertyBag::SetValueInt8(const FPropertyBagPropertyDesc& Desc, const int8 InValue)
+{
+	if (!OwnsPropertyDesc(Desc))
+	{
+		return EPropertyBagResult::PropertyNotFound;
+	}
+	return UE::StructUtils::Private::SetPropertyFromInt64(&Desc, GetMutableValueAddress(&Desc), InValue);
+}
+
 EPropertyBagResult FInstancedPropertyBag::SetValueByte(const FPropertyBagPropertyDesc& Desc, const uint8 InValue)
 {
 	if (!OwnsPropertyDesc(Desc))
@@ -3213,6 +2916,24 @@ EPropertyBagResult FInstancedPropertyBag::SetValueByte(const FPropertyBagPropert
 	return UE::StructUtils::Private::SetPropertyFromUInt64(&Desc, GetMutableValueAddress(&Desc), InValue);
 }
 
+EPropertyBagResult FInstancedPropertyBag::SetValueInt16(const FPropertyBagPropertyDesc& Desc, const int16 InValue)
+{
+	if (!OwnsPropertyDesc(Desc))
+	{
+		return EPropertyBagResult::PropertyNotFound;
+	}
+	return UE::StructUtils::Private::SetPropertyFromInt64(&Desc, GetMutableValueAddress(&Desc), InValue);
+}
+
+EPropertyBagResult FInstancedPropertyBag::SetValueUInt16(const FPropertyBagPropertyDesc& Desc, const uint16 InValue)
+{
+	if (!OwnsPropertyDesc(Desc))
+	{
+		return EPropertyBagResult::PropertyNotFound;
+	}
+	return UE::StructUtils::Private::SetPropertyFromUInt64(&Desc, GetMutableValueAddress(&Desc), InValue);
+}
+
 EPropertyBagResult FInstancedPropertyBag::SetValueInt32(const FPropertyBagPropertyDesc& Desc, const int32 InValue)
 {
 	if (!OwnsPropertyDesc(Desc))
@@ -3228,7 +2949,7 @@ EPropertyBagResult FInstancedPropertyBag::SetValueUInt32(const FPropertyBagPrope
 	{
 		return EPropertyBagResult::PropertyNotFound;
 	}
-	return UE::StructUtils::Private::SetPropertyFromUInt64(&Desc, GetMutableValueAddress(&Desc), (uint64)InValue);
+	return UE::StructUtils::Private::SetPropertyFromUInt64(&Desc, GetMutableValueAddress(&Desc), InValue);
 }
 
 EPropertyBagResult FInstancedPropertyBag::SetValueInt64(const FPropertyBagPropertyDesc& Desc, const int64 InValue)
@@ -3294,7 +3015,7 @@ EPropertyBagResult FInstancedPropertyBag::SetValueText(const FPropertyBagPropert
 	return UE::StructUtils::Private::SetPropertyValue<FText, FTextProperty>(&Desc, GetMutableValueAddress(&Desc), InValue);
 }
 
-EPropertyBagResult FInstancedPropertyBag::SetValueEnum(const FPropertyBagPropertyDesc& Desc, const uint8 InValue, const UEnum* Enum)
+EPropertyBagResult FInstancedPropertyBag::SetValueEnum(const FPropertyBagPropertyDesc& Desc, const int64 InValue, const UEnum* Enum)
 {
 	if (!OwnsPropertyDesc(Desc))
 	{
@@ -3809,81 +3530,89 @@ TValueOrError<bool, EPropertyBagResult> FPropertyBagArrayRef::GetValueBool(const
 	return MakeValue(IntValue != 0);
 }
 
-TValueOrError<uint8, EPropertyBagResult> FPropertyBagArrayRef::GetValueByte(const int32 Index) const
+template<typename T>
+TValueOrError<T, EPropertyBagResult> FPropertyBagArrayRef::GetNumericalValue(const int32 Index) const
 {
-	uint64 IntValue = 0;
-	const EPropertyBagResult Result = UE::StructUtils::Private::GetPropertyAsUInt64(&ValueDesc, GetAddress(Index), IntValue);
-	if (Result != EPropertyBagResult::Success)
+	if constexpr (std::is_floating_point_v<T>)
 	{
-		return MakeError(Result);
+		double DblValue = 0.0;
+		const EPropertyBagResult Result = UE::StructUtils::Private::GetPropertyAsDouble(&ValueDesc, GetAddress(Index), DblValue);
+		if (Result != EPropertyBagResult::Success)
+		{
+			return MakeError(Result);
+		}
+		return MakeValue(static_cast<T>(DblValue));
+	}
+	else if constexpr (std::is_unsigned_v<T>)
+	{
+		uint64 IntValue = 0;
+		const EPropertyBagResult Result = UE::StructUtils::Private::GetPropertyAsUInt64(&ValueDesc, GetAddress(Index), IntValue);
+		if (Result != EPropertyBagResult::Success)
+		{
+			return MakeError(Result);
+		}
+		return MakeValue(static_cast<T>(IntValue));
+	}
+	else
+	{
+		int64 IntValue = 0;
+		const EPropertyBagResult Result = UE::StructUtils::Private::GetPropertyAsInt64(&ValueDesc, GetAddress(Index), IntValue);
+		if (Result != EPropertyBagResult::Success)
+		{
+			return MakeError(Result);
+		}
+		return MakeValue(static_cast<T>(IntValue));
 	}
-	return MakeValue(static_cast<uint8>(IntValue));
+}
+
+TValueOrError<int8, EPropertyBagResult> FPropertyBagArrayRef::GetValueInt8(const int32 Index) const
+{
+	return GetNumericalValue<int8>(Index);
+}
+
+TValueOrError<uint8, EPropertyBagResult> FPropertyBagArrayRef::GetValueByte(const int32 Index) const
+{
+	return GetNumericalValue<uint8>(Index);
+}
+
+TValueOrError<int16, EPropertyBagResult> FPropertyBagArrayRef::GetValueInt16(const int32 Index) const
+{
+	return GetNumericalValue<int16>(Index);
+}
+
+TValueOrError<uint16, EPropertyBagResult> FPropertyBagArrayRef::GetValueUInt16(const int32 Index) const
+{
+	return GetNumericalValue<uint16>(Index);
 }
 
 TValueOrError<int32, EPropertyBagResult> FPropertyBagArrayRef::GetValueInt32(const int32 Index) const
 {
-	int64 IntValue = 0;
-	const EPropertyBagResult Result = UE::StructUtils::Private::GetPropertyAsInt64(&ValueDesc, GetAddress(Index), IntValue);
-	if (Result != EPropertyBagResult::Success)
-	{
-		return MakeError(Result);
-	}
-	return MakeValue(static_cast<int32>(IntValue));
+	return GetNumericalValue<int32>(Index);
 }
 
 TValueOrError<uint32, EPropertyBagResult> FPropertyBagArrayRef::GetValueUInt32(const int32 Index) const
 {
-	uint64 IntValue = 0;
-	const EPropertyBagResult Result = UE::StructUtils::Private::GetPropertyAsUInt64(&ValueDesc, GetAddress(Index), IntValue);
-	if (Result != EPropertyBagResult::Success)
-	{
-		return MakeError(Result);
-	}
-	return MakeValue(static_cast<uint32>(IntValue));
+	return GetNumericalValue<uint32>(Index);
 }
 
 TValueOrError<int64, EPropertyBagResult> FPropertyBagArrayRef::GetValueInt64(const int32 Index) const
 {
-	int64 IntValue = 0;
-	const EPropertyBagResult Result = UE::StructUtils::Private::GetPropertyAsInt64(&ValueDesc, GetAddress(Index), IntValue);
-	if (Result != EPropertyBagResult::Success)
-	{
-		return MakeError(Result);
-	}
-	return MakeValue(IntValue);
+	return GetNumericalValue<int64>(Index);
 }
 
 TValueOrError<uint64, EPropertyBagResult> FPropertyBagArrayRef::GetValueUInt64(const int32 Index) const
 {
-	uint64 IntValue = 0;
-	const EPropertyBagResult Result = UE::StructUtils::Private::GetPropertyAsUInt64(&ValueDesc, GetAddress(Index), IntValue);
-	if (Result != EPropertyBagResult::Success)
-	{
-		return MakeError(Result);
-	}
-	return MakeValue((uint64)IntValue);
+	return GetNumericalValue<uint64>(Index);
 }
 
 TValueOrError<float, EPropertyBagResult> FPropertyBagArrayRef::GetValueFloat(const int32 Index) const
 {
-	double DblValue = 0;
-	const EPropertyBagResult Result = UE::StructUtils::Private::GetPropertyAsDouble(&ValueDesc, GetAddress(Index), DblValue);
-	if (Result != EPropertyBagResult::Success)
-	{
-		return MakeError(Result);
-	}
-	return MakeValue(static_cast<float>(DblValue));
+	return GetNumericalValue<float>(Index);
 }
 
 TValueOrError<double, EPropertyBagResult> FPropertyBagArrayRef::GetValueDouble(const int32 Index) const
 {
-	double DblValue = 0;
-	const EPropertyBagResult Result = UE::StructUtils::Private::GetPropertyAsDouble(&ValueDesc, GetAddress(Index), DblValue);
-	if (Result != EPropertyBagResult::Success)
-	{
-		return MakeError(Result);
-	}
-	return MakeValue(DblValue);
+	return GetNumericalValue<double>(Index);
 }
 
 TValueOrError<FName, EPropertyBagResult> FPropertyBagArrayRef::GetValueName(const int32 Index) const
@@ -3921,7 +3650,18 @@ TValueOrError<FText, EPropertyBagResult> FPropertyBagArrayRef::GetValueText(cons
 
 TValueOrError<uint8, EPropertyBagResult> FPropertyBagArrayRef::GetValueEnum(const int32 Index, const UEnum* RequestedEnum) const
 {
-	uint8 Value;
+	int64 Value;
+	const EPropertyBagResult Result = UE::StructUtils::Private::GetPropertyValueAsEnum(&ValueDesc, RequestedEnum, GetAddress(Index), Value);
+	if (Result != EPropertyBagResult::Success)
+	{
+		return MakeError(Result);
+	}
+	return MakeValue(static_cast<uint8>(Value));
+}
+
+TValueOrError<int64, EPropertyBagResult> FPropertyBagArrayRef::GetValueEnumInt64(const int32 Index, const UEnum* RequestedEnum) const
+{
+	int64 Value;
 	const EPropertyBagResult Result = UE::StructUtils::Private::GetPropertyValueAsEnum(&ValueDesc, RequestedEnum, GetAddress(Index), Value);
 	if (Result != EPropertyBagResult::Success)
 	{
@@ -4024,11 +3764,26 @@ EPropertyBagResult FPropertyBagArrayRef::SetValueBool(const int32 Index, const b
 	return UE::StructUtils::Private::SetPropertyFromInt64(&ValueDesc, GetMutableAddress(Index), bInValue ? 1 : 0);
 }
 
+EPropertyBagResult FPropertyBagArrayRef::SetValueInt8(const int32 Index, const int8 InValue)
+{
+	return UE::StructUtils::Private::SetPropertyFromInt64(&ValueDesc, GetMutableAddress(Index), InValue);
+}
+
 EPropertyBagResult FPropertyBagArrayRef::SetValueByte(const int32 Index, const uint8 InValue)
 {
 	return UE::StructUtils::Private::SetPropertyFromUInt64(&ValueDesc, GetMutableAddress(Index), InValue);
 }
 
+EPropertyBagResult FPropertyBagArrayRef::SetValueInt16(const int32 Index, const int16 InValue)
+{
+	return UE::StructUtils::Private::SetPropertyFromInt64(&ValueDesc, GetMutableAddress(Index), InValue);
+}
+
+EPropertyBagResult FPropertyBagArrayRef::SetValueUInt16(const int32 Index, const uint16 InValue)
+{
+	return UE::StructUtils::Private::SetPropertyFromUInt64(&ValueDesc, GetMutableAddress(Index), InValue);
+}
+
 EPropertyBagResult FPropertyBagArrayRef::SetValueInt32(const int32 Index, const int32 InValue)
 {
 	return UE::StructUtils::Private::SetPropertyFromInt64(&ValueDesc, GetMutableAddress(Index), InValue);
@@ -4074,7 +3829,7 @@ EPropertyBagResult FPropertyBagArrayRef::SetValueText(const int32 Index, const F
 	return UE::StructUtils::Private::SetPropertyValue<FText, FTextProperty>(&ValueDesc, GetMutableAddress(Index), InValue);
 }
 
-EPropertyBagResult FPropertyBagArrayRef::SetValueEnum(const int32 Index, const uint8 InValue, const UEnum* Enum)
+EPropertyBagResult FPropertyBagArrayRef::SetValueEnum(const int32 Index, const int64 InValue, const UEnum* Enum)
 {
 	return UE::StructUtils::Private::SetPropertyValueAsEnum(&ValueDesc, GetMutableAddress(Index), InValue, Enum);
 }
@@ -4113,11 +3868,26 @@ EPropertyBagResult FPropertyBagSetRef::AddValueBool(const bool bInValue)
 	return ValueDesc.IsNumericType() ? Add(bInValue) : EPropertyBagResult::TypeMismatch;
 }
 
+EPropertyBagResult FPropertyBagSetRef::AddValueInt8(const int8 InValue)
+{
+	return ValueDesc.IsNumericType() ? Add(InValue) : EPropertyBagResult::TypeMismatch;
+}
+
 EPropertyBagResult FPropertyBagSetRef::AddValueByte(const uint8 InValue)
 {
 	return ValueDesc.IsNumericType() ? Add(InValue) : EPropertyBagResult::TypeMismatch;
 }
 
+EPropertyBagResult FPropertyBagSetRef::AddValueInt16(const int16 InValue)
+{
+	return ValueDesc.IsNumericType() ? Add(InValue) : EPropertyBagResult::TypeMismatch;
+}
+
+EPropertyBagResult FPropertyBagSetRef::AddValueUInt16(const uint16 InValue)
+{
+	return ValueDesc.IsNumericType() ? Add(InValue) : EPropertyBagResult::TypeMismatch;
+}
+
 EPropertyBagResult FPropertyBagSetRef::AddValueInt32(const int32 InValue)
 {
 	return ValueDesc.IsNumericType() ? Add(InValue) : EPropertyBagResult::TypeMismatch;
diff --git a/Engine/Source/Runtime/CoreUObject/Private/StructUtils/StructUtils.cpp b/Engine/Source/Runtime/CoreUObject/Private/StructUtils/StructUtils.cpp
--- a/Engine/Source/Runtime/CoreUObject/Private/StructUtils/StructUtils.cpp
+++ b/Engine/Source/Runtime/CoreUObject/Private/StructUtils/StructUtils.cpp
@@ -1,4 +1,4 @@
-﻿// Copyright Epic Games, Inc. All Rights Reserved.
+// Copyright Epic Games, Inc. All Rights Reserved.
 
 #include "StructUtils/StructUtils.h"
 
@@ -80,3 +80,190 @@ uint64 UE::StructUtils::CalcPropertyDescArrayHash(const TConstArrayView<FPropert
 	}
 	return HashBuilder.Finalize().Hash;
 }
+
+TValueOrError<int64, void> UE::StructUtils::GetPropertyValueAsInt64(TNotNull<const FProperty*> InProperty, TNotNull<const void*> Address)
+{
+	if (const FNumericProperty* NumericProperty = CastField<const FNumericProperty>(InProperty))
+	{
+		if (NumericProperty->IsInteger())
+		{
+			return MakeValue(NumericProperty->GetSignedIntPropertyValue(Address));
+		}
+		else if (NumericProperty->IsFloatingPoint())
+		{
+			return MakeValue((int64)NumericProperty->GetFloatingPointPropertyValue(Address));
+		}
+
+	}
+	else if (const FBoolProperty* BoolProperty = CastField<FBoolProperty>(InProperty))
+	{
+		const int64 Value = BoolProperty->GetPropertyValue(Address) ? 1 : 0;
+		return MakeValue(Value);
+	}
+	else if (const FEnumProperty* EnumProperty = CastField<FEnumProperty>(InProperty))
+	{
+		const FNumericProperty* UnderlyingProperty = EnumProperty->GetUnderlyingProperty();
+		if (UnderlyingProperty->IsInteger())
+		{
+			return MakeValue(UnderlyingProperty->GetSignedIntPropertyValue(Address));
+		}
+	}
+	return MakeError();
+}
+
+TValueOrError<uint64, void> UE::StructUtils::GetPropertyValueAsUInt64(TNotNull<const FProperty*> InProperty, TNotNull<const void*> Address)
+{
+	if (const FNumericProperty* NumericProperty = CastField<const FNumericProperty>(InProperty))
+	{
+		if (NumericProperty->IsInteger())
+		{
+			return MakeValue(NumericProperty->GetUnsignedIntPropertyValue(Address));
+		}
+		else if (NumericProperty->IsFloatingPoint())
+		{
+			return MakeValue((uint64)NumericProperty->GetFloatingPointPropertyValue(Address));
+		}
+	}
+	else if (const FBoolProperty* BoolProperty = CastField<FBoolProperty>(InProperty))
+	{
+		const uint64 Value = BoolProperty->GetPropertyValue(Address) ? 1 : 0;
+		return MakeValue(Value);
+	}
+	else if (const FEnumProperty* EnumProperty = CastField<FEnumProperty>(InProperty))
+	{
+		const FNumericProperty* UnderlyingProperty = EnumProperty->GetUnderlyingProperty();
+		if (UnderlyingProperty->IsInteger())
+		{
+			return MakeValue(UnderlyingProperty->GetUnsignedIntPropertyValue(Address));
+		}
+	}
+	return MakeError();
+}
+
+TValueOrError<double, void> UE::StructUtils::GetPropertyValueAsDouble(TNotNull<const FProperty*> InProperty, TNotNull<const void*> Address)
+{
+	if (const FNumericProperty* NumericProperty = CastField<const FNumericProperty>(InProperty))
+	{
+		if (NumericProperty->IsInteger())
+		{
+			return MakeValue((double)NumericProperty->GetSignedIntPropertyValue(Address));
+		}
+		else if (NumericProperty->IsFloatingPoint())
+		{
+			return MakeValue(NumericProperty->GetFloatingPointPropertyValue(Address));
+		}
+	}
+	else if (const FBoolProperty* BoolProperty = CastField<FBoolProperty>(InProperty))
+	{
+		const double Value = BoolProperty->GetPropertyValue(Address) ? 1.0 : 0.0;
+		return MakeValue(Value);
+	}
+	else if (const FEnumProperty* EnumProperty = CastField<FEnumProperty>(InProperty))
+	{
+		const FNumericProperty* UnderlyingProperty = EnumProperty->GetUnderlyingProperty();
+		if (UnderlyingProperty->IsInteger())
+		{
+			return MakeValue((double)UnderlyingProperty->GetSignedIntPropertyValue(Address));
+		}
+	}
+	return MakeError();
+}
+
+bool UE::StructUtils::SetPropertyValueFromInt64(TNotNull<const FProperty*> InProperty, TNotNull<void*> Address, const int64 Value)
+{
+	if (const FNumericProperty* NumericProperty = CastField<const FNumericProperty>(InProperty))
+	{
+		if (NumericProperty->IsInteger())
+		{
+			NumericProperty->SetIntPropertyValue(Address, Value);
+			return true;
+		}
+		else if (NumericProperty->IsFloatingPoint())
+		{
+			NumericProperty->SetFloatingPointPropertyValue(Address, (double)Value);
+			return true;
+		}
+	}
+	else if (const FBoolProperty* BoolProperty = CastField<FBoolProperty>(InProperty))
+	{
+		const bool bValue = Value != 0;
+		BoolProperty->SetPropertyValue(Address, bValue);
+		return true;
+	}
+	else if (const FEnumProperty* EnumProperty = CastField<FEnumProperty>(InProperty))
+	{
+		const FNumericProperty* UnderlyingProperty = EnumProperty->GetUnderlyingProperty();
+		if (UnderlyingProperty->IsInteger())
+		{
+			UnderlyingProperty->SetIntPropertyValue(Address, Value);
+			return true;
+		}
+	}
+	return false;
+}
+
+bool UE::StructUtils::SetPropertyValueFromUInt64(TNotNull<const FProperty*> InProperty, TNotNull<void*> Address, uint64 Value)
+{
+	if (const FNumericProperty* NumericProperty = CastField<const FNumericProperty>(InProperty))
+	{
+		if (NumericProperty->IsInteger())
+		{
+			NumericProperty->SetIntPropertyValue(Address, Value);
+			return true;
+		}
+		else if (NumericProperty->IsFloatingPoint())
+		{
+			NumericProperty->SetFloatingPointPropertyValue(Address, (double)Value);
+			return true;
+		}
+	}
+	else if (const FBoolProperty* BoolProperty = CastField<FBoolProperty>(InProperty))
+	{
+		const bool bValue = Value != 0;
+		BoolProperty->SetPropertyValue(Address, bValue);
+		return true;
+	}
+	else if (const FEnumProperty* EnumProperty = CastField<FEnumProperty>(InProperty))
+	{
+		const FNumericProperty* UnderlyingProperty = EnumProperty->GetUnderlyingProperty();
+		if (UnderlyingProperty->IsInteger())
+		{
+			UnderlyingProperty->SetIntPropertyValue(Address, Value);
+			return true;
+		}
+	}
+	return false;
+}
+
+bool UE::StructUtils::SetPropertyValueFromDouble(TNotNull<const FProperty*> InProperty, TNotNull<void*> Address, double Value)
+{
+	if (const FNumericProperty* NumericProperty = CastField<const FNumericProperty>(InProperty))
+	{
+		if (NumericProperty->IsInteger())
+		{
+			NumericProperty->SetIntPropertyValue(Address, FMath::RoundToInt64(Value));
+			return true;
+		}
+		else if (NumericProperty->IsFloatingPoint())
+		{
+			NumericProperty->SetFloatingPointPropertyValue(Address, Value);
+			return true;
+		}
+	}
+	else if (const FBoolProperty* BoolProperty = CastField<FBoolProperty>(InProperty))
+	{
+		const bool bValue = !(FMath::IsNearlyZero(Value));
+		BoolProperty->SetPropertyValue(Address, bValue);
+		return true;
+	}
+	else if (const FEnumProperty* EnumProperty = CastField<FEnumProperty>(InProperty))
+	{
+		const FNumericProperty* UnderlyingProperty = EnumProperty->GetUnderlyingProperty();
+		if (UnderlyingProperty->IsInteger())
+		{
+			UnderlyingProperty->SetIntPropertyValue(Address, FMath::RoundToInt64(Value));
+			return true;
+		}
+	}
+	return false;
+}
diff --git a/Engine/Source/Runtime/CoreUObject/Public/StructUtils/PropertyBag.h b/Engine/Source/Runtime/CoreUObject/Public/StructUtils/PropertyBag.h
--- a/Engine/Source/Runtime/CoreUObject/Public/StructUtils/PropertyBag.h
+++ b/Engine/Source/Runtime/CoreUObject/Public/StructUtils/PropertyBag.h
@@ -33,6 +33,9 @@ enum class EPropertyBagPropertyType : uint8
 	SoftObject,
 	Class,
 	SoftClass,
+	Int8,	// Type not fully supported at UI, will work with restrictions to type editing
+	Int16,	// Type not fully supported at UI, will work with restrictions to type editing
+	UInt16,	// Type not fully supported at UI, will work with restrictions to type editing
 	UInt32,	// Type not fully supported at UI, will work with restrictions to type editing
 	UInt64, // Type not fully supported at UI, will work with restrictions to type editing
 
@@ -286,11 +289,14 @@ struct FPropertyBagPropertyDesc
 	/** @return true if the two descriptors have the same type. Object types are compatible if Other can be cast to this type. */
 	UE_API bool CompatibleType(const FPropertyBagPropertyDesc& Other) const;
 
-	/** @return true if the property type is numeric (bool, (u)int32, (u)int64, float, double, enum) */
+	/** @return true if the property type is numeric (bool, (u)int8, (u)int16, (u)int32, (u)int64, float, double, enum) */
 	UE_API bool IsNumericType() const;
 
-	/** @return true if the property type is unsigned (uint32, uint64) */
+	/** @return true if the property type is unsigned (uint8, uint16, uint32, uint64) */
 	UE_API bool IsUnsignedNumericType() const;
+
+	/** @return true if the property type is integral ((u)int8, (u)int16, (u)int32, (u)int64) */
+	UE_API bool IsNumericIntegralType() const;
 	
 	/** @return true if the property type is floating point numeric (float, double) */
 	UE_API bool IsNumericFloatType() const;
@@ -648,11 +654,13 @@ struct FInstancedPropertyBag
 
 	/**
 	 * Getters
-	 * Numeric types (bool, (u)int32, (u)int64, float, double) support type conversion.
+	 * Numeric types (bool, (u)int8, (u)int16, (u)int32, (u)int64, float, double) support type conversion.
 	 */
-
 	UE_API TValueOrError<bool, EPropertyBagResult> GetValueBool(const FName Name) const;
+	UE_API TValueOrError<int8, EPropertyBagResult> GetValueInt8(const FName Name) const;
 	UE_API TValueOrError<uint8, EPropertyBagResult> GetValueByte(const FName Name) const;
+	UE_API TValueOrError<int16, EPropertyBagResult> GetValueInt16(const FName Name) const;
+	UE_API TValueOrError<uint16, EPropertyBagResult> GetValueUInt16(const FName Name) const;
 	UE_API TValueOrError<int32, EPropertyBagResult> GetValueInt32(const FName Name) const;
 	UE_API TValueOrError<uint32, EPropertyBagResult> GetValueUInt32(const FName Name) const;
 	UE_API TValueOrError<int64, EPropertyBagResult> GetValueInt64(const FName Name) const;
@@ -663,6 +671,7 @@ struct FInstancedPropertyBag
 	UE_API TValueOrError<FString, EPropertyBagResult> GetValueString(const FName Name) const;
 	UE_API TValueOrError<FText, EPropertyBagResult> GetValueText(const FName Name) const;
 	UE_API TValueOrError<uint8, EPropertyBagResult> GetValueEnum(const FName Name, const UEnum* RequestedEnum) const;
+	UE_API TValueOrError<int64, EPropertyBagResult> GetValueEnumInt64(const FName Name, const UEnum* RequestedEnum) const;
 	UE_API TValueOrError<FStructView, EPropertyBagResult> GetValueStruct(const FName Name, const UScriptStruct* RequestedStruct = nullptr) const;
 	UE_API TValueOrError<UObject*, EPropertyBagResult> GetValueObject(const FName Name, const UClass* RequestedClass = nullptr) const;
 	UE_API TValueOrError<UClass*, EPropertyBagResult> GetValueClass(const FName Name) const;
@@ -677,7 +686,7 @@ struct FInstancedPropertyBag
 	{
 		static_assert(TIsEnum<T>::Value, "Should only call this with enum types");
 		
-		TValueOrError<uint8, EPropertyBagResult> Result = GetValueEnum(Name, StaticEnum<T>());
+		TValueOrError<int64, EPropertyBagResult> Result = GetValueEnumInt64(Name, StaticEnum<T>());
 		if (!Result.IsValid())
 		{
 			return MakeError(Result.GetError());
@@ -723,8 +732,16 @@ struct FInstancedPropertyBag
 		return MakeError(EPropertyBagResult::TypeMismatch);
 	}
 
+private:
+	template<typename T>
+	TValueOrError<T, EPropertyBagResult> GetNumericalValue(const FName Desc) const;
+
+public:
 	UE_API TValueOrError<bool, EPropertyBagResult> GetValueBool(const FPropertyBagPropertyDesc& Desc) const;
+	UE_API TValueOrError<int8, EPropertyBagResult> GetValueInt8(const FPropertyBagPropertyDesc& Desc) const;
 	UE_API TValueOrError<uint8, EPropertyBagResult> GetValueByte(const FPropertyBagPropertyDesc& Desc) const;
+	UE_API TValueOrError<int16, EPropertyBagResult> GetValueInt16(const FPropertyBagPropertyDesc& Desc) const;
+	UE_API TValueOrError<uint16, EPropertyBagResult> GetValueUInt16(const FPropertyBagPropertyDesc& Desc) const;
 	UE_API TValueOrError<int32, EPropertyBagResult> GetValueInt32(const FPropertyBagPropertyDesc& Desc) const;
 	UE_API TValueOrError<uint32, EPropertyBagResult> GetValueUInt32(const FPropertyBagPropertyDesc& Desc) const;
 	UE_API TValueOrError<int64, EPropertyBagResult> GetValueInt64(const FPropertyBagPropertyDesc& Desc) const;
@@ -735,6 +752,7 @@ struct FInstancedPropertyBag
 	UE_API TValueOrError<FString, EPropertyBagResult> GetValueString(const FPropertyBagPropertyDesc& Desc) const;
 	UE_API TValueOrError<FText, EPropertyBagResult> GetValueText(const FPropertyBagPropertyDesc& Desc) const;
 	UE_API TValueOrError<uint8, EPropertyBagResult> GetValueEnum(const FPropertyBagPropertyDesc& Desc, const UEnum* RequestedEnum) const;
+	UE_API TValueOrError<int64, EPropertyBagResult> GetValueEnumInt64(const FPropertyBagPropertyDesc& Desc, const UEnum* RequestedEnum) const;
 	UE_API TValueOrError<FStructView, EPropertyBagResult> GetValueStruct(const FPropertyBagPropertyDesc& Desc, const UScriptStruct* RequestedStruct = nullptr) const;
 	UE_API TValueOrError<UObject*, EPropertyBagResult> GetValueObject(const FPropertyBagPropertyDesc& Desc, const UClass* RequestedClass = nullptr) const;
 	UE_API TValueOrError<UClass*, EPropertyBagResult> GetValueClass(const FPropertyBagPropertyDesc& Desc) const;
@@ -752,7 +770,7 @@ struct FInstancedPropertyBag
 	{
 		static_assert(TIsEnum<T>::Value, "Should only call this with enum types");
 
-		TValueOrError<uint8, EPropertyBagResult> Result = GetValueEnum(Desc, StaticEnum<T>());
+		TValueOrError<int64, EPropertyBagResult> Result = GetValueEnumInt64(Desc, StaticEnum<T>());
 		if (!Result.IsValid())
 		{
 			return MakeError(Result.GetError());
@@ -798,12 +816,20 @@ struct FInstancedPropertyBag
 		return MakeError(EPropertyBagResult::TypeMismatch);
 	}
 
+private:
+	template<typename T>
+	TValueOrError<T, EPropertyBagResult> GetNumericalValue(const FPropertyBagPropertyDesc& Desc) const;
+
+public:
 	/**
 	 * Value Setters. A property must exist in that bag before it can be set.  
-	 * Numeric types (bool, (u)int32, (u)int64, float, double) support type conversion.
+	 * Numeric types (bool, (u)int8, (u)int16, (u)int32, (u)int64, float, double) support type conversion.
 	 */
 	UE_API EPropertyBagResult SetValueBool(const FName Name, const bool bInValue);
+	UE_API EPropertyBagResult SetValueInt8(const FName Name, const int8 InValue);
 	UE_API EPropertyBagResult SetValueByte(const FName Name, const uint8 InValue);
+	UE_API EPropertyBagResult SetValueInt16(const FName Name, const int16 InValue);
+	UE_API EPropertyBagResult SetValueUInt16(const FName Name, const uint16 InValue);
 	UE_API EPropertyBagResult SetValueInt32(const FName Name, const int32 InValue);
 	UE_API EPropertyBagResult SetValueUInt32(const FName Name, const uint32 InValue);
 	UE_API EPropertyBagResult SetValueInt64(const FName Name, const int64 InValue);
@@ -813,7 +839,7 @@ struct FInstancedPropertyBag
 	UE_API EPropertyBagResult SetValueName(const FName Name, const FName InValue);
 	UE_API EPropertyBagResult SetValueString(const FName Name, const FString& InValue);
 	UE_API EPropertyBagResult SetValueText(const FName Name, const FText& InValue);
-	UE_API EPropertyBagResult SetValueEnum(const FName Name, const uint8 InValue, const UEnum* Enum);
+	UE_API EPropertyBagResult SetValueEnum(const FName Name, const int64 InValue, const UEnum* Enum);
 	UE_API EPropertyBagResult SetValueStruct(const FName Name, FConstStructView InValue);
 	UE_API EPropertyBagResult SetValueObject(const FName Name, UObject* InValue);
 	UE_API EPropertyBagResult SetValueClass(const FName Name, UClass* InValue);
@@ -831,7 +857,7 @@ struct FInstancedPropertyBag
 	EPropertyBagResult SetValueEnum(const FName Name, const T InValue)
 	{
 		static_assert(TIsEnum<T>::Value, "Should only call this with enum types");
-		return SetValueEnum(Name, static_cast<uint8>(InValue), StaticEnum<T>());
+		return SetValueEnum(Name, static_cast<int64>(InValue), StaticEnum<T>());
 	}
 
 	/** Sets struct value specified type. */
@@ -856,7 +882,10 @@ struct FInstancedPropertyBag
 	UE_API EPropertyBagResult SetValue(const FName Name, const FProperty* InSourceProperty, const void* InSourceContainerAddress);
 
 	UE_API EPropertyBagResult SetValueBool(const FPropertyBagPropertyDesc& Desc, const bool bInValue);
+	UE_API EPropertyBagResult SetValueInt8(const FPropertyBagPropertyDesc& Desc, const int8 InValue);
 	UE_API EPropertyBagResult SetValueByte(const FPropertyBagPropertyDesc& Desc, const uint8 InValue);
+	UE_API EPropertyBagResult SetValueInt16(const FPropertyBagPropertyDesc& Desc, const int16 InValue);
+	UE_API EPropertyBagResult SetValueUInt16(const FPropertyBagPropertyDesc& Desc, const uint16 InValue);
 	UE_API EPropertyBagResult SetValueInt32(const FPropertyBagPropertyDesc& Desc, const int32 InValue);
 	UE_API EPropertyBagResult SetValueUInt32(const FPropertyBagPropertyDesc& Desc, const uint32 InValue);
 	UE_API EPropertyBagResult SetValueInt64(const FPropertyBagPropertyDesc& Desc, const int64 InValue);
@@ -866,7 +895,7 @@ struct FInstancedPropertyBag
 	UE_API EPropertyBagResult SetValueName(const FPropertyBagPropertyDesc& Desc, const FName InValue);
 	UE_API EPropertyBagResult SetValueString(const FPropertyBagPropertyDesc& Desc, const FString& InValue);
 	UE_API EPropertyBagResult SetValueText(const FPropertyBagPropertyDesc& Desc, const FText& InValue);
-	UE_API EPropertyBagResult SetValueEnum(const FPropertyBagPropertyDesc& Desc, const uint8 InValue, const UEnum* Enum);
+	UE_API EPropertyBagResult SetValueEnum(const FPropertyBagPropertyDesc& Desc, const int64 InValue, const UEnum* Enum);
 	UE_API EPropertyBagResult SetValueStruct(const FPropertyBagPropertyDesc& Desc, FConstStructView InValue);
 	UE_API EPropertyBagResult SetValueObject(const FPropertyBagPropertyDesc& Desc, UObject* InValue);
 	UE_API EPropertyBagResult SetValueClass(const FPropertyBagPropertyDesc& Desc, UClass* InValue);
@@ -878,7 +907,7 @@ struct FInstancedPropertyBag
 	EPropertyBagResult SetValueEnum(const FPropertyBagPropertyDesc& Desc, const T InValue)
 	{
 		static_assert(TIsEnum<T>::Value, "Should only call this with enum types");
-		return SetValueEnum(Desc, static_cast<uint8>(InValue), StaticEnum<T>());
+		return SetValueEnum(Desc, static_cast<int64>(InValue), StaticEnum<T>());
 	}
 
 	/** Sets struct value specified type. */
@@ -1058,11 +1087,14 @@ public:
 
 	/**
 	 * Getters
-	 * Numeric types (bool, (u)int32, (u)int64, float, double) support type conversion.
+	 * Numeric types (bool, (u)int8, (u)int16, (u)int32, (u)int64, float, double) support type conversion.
 	 */
 	
 	UE_API TValueOrError<bool, EPropertyBagResult> GetValueBool(const int32 Index) const;
+	UE_API TValueOrError<int8, EPropertyBagResult> GetValueInt8(const int32 Index) const;
 	UE_API TValueOrError<uint8, EPropertyBagResult> GetValueByte(const int32 Index) const;
+	UE_API TValueOrError<int16, EPropertyBagResult> GetValueInt16(const int32 Index) const;
+	UE_API TValueOrError<uint16, EPropertyBagResult> GetValueUInt16(const int32 Index) const;
 	UE_API TValueOrError<int32, EPropertyBagResult> GetValueInt32(const int32 Index) const;
 	UE_API TValueOrError<uint32, EPropertyBagResult> GetValueUInt32(const int32 Index) const;
 	UE_API TValueOrError<int64, EPropertyBagResult> GetValueInt64(const int32 Index) const;
@@ -1073,6 +1105,7 @@ public:
 	UE_API TValueOrError<FString, EPropertyBagResult> GetValueString(const int32 Index) const;
 	UE_API TValueOrError<FText, EPropertyBagResult> GetValueText(const int32 Index) const;
 	UE_API TValueOrError<uint8, EPropertyBagResult> GetValueEnum(const int32 Index, const UEnum* RequestedEnum) const;
+	UE_API TValueOrError<int64, EPropertyBagResult> GetValueEnumInt64(const int32 Index, const UEnum* RequestedEnum) const;
 	UE_API TValueOrError<FStructView, EPropertyBagResult> GetValueStruct(const int32 Index, const UScriptStruct* RequestedStruct = nullptr) const;
 	UE_API TValueOrError<UObject*, EPropertyBagResult> GetValueObject(const int32 Index, const UClass* RequestedClass = nullptr) const;
 	UE_API TValueOrError<UClass*, EPropertyBagResult> GetValueClass(const int32 Index) const;
@@ -1084,7 +1117,7 @@ public:
 	{
 		static_assert(TIsEnum<T>::Value, "Should only call this with enum types");
 		
-		TValueOrError<uint8, EPropertyBagResult> Result = GetValueEnum(Index, StaticEnum<T>());
+		TValueOrError<int64, EPropertyBagResult> Result = GetValueEnumInt64(Index, StaticEnum<T>());
 		if (!Result.IsValid())
 		{
 			return MakeError(Result.GetError());
@@ -1130,6 +1163,11 @@ public:
 		return MakeError(EPropertyBagResult::TypeMismatch);
 	}
 
+private:
+	template<typename T>
+	TValueOrError<T, EPropertyBagResult> GetNumericalValue(const int32 Index) const;
+
+public:
 	/**
      * Returns helper class to modify and access a nested array (mutable version).
      * Note: Note: The array reference is not valid after the layout of the referenced property bag has changed!
@@ -1146,10 +1184,13 @@ public:
 
 	/**
 	 * Value Setters. A property must exist in that bag before it can be set.  
-	 * Numeric types (bool, (u)int32, (u)int64, float, double) support type conversion.
+	 * Numeric types (bool, (u)int8, (u)int16, (u)int32, (u)int64, float, double) support type conversion.
 	 */
 	UE_API EPropertyBagResult SetValueBool(const int32 Index, const bool bInValue);
+	UE_API EPropertyBagResult SetValueInt8(const int32 Index, const int8 InValue);
 	UE_API EPropertyBagResult SetValueByte(const int32 Index, const uint8 InValue);
+	UE_API EPropertyBagResult SetValueInt16(const int32 Index, const int16 InValue);
+	UE_API EPropertyBagResult SetValueUInt16(const int32 Index, const uint16 InValue);
 	UE_API EPropertyBagResult SetValueInt32(const int32 Index, const int32 InValue);
 	UE_API EPropertyBagResult SetValueUInt32(const int32 Index, const uint32 InValue);
 	UE_API EPropertyBagResult SetValueInt64(const int32 Index, const int64 InValue);
@@ -1159,7 +1200,7 @@ public:
 	UE_API EPropertyBagResult SetValueName(const int32 Index, const FName InValue);
 	UE_API EPropertyBagResult SetValueString(const int32 Index, const FString& InValue);
 	UE_API EPropertyBagResult SetValueText(const int32 Index, const FText& InValue);
-	UE_API EPropertyBagResult SetValueEnum(const int32 Index, const uint8 InValue, const UEnum* Enum);
+	UE_API EPropertyBagResult SetValueEnum(const int32 Index, const int64 InValue, const UEnum* Enum);
 	UE_API EPropertyBagResult SetValueStruct(const int32 Index, FConstStructView InValue);
 	UE_API EPropertyBagResult SetValueObject(const int32 Index, UObject* InValue);
 	UE_API EPropertyBagResult SetValueClass(const int32 Index, UClass* InValue);
@@ -1171,7 +1212,7 @@ public:
 	EPropertyBagResult SetValueEnum(const int32 Index, const T InValue)
 	{
 		static_assert(TIsEnum<T>::Value, "Should only call this with enum types");
-		return SetValueEnum(Index, static_cast<uint8>(InValue), StaticEnum<T>());
+		return SetValueEnum(Index, static_cast<int64>(InValue), StaticEnum<T>());
 	}
 
 	/** Sets struct value specified type. */
@@ -1226,7 +1267,10 @@ public:
 
 	/** Add value to set. If value is already present, it will not be added */
 	UE_API EPropertyBagResult AddValueBool(const bool bInValue);
+	UE_API EPropertyBagResult AddValueInt8(const int8 InValue);
 	UE_API EPropertyBagResult AddValueByte(const uint8 InValue);
+	UE_API EPropertyBagResult AddValueInt16(const int16 InValue);
+	UE_API EPropertyBagResult AddValueUInt16(const uint16 InValue);
 	UE_API EPropertyBagResult AddValueInt32(const int32 InValue);
 	UE_API EPropertyBagResult AddValueUInt32(const uint32 InValue);
 	UE_API EPropertyBagResult AddValueInt64(const int64 InValue);
@@ -1247,7 +1291,7 @@ public:
 	EPropertyBagResult AddValueEnum(const T InValue)
 	{
 		static_assert(TIsEnum<T>::Value, "Should only call this with enum types");
-		return AddValueEnum(static_cast<uint8>(InValue), StaticEnum<T>());
+		return AddValueEnum(static_cast<int64>(InValue), StaticEnum<T>());
 	}
 
 	/** Adds struct value specified type. */
diff --git a/Engine/Source/Runtime/CoreUObject/Public/StructUtils/PropertyBagMapRef.h b/Engine/Source/Runtime/CoreUObject/Public/StructUtils/PropertyBagMapRef.h
--- a/Engine/Source/Runtime/CoreUObject/Public/StructUtils/PropertyBagMapRef.h
+++ b/Engine/Source/Runtime/CoreUObject/Public/StructUtils/PropertyBagMapRef.h
@@ -356,8 +356,11 @@ private:
 		// This generic mechanism should consider the entirety of the PropertyBag interface and be applied consistently.
 		if ((std::is_same_v<T, bool> && InDesc.ValueType == EPropertyBagPropertyType::Bool)
 			|| (std::is_same_v<T, uint8> && InDesc.ValueType == EPropertyBagPropertyType::Byte)
+			|| (std::is_same_v<T, int8> && InDesc.ValueType == EPropertyBagPropertyType::Int8)
+			|| (std::is_same_v<T, int16> && InDesc.ValueType == EPropertyBagPropertyType::Int16)
 			|| (std::is_same_v<T, int32> && InDesc.ValueType == EPropertyBagPropertyType::Int32)
 			|| (std::is_same_v<T, int64> && InDesc.ValueType == EPropertyBagPropertyType::Int64)
+			|| (std::is_same_v<T, uint16> && InDesc.ValueType == EPropertyBagPropertyType::UInt16)
 			|| (std::is_same_v<T, uint32> && InDesc.ValueType == EPropertyBagPropertyType::UInt32)
 			|| (std::is_same_v<T, uint64> && InDesc.ValueType == EPropertyBagPropertyType::UInt64)
 			|| (std::is_same_v<T, float> && InDesc.ValueType == EPropertyBagPropertyType::Float)
diff --git a/Engine/Source/Runtime/CoreUObject/Public/StructUtils/StructUtils.h b/Engine/Source/Runtime/CoreUObject/Public/StructUtils/StructUtils.h
--- a/Engine/Source/Runtime/CoreUObject/Public/StructUtils/StructUtils.h
+++ b/Engine/Source/Runtime/CoreUObject/Public/StructUtils/StructUtils.h
@@ -3,6 +3,7 @@
 #pragma once
 
 #include "Templates/UnrealTypeTraits.h"
+#include "Templates/ValueOrError.h"
 #include "UObject/Class.h"
 
 class UPropertyBagHierarchyRoot;
@@ -129,4 +130,20 @@ namespace UE::StructUtils
 	COREUOBJECT_API void AppendPropertyDescs(FXxHash64Builder& HashBuilder, const FPropertyBagPropertyDesc& Desc);
 
 	COREUOBJECT_API uint64 CalcPropertyDescArrayHash(const TConstArrayView<FPropertyBagPropertyDesc> Descs);
+
+	/**
+	 * Return the numerical value (not in container).
+	 * Works with property of type: bool, int8, uint8, int16, uint16, int32, uint32, int64, uint64, float, double, enum.
+	 */
+	COREUOBJECT_API TValueOrError<int64, void> GetPropertyValueAsInt64(TNotNull<const FProperty*> InProperty, TNotNull<const void*> Address);
+	COREUOBJECT_API TValueOrError<uint64, void> GetPropertyValueAsUInt64(TNotNull<const FProperty*> InProperty, TNotNull<const void*> Address);
+	COREUOBJECT_API TValueOrError<double, void> GetPropertyValueAsDouble(TNotNull<const FProperty*> InProperty, TNotNull<const void*> Address);
+
+	/**
+	 * Set the numerical value (not in container).
+	 * Works with property of type: bool, int8, uint8, int16, uint16, int32, uint32, int64, uint64, float, double, enum.
+	 */
+	COREUOBJECT_API bool SetPropertyValueFromInt64(TNotNull<const FProperty*> InProperty, TNotNull<void*> Address, int64 Value);
+	COREUOBJECT_API bool SetPropertyValueFromUInt64(TNotNull<const FProperty*> InProperty, TNotNull<void*> Address, uint64 Value);
+	COREUOBJECT_API bool SetPropertyValueFromDouble(TNotNull<const FProperty*> InProperty, TNotNull<void*> Address, double Value);
 }
