[SoftObjectPath/ALT]
- Always allow resolving softobjectpath from the loader to allow overlapping of save with async loading

#Synced-CL 50902286

diff --git a/Engine/Source/Runtime/CoreUObject/Private/UObject/SoftObjectPath.cpp b/Engine/Source/Runtime/CoreUObject/Private/UObject/SoftObjectPath.cpp
--- a/Engine/Source/Runtime/CoreUObject/Private/UObject/SoftObjectPath.cpp
+++ b/Engine/Source/Runtime/CoreUObject/Private/UObject/SoftObjectPath.cpp
@@ -872,7 +872,7 @@ UObject* FSoftObjectPath::ResolveObject() const
 {
 	// Don't try to resolve if we're saving a package because StaticFindObject can't be used here
 	// and we usually don't want to force references to weak pointers while saving.
-	if (IsNull() || UE::IsSavingPackage())
+	if (IsNull() || (UE::IsSavingPackage() && !IsInAsyncLoadingThread()))
 	{
 		return nullptr;
 	}
