[LinkerManager]
- Replace a global flush by flushing only the required linker

#Synced-CL 50902286

diff --git a/Engine/Source/Runtime/CoreUObject/Private/UObject/LinkerManager.cpp b/Engine/Source/Runtime/CoreUObject/Private/UObject/LinkerManager.cpp
--- a/Engine/Source/Runtime/CoreUObject/Private/UObject/LinkerManager.cpp
+++ b/Engine/Source/Runtime/CoreUObject/Private/UObject/LinkerManager.cpp
@@ -95,15 +95,13 @@ bool FLinkerManager::Exec_Dev(class UWorld* InWorld, const TCHAR* Cmd, FOutputDe
 	return false;
 }
 
+extern void ConditionalFlushAsyncLoadingForLinkers(TConstArrayView<FLinkerLoad*> InLinkers);
+
 void FLinkerManager::ResetLinkerExports(UPackage* InPackage)
 {
 	if (FLinkerLoad* LinkerToReset = FLinkerLoad::FindExistingLinkerForPackage(InPackage))
 	{
-		// if the linker owner thread is not the main thread, we need to flush async loading (todo: for that package) before we can reset the linker
-		if (LinkerToReset->GetOwnerThreadId() != GGameThreadId)
-		{
-			FlushAsyncLoading();
-		}
+		ConditionalFlushAsyncLoadingForLinkers({LinkerToReset});
 		LinkerToReset->DetachExports();
 	}
 }
