Splitting remote object deserialization away from transactional PostMigrate. Deserialization will now be performed outside of the transaction by default unless specifically requested with a cvar

#Synced-CL 50967057

diff --git a/Engine/Source/Runtime/CoreUObject/Private/UObject/RemoteObject.cpp b/Engine/Source/Runtime/CoreUObject/Private/UObject/RemoteObject.cpp
--- a/Engine/Source/Runtime/CoreUObject/Private/UObject/RemoteObject.cpp
+++ b/Engine/Source/Runtime/CoreUObject/Private/UObject/RemoteObject.cpp
@@ -217,7 +217,7 @@ public:
 		for (TPair<FRemoteObjectId,FRemoteObjectStub*>& Pair : RemoteObjects)
 		{
 			FRemoteObjectStub* Stub = Pair.Value;
-			if (Stub->MigratedData && !Stub->MigratedData->bCompleted)
+			if (Stub->MigratedData && Stub->MigratedData->State != EMigratedDataState::Completed)
 			{
 				if (!Operation(Stub))
 				{
diff --git a/Engine/Source/Runtime/CoreUObject/Private/UObject/RemoteObjectPrivate.h b/Engine/Source/Runtime/CoreUObject/Private/UObject/RemoteObjectPrivate.h
--- a/Engine/Source/Runtime/CoreUObject/Private/UObject/RemoteObjectPrivate.h
+++ b/Engine/Source/Runtime/CoreUObject/Private/UObject/RemoteObjectPrivate.h
@@ -14,15 +14,38 @@ DECLARE_LOG_CATEGORY_EXTERN(LogRemoteObject, Display, All);
 
 struct FRemoteObjectPathName;
 
+enum class EMigratedDataState
+{
+	Uninitialized = 0,
+	Received = 1,
+	Deserialized = 2,
+	PostMigrated = 3,
+	Completed = 4
+};
+const TCHAR* EnumToString(EMigratedDataState InState);
+
 /** Structure that contains object data and migration context for deserializing and PostMigrating objects inside of a transaction */
 struct FTransactionalMigrationData : public FRefCountedObject
 {
+	FTransactionalMigrationData() = default;
+	FTransactionalMigrationData(FRemoteObjectData& InData, FUObjectMigrationContext& InContext);
+
 	/** Serialized data of the migrated objects */
 	FRemoteObjectData Data;
 	/** The context for the migrated objects */
 	FUObjectMigrationContext Context;
-	/** If true, the data has been fully deserialized and should no longer be used */
-	bool bCompleted = false;
+	
+	/** List of remote ids of the received objects */
+	TArray<FRemoteObjectId> ReceivedObjectRemoteIds;
+	/** List of received objects */
+	TArray<UObject*> ReceivedObjects;
+	/** The index of the object requested by the migration request from the received objects list */
+	int32 RequestedObjectIndex = -1;
+
+	/** The current state of the migrated data */
+	EMigratedDataState State = EMigratedDataState::Uninitialized;
+
+	void SetReceivedObjectsResidence(EResidence Residence);
 };
 
 namespace UE::RemoteObject::Private
diff --git a/Engine/Source/Runtime/CoreUObject/Private/UObject/RemoteObjectTransfer.cpp b/Engine/Source/Runtime/CoreUObject/Private/UObject/RemoteObjectTransfer.cpp
--- a/Engine/Source/Runtime/CoreUObject/Private/UObject/RemoteObjectTransfer.cpp
+++ b/Engine/Source/Runtime/CoreUObject/Private/UObject/RemoteObjectTransfer.cpp
@@ -28,7 +28,13 @@ int32 GTransactionalMigration = 1;
 static FAutoConsoleVariableRef CVarTransactionalMigrationEnabled(
 	TEXT("DSTM.TransactionalMigration"),
 	GTransactionalMigration,
-	TEXT("When non-zero, migrated objects will be created, deserialized and PostMigrated inside of a transaction"));
+	TEXT("When non-zero, migrated objects will be PostMigrated inside of a transaction"));
+
+int32 GTransactionalDeserialization = 0;
+static FAutoConsoleVariableRef CVarTransactionalDeserializationEnabled(
+	TEXT("DSTM.TransactionalDeserialization"),
+	GTransactionalDeserialization,
+	TEXT("When non-zero, migrated objects will be created and deserialized inside of a transaction"));
 
 // Prevent us from having to match this definition to the declartion in the header just to initialize the delegate
 #define DEFINE_MIGRATION_DELEGATE(x) decltype(x) x
@@ -58,13 +64,15 @@ TRefCountPtr<FTransactionalMigrationData> StoreMigratedObjectData(FRemoteObjectD
 */
 void ResetMigratedData(TRefCountPtr<FTransactionalMigrationData> MigratedData);
 
+void DeserializeObjectsFromMigratedData(FTransactionalMigrationData& InOutMigratedData);
+
 /**
 * Creates, deserializes and PostMigrates an object(s) from object data received from another server
 * @param ObjectData Data received from another server
 * @param MigrationContext Migration context
 * @return The object requested by the current server or pushed by the sending server (when transferring ownership)
 */
-UObject* MigrateObjectFromObjectData(FRemoteObjectData& ObjectData, const FUObjectMigrationContext& MigrationContext);
+UObject* MigrateObjectFromObjectData(TRefCountPtr<FTransactionalMigrationData> MigratedData);
 
 /**
  * Remote object transfer queue. 
@@ -252,6 +260,10 @@ public:
 					}
 					else
 					{
+						checkf(MigratedData->State == EMigratedDataState::Received || (MigratedData->State == EMigratedDataState::Deserialized && !GTransactionalDeserialization),
+							TEXT("Sending migrated data back without PostMigrating it but the data is in unexpected state (%s)"), EnumToString(MigratedData->State));
+						// Pretend we PostMigrated the data for the reset function
+						MigratedData->State = EMigratedDataState::PostMigrated;
 						ResetMigratedData(MigratedData);
 						ObjectData = MoveTemp(MigratedData->Data);						
 						
@@ -887,6 +899,14 @@ public:
 		{
 			TRefCountPtr<FTransactionalMigrationData> StoredData = StoreMigratedObjectData(ObjectData, MigrationContext);
 	
+			if (!GTransactionalDeserialization)
+			{
+				DeserializeObjectsFromMigratedData(*StoredData);
+				// Deserialization sets the residence of the deserialized objects to Local as soon as they are created to avoid re-resolving the deserialized objects when accessing them as outers
+				// or in Serialize() function overrides. Since we queue the PostMigrate step to be transactional we need to mark the objects as LocalNotReady again so that if something else attempts
+				// to access the deserialized objects BEFORE the transactional PostMigrate runs the engine will attempt to resolve (PostMigrate) the deserialized objects
+				StoredData->SetReceivedObjectsResidence(EResidence::LocalNotReady);
+			}
 			static FName TransactionalWorkName("TransactionallyMigrateObjects");
 			UE::RemoteExecutor::EnqueueMigrationWork(TransactionalWorkName, [StoredData]()
 				{
@@ -896,7 +916,8 @@ public:
 		else
 #endif
 		{
-			MigrateObjectFromObjectData(ObjectData, MigrationContext);
+			TRefCountPtr<FTransactionalMigrationData> MigratedData = new FTransactionalMigrationData(ObjectData, MigrationContext);
+			MigrateObjectFromObjectData(MigratedData);
 		}
 
 		if (bFoundRequest)
@@ -943,9 +964,7 @@ TRefCountPtr<FTransactionalMigrationData> StoreMigratedObjectData(FRemoteObjectD
 
 	FRemoteObjectStub* RootStub = FindOrAddRemoteObjectStub(MigrationContext.ObjectId, FRemoteServerId::GetLocalServerId());
 
-	TRefCountPtr<FTransactionalMigrationData> MigratedData = new FTransactionalMigrationData();
-	MigratedData->Data = MoveTemp(ObjectData);
-	MigratedData->Context = MoveTemp(MigrationContext);
+	TRefCountPtr<FTransactionalMigrationData> MigratedData = new FTransactionalMigrationData(ObjectData, MigrationContext);
 
 	UE_LOG(LogRemoteObject, VeryVerbose, TEXT("%s received %d objects"), *FRemoteServerId::GetLocalServerId().ToString(), MigratedData->Data.SerializedObjectHeaders.Num());
 
@@ -993,24 +1012,24 @@ TRefCountPtr<FTransactionalMigrationData> StoreMigratedObjectData(FRemoteObjectD
 	return MigratedData;
 }
 
-UObject* MigrateObjectFromObjectData(FRemoteObjectData& ObjectData, const FUObjectMigrationContext& MigrationContext)
+void DeserializeObjectsFromMigratedData(FTransactionalMigrationData& InOutMigratedData)
 {
 	using namespace UE::RemoteObject::Serialization;
 	using namespace UE::RemoteObject::Private;
 
-	uint32 ByteSize = 0;
-	for (const FRemoteObjectBytes& Chunk : ObjectData.Bytes)
-	{
-		ByteSize += Chunk.Bytes.Num();
-	}
+	checkf(InOutMigratedData.State == EMigratedDataState::Received, TEXT("Starting to deserialize remote object data but it's in an unexpected state %s (expected: %s)"), EnumToString(InOutMigratedData.State), EnumToString(EMigratedDataState::Received));
 
 	static FName GetObjectRegionName("ObjectTransferQueue_RecvObject");
 	UE::RemoteExecutor::OnRegionBeginDelegate.ExecuteIfBound(GetObjectRegionName);
 
- 	FRemoteServerId LocalServerId = FRemoteServerId::GetLocalServerId();
-	TArray<FRemoteObjectId> ReceivedObjectRemoteIds;
-	TArray<UObject*> ReceivedObjects;
-	int32 RequestedObjectIndex = DeserializeObjectData(ObjectData, &MigrationContext, ReceivedObjectRemoteIds, ReceivedObjects);
+	const FRemoteServerId LocalServerId = FRemoteServerId::GetLocalServerId();
+	TArray<FRemoteObjectId>& ReceivedObjectRemoteIds = InOutMigratedData.ReceivedObjectRemoteIds;
+	TArray<UObject*>& ReceivedObjects = InOutMigratedData.ReceivedObjects;
+	FRemoteObjectData& ObjectData = InOutMigratedData.Data;
+	const FUObjectMigrationContext& MigrationContext = InOutMigratedData.Context;
+
+	InOutMigratedData.RequestedObjectIndex = DeserializeObjectData(ObjectData, &MigrationContext, ReceivedObjectRemoteIds, ReceivedObjects);
+	InOutMigratedData.State = EMigratedDataState::Deserialized;
 
 	// Debug section
 	checkf(ReceivedObjects.Num() > 0, TEXT("PeerId:%s unable to deserialize object data (%d bytes)"), *FRemoteServerId::GetLocalServerId().ToString(), ObjectData.GetNumBytes());
@@ -1021,7 +1040,7 @@ UObject* MigrateObjectFromObjectData(FRemoteObjectData& ObjectData, const FUObje
 		if (UE_LOG_ACTIVE(LogRemoteObject, VeryVerbose))
 		{
 			UE_LOG(LogRemoteObject, VeryVerbose, TEXT("PeerId:%s deserialized object data (%d bytes)"), *FRemoteServerId::GetLocalServerId().ToString(), ObjectData.GetNumBytes());
-			for(int32 SubObjectIndex = 0; SubObjectIndex < ReceivedObjects.Num(); SubObjectIndex++)
+			for (int32 SubObjectIndex = 0; SubObjectIndex < ReceivedObjects.Num(); SubObjectIndex++)
 			{
 				FRemoteObjectId SubObjectId = UE::CoreUObject::Private::FObjectHandleUtils::GetRemoteId(ReceivedObjects[SubObjectIndex]);
 				check(ReceivedObjectRemoteIds[SubObjectIndex] == SubObjectId);
@@ -1029,7 +1048,7 @@ UObject* MigrateObjectFromObjectData(FRemoteObjectData& ObjectData, const FUObje
 				UE_LOG(LogRemoteObject, VeryVerbose, TEXT("         subobject[%2d] : %s %s"),
 					SubObjectIndex,
 					*SubObjectId.ToString(),
-					(SubObjectIndex == RequestedObjectIndex) ? TEXT("<requested>") : TEXT(""));
+					(SubObjectIndex == InOutMigratedData.RequestedObjectIndex) ? TEXT("<requested>") : TEXT(""));
 			}
 
 			UE_LOG(LogRemoteObject, VeryVerbose, TEXT("         <done listing subobjects>"));
@@ -1072,8 +1091,21 @@ UObject* MigrateObjectFromObjectData(FRemoteObjectData& ObjectData, const FUObje
 			}
 		}
 	}
+}
+
+UObject* PostMigrateObjectsFromMigratedData(FTransactionalMigrationData& InOutMigratedData)
+{
+	using namespace UE::RemoteObject::Serialization;
+	using namespace UE::RemoteObject::Private;
+
+	checkf(InOutMigratedData.State == EMigratedDataState::Deserialized, TEXT("Migrated data has not been deserialized yet (State: %s)"), EnumToString(InOutMigratedData.State));
+
+	const FRemoteServerId LocalServerId = FRemoteServerId::GetLocalServerId();
+	const FRemoteObjectData& ObjectData = InOutMigratedData.Data;
+	const FUObjectMigrationContext& MigrationContext = InOutMigratedData.Context;
+	const TArray<UObject*>& ReceivedObjects = InOutMigratedData.ReceivedObjects;
+	const TArray<FRemoteObjectId>& ReceivedObjectRemoteIds = InOutMigratedData.ReceivedObjectRemoteIds;
 
-#if UE_WITH_REMOTE_OBJECT_HANDLE
 	if (LocalServerId != MigrationContext.OwnerServerId)
 	{
 		for (UObject* ReceivedObject : ReceivedObjects)
@@ -1081,6 +1113,7 @@ UObject* MigrateObjectFromObjectData(FRemoteObjectData& ObjectData, const FUObje
 			MarkAsBorrowed(ReceivedObject);
 		}
 	}
+#if UE_WITH_REMOTE_OBJECT_HANDLE
 	{
 		TRACE_CPUPROFILER_EVENT_SCOPE(PostMigrateObjects);
 
@@ -1091,12 +1124,12 @@ UObject* MigrateObjectFromObjectData(FRemoteObjectData& ObjectData, const FUObje
 	}
 #endif
 
- 	if (MigrationContext.RemoteServerId == DatabaseId)
- 	{
- 		// Let other systems know that we've loaded an object from storage
- 		OnObjectsLoadedFromDiskDelegate.Broadcast(ReceivedObjects, ObjectData, MigrationContext);
- 	}
- 	else
+	if (MigrationContext.RemoteServerId == DatabaseId)
+	{
+		// Let other systems know that we've loaded an object from storage
+		OnObjectsLoadedFromDiskDelegate.Broadcast(ReceivedObjects, ObjectData, MigrationContext);
+	}
+	else
 	{
 		// Notify that objects have been migrated from a remote server
 		OnObjectsReceivedDelegate.Broadcast(ReceivedObjects, ObjectData, MigrationContext);
@@ -1133,15 +1166,41 @@ UObject* MigrateObjectFromObjectData(FRemoteObjectData& ObjectData, const FUObje
 		}
 	}
 
-	UE::RemoteExecutor::OnRegionEndDelegate.ExecuteIfBound(FString::Format(TEXT("{0} {1} from {2} ({3} bytes)"), {*ReceivedObjects[0]->GetName(), *RootObjectRemoteId.ToString(), *MigrationContext.RemoteServerId.ToString(), ByteSize}));
+	{
+		uint32 ByteSize = 0;
+		for (const FRemoteObjectBytes& Chunk : ObjectData.Bytes)
+		{
+			ByteSize += Chunk.Bytes.Num();
+		}
+		UE::RemoteExecutor::OnRegionEndDelegate.ExecuteIfBound(FString::Format(TEXT("{0} {1} from {2} ({3} bytes)"), { *ReceivedObjects[0]->GetName(), *RootObjectRemoteId.ToString(), *MigrationContext.RemoteServerId.ToString(), ByteSize }));
+	}
+
+	UObject* Result = ReceivedObjects[InOutMigratedData.RequestedObjectIndex];
 
-	UObject* Result = ReceivedObjects[RequestedObjectIndex];
+	InOutMigratedData.State = EMigratedDataState::PostMigrated;
 
 	checkf(ReceivedObjectRemoteIds.Contains(MigrationContext.ObjectId), TEXT("PeerId:%s requested object %s migration from PeerId:%s but received %s"),
 		*FRemoteServerId::GetLocalServerId().ToString(), *MigrationContext.ObjectId.ToString(), *MigrationContext.RemoteServerId.ToString(), *UE::CoreUObject::Private::FObjectHandleUtils::GetRemoteId(Result).ToString());
 	checkf(MigrationContext.OwnerServerId == Handle::GetOwnerServerId(Result), TEXT("Expected owner id %s doesn't match post-deserialize owner: %s"), *MigrationContext.OwnerServerId.ToString(), *Handle::GetOwnerServerId(Result).ToString());
 
 	return Result;
+}
+
+UObject* MigrateObjectFromObjectData(TRefCountPtr<FTransactionalMigrationData> MigratedData)
+{
+	if (MigratedData->State != EMigratedDataState::Deserialized)
+	{
+		DeserializeObjectsFromMigratedData(*MigratedData);
+	}
+	else
+	{
+		// If the migrated data has already been deserialized then we must be PostMigrating transactionally in which case
+		// we need to set the migrated objects' residence back to local before we start calling PostMigrate functions
+		// to avoid re-resolving the deserialized objects
+		check(AutoRTFM::IsClosed());
+		MigratedData->SetReceivedObjectsResidence(EResidence::Local);
+	}
+	return PostMigrateObjectsFromMigratedData(*MigratedData);
 } 
 
 void ResetMigratedData(TRefCountPtr<FTransactionalMigrationData> MigratedData)
@@ -1150,7 +1209,8 @@ void ResetMigratedData(TRefCountPtr<FTransactionalMigrationData> MigratedData)
 	using namespace UE::RemoteObject::Private;
 
 	checkf(MigratedData && MigratedData->Data.SerializedObjectHeaders.Num() > 0, TEXT("Trying to reset invalid migrated data"));
-	checkf(!MigratedData->bCompleted, TEXT("Trying to reset data that has already been fully deserialized"));
+	checkf(MigratedData->State != EMigratedDataState::Completed, TEXT("Trying to reset data that has already been fully PostMigrated"));
+	checkf(MigratedData->State == EMigratedDataState::PostMigrated, TEXT("Trying to reset data that has not yet been fully PostMigrated (state: %s)"), EnumToString(MigratedData->State));
 
 	for (FRemoteObjectId SerializedId : MigratedData->Data)
 	{
@@ -1159,7 +1219,7 @@ void ResetMigratedData(TRefCountPtr<FTransactionalMigrationData> MigratedData)
 		checkf(Stub->MigratedData == MigratedData || !Stub->MigratedData, TEXT("Stub for serialized object id %s does not have the expected migrated object data"), *SerializedId.ToString());		
 		Stub->MigratedData = nullptr;
 	}
-	MigratedData->bCompleted = true;
+	MigratedData->State = EMigratedDataState::Completed;
 }
 
 void TransactionallyMigrateObjects(TRefCountPtr<FTransactionalMigrationData> MigratedData)
@@ -1171,10 +1231,10 @@ void TransactionallyMigrateObjects(TRefCountPtr<FTransactionalMigrationData> Mig
 	// and deserializes their data before the queued TransactionallyMigrateObjects transaction is executed.
 	// It's also possible that the migrated object data is sent to another server before we had a chance to deserialize it.
 	// In that case (bCompleted == true) just skip the deserialization and run some extra checks (if enabled)
-	if (!MigratedData->bCompleted)
+	if (MigratedData->State != EMigratedDataState::Completed)
 	{
 		SetTransactionallyPostMigratingObjects(true);
-		MigrateObjectFromObjectData(MigratedData->Data, MigratedData->Context);
+		MigrateObjectFromObjectData(MigratedData);
 		SetTransactionallyPostMigratingObjects(false);
 
 		ResetMigratedData(MigratedData);
@@ -1633,6 +1693,54 @@ void StoreUnreachableRemoteObjectsToDatabase(const TArrayView<FUObjectItem*>& Un
 
 } // namespace UE::RemoteObject::Transfer::Private
 
+const TCHAR* EnumToString(EMigratedDataState InState)
+{
+	switch (InState)
+	{
+	case EMigratedDataState::Uninitialized:
+		return TEXT("Uninitialized");
+	case EMigratedDataState::Received:
+		return TEXT("Received");
+	case EMigratedDataState::Deserialized:
+		return TEXT("Deserialized");
+	case EMigratedDataState::PostMigrated:
+		return TEXT("PostMigrated");
+	case EMigratedDataState::Completed:
+		return TEXT("Completed");
+	default:
+		checkf(false, TEXT("EnumToString: Unknown EMigratedDataState"));
+		break;
+	};
+	return TEXT("");
+}
+
+FTransactionalMigrationData::FTransactionalMigrationData(FRemoteObjectData& InData, FUObjectMigrationContext& InContext)
+{
+	Data = MoveTemp(InData);
+	Context = MoveTemp(InContext);
+	State = EMigratedDataState::Received;
+}
+
+void FTransactionalMigrationData::SetReceivedObjectsResidence(EResidence Residence)
+{
+	using namespace UE::RemoteObject::Handle;
+	using namespace UE::RemoteObject::Private;
+
+	checkf(State == EMigratedDataState::Deserialized, TEXT("Received objects residence can only be set when migrated data state is Deserialized"));
+	checkf(Residence == EResidence::Local || Residence == EResidence::LocalNotReady, TEXT("Migrated object data can only set %s or %s residence (requested: %s)"), EnumToString(EResidence::Local), EnumToString(EResidence::LocalNotReady), EnumToString(Residence));
+
+	FTransactionalMigrationData* DataToSet = (Residence == EResidence::LocalNotReady ? this : nullptr);
+	for (int32 DeserializedObjectIndex = 0; DeserializedObjectIndex < ReceivedObjects.Num(); ++DeserializedObjectIndex)
+	{
+		UObject* DeserializedObject = ReceivedObjects[DeserializedObjectIndex];
+		FRemoteObjectId DeserializedObjectId = ReceivedObjectRemoteIds[DeserializedObjectIndex];
+		checkf(DeserializedObjectId == FRemoteObjectId(DeserializedObject), TEXT("Object id (%s) does not match the respective ReceivedObjectRemoteId (%s)"), *FRemoteObjectId(DeserializedObject).ToString(), *DeserializedObjectId.ToString());
+		FRemoteObjectStub* Stub = FindRemoteObjectStub(DeserializedObjectId);
+		checkf(Stub, TEXT("Stub for deserialized object %s found!"), *DeserializedObjectId.ToString());
+		Stub->MigratedData = DataToSet;
+		SetResidence(DeserializedObject, Residence, FRemoteServerId::GetLocalServerId());
+	}
+}
 
 FRemoteObjectReference::FRemoteObjectReference(const FObjectPtr& Ptr)
 {
