Fixing remote executor region callbacks when transactional deserialization of remote objects is disabled

#Synced-CL 50967057

diff --git a/Engine/Source/Runtime/CoreUObject/Private/UObject/RemoteObjectTransfer.cpp b/Engine/Source/Runtime/CoreUObject/Private/UObject/RemoteObjectTransfer.cpp
--- a/Engine/Source/Runtime/CoreUObject/Private/UObject/RemoteObjectTransfer.cpp
+++ b/Engine/Source/Runtime/CoreUObject/Private/UObject/RemoteObjectTransfer.cpp
@@ -1091,6 +1091,8 @@ void DeserializeObjectsFromMigratedData(FTransactionalMigrationData& InOutMigrat
 			}
 		}
 	}
+
+	UE::RemoteExecutor::OnRegionEndDelegate.ExecuteIfBound(FString::Format(TEXT("{0} {1} from {2} ({3} bytes)"), { *ReceivedObjects[0]->GetName(), *ObjectData.GetRootSerializedObjectId().ToString(), *MigrationContext.RemoteServerId.ToString(), ObjectData.GetNumBytes()}));
 }
 
 UObject* PostMigrateObjectsFromMigratedData(FTransactionalMigrationData& InOutMigratedData)
@@ -1100,6 +1102,9 @@ UObject* PostMigrateObjectsFromMigratedData(FTransactionalMigrationData& InOutMi
 
 	checkf(InOutMigratedData.State == EMigratedDataState::Deserialized, TEXT("Migrated data has not been deserialized yet (State: %s)"), EnumToString(InOutMigratedData.State));
 
+	static FName PostMigrateRegionName("ObjectTransferQueue_PostMigrate");
+	UE::RemoteExecutor::OnRegionBeginDelegate.ExecuteIfBound(PostMigrateRegionName);
+
 	const FRemoteServerId LocalServerId = FRemoteServerId::GetLocalServerId();
 	const FRemoteObjectData& ObjectData = InOutMigratedData.Data;
 	const FUObjectMigrationContext& MigrationContext = InOutMigratedData.Context;
@@ -1166,14 +1171,7 @@ UObject* PostMigrateObjectsFromMigratedData(FTransactionalMigrationData& InOutMi
 		}
 	}
 
-	{
-		uint32 ByteSize = 0;
-		for (const FRemoteObjectBytes& Chunk : ObjectData.Bytes)
-		{
-			ByteSize += Chunk.Bytes.Num();
-		}
-		UE::RemoteExecutor::OnRegionEndDelegate.ExecuteIfBound(FString::Format(TEXT("{0} {1} from {2} ({3} bytes)"), { *ReceivedObjects[0]->GetName(), *RootObjectRemoteId.ToString(), *MigrationContext.RemoteServerId.ToString(), ByteSize }));
-	}
+	UE::RemoteExecutor::OnRegionEndDelegate.ExecuteIfBound(FString::Format(TEXT("{0} {1} num objects: {2}"), { *ReceivedObjects[0]->GetName(), *RootObjectRemoteId.ToString(), ReceivedObjects.Num(),}));
 
 	UObject* Result = ReceivedObjects[InOutMigratedData.RequestedObjectIndex];
 
