[UObjectPropertyTest]

Add bool bFailIfLockHeld argument to MockAssetRegistryInterface::TryGetAssetPackageData overrides

#rnx
#Synced-CL 50826644

#Synced-CL 50826644

diff --git a/Engine/Source/Runtime/CoreUObject/Tests/ObjectPropertyTest.cpp b/Engine/Source/Runtime/CoreUObject/Tests/ObjectPropertyTest.cpp
--- a/Engine/Source/Runtime/CoreUObject/Tests/ObjectPropertyTest.cpp
+++ b/Engine/Source/Runtime/CoreUObject/Tests/ObjectPropertyTest.cpp
@@ -323,12 +323,12 @@ public:
 		return UE::AssetRegistry::EExists::DoesNotExist;
 	}
 
-	virtual UE::AssetRegistry::EExists TryGetAssetPackageData(FName PackageName, FAssetPackageData& OutPackageData) const override
+	virtual UE::AssetRegistry::EExists TryGetAssetPackageData(FName PackageName, FAssetPackageData& OutPackageData, bool bFailIfLockHeld) const override
 	{
 		return UE::AssetRegistry::EExists::Exists;
 	}
 
-	virtual UE::AssetRegistry::EExists TryGetAssetPackageData(FName PackageName, FAssetPackageData& OutPackageData, FName& OutCorrectCasePackageName) const override
+	virtual UE::AssetRegistry::EExists TryGetAssetPackageData(FName PackageName, FAssetPackageData& OutPackageData, FName& OutCorrectCasePackageName, bool bFailIfLockHeld) const override
 	{
 		return UE::AssetRegistry::EExists::Exists;
 	}
