[SavePackageUtilities/ALT]
- Get rid of the fake GC lock while saving to allow async loading to run in parallel
- Provide a TLS version of GIsSavingPackage that can be used on any other thread than GT for protection against object find during saving on worker threads

#Synced-CL 50902286

diff --git a/Engine/Source/Runtime/CoreUObject/Private/UObject/SavePackage/SavePackageUtilities.cpp b/Engine/Source/Runtime/CoreUObject/Private/UObject/SavePackage/SavePackageUtilities.cpp
--- a/Engine/Source/Runtime/CoreUObject/Private/UObject/SavePackage/SavePackageUtilities.cpp
+++ b/Engine/Source/Runtime/CoreUObject/Private/UObject/SavePackage/SavePackageUtilities.cpp
@@ -178,21 +178,29 @@ static FThreadSafeCounter OutstandingAsyncWrites;
 // Initialization for GIsSavingPackage
 bool GIsSavingPackage = false;
 
+// GIsSavingPackage is still accessible globally but is meant only for the game thread
+// so provide another one for any other thread than GT.
+thread_local bool GIsSavingPackageTLS = false;
+
 namespace UE
 {
 	bool IsSavingPackage(UObject* InOuter)
 	{
 		if (InOuter == nullptr)
 		{
-			// That global is only meant to be set and read from the game-thread...
-			// otherwise it could interfere with normal operations coming from
-			// other threads like async loading, etc...
-			// We need to do a volatile read otherwise Clang/LLVM can decide to do
-			// some crazy optim where both GIsSavingPackage and IsInGameThread()
-			// are read into registry and evaluated together instead of going
-			// one after the other like we want to to avoid touching the GIsSavingPackage
-			// value from other threads and causing TSAN warnings.
-			return IsInGameThread() && *(volatile bool*)&GIsSavingPackage;
+			if (IsInGameThread())
+			{
+				// We need to do a volatile read otherwise Clang/LLVM can decide to do
+				// some crazy optim where both GIsSavingPackage and IsInGameThread()
+				// are read into registry and evaluated together instead of going
+				// one after the other like we want to to avoid touching the GIsSavingPackage
+				// value from other threads and causing TSAN warnings.
+				return *(volatile bool*)&GIsSavingPackage;
+			}
+			else
+			{
+				return GIsSavingPackageTLS;
+			}
 		}
 		return InOuter->GetPackage()->HasAnyPackageFlags(PKG_IsSaving);
 	}
@@ -1230,20 +1238,17 @@ FScopedSavingFlag::FScopedSavingFlag(bool InSavingConcurrent, UPackage* InSavedP
 {
 	check(!IsGarbageCollecting());
 
-	// We need the same lock as GC so that no StaticFindObject can happen in parallel to saving a package
 	if (IsInGameThread())
 	{
-		FGCCSyncObject::Get().GCLock();
+		// Don't need to take the GC Lock here since we've already validated we're not in the middle of GC
+		// with the check above. GIsSavingPackage provides the proper protection for avoiding object find during save (i.e. StaticFindObject)
+		GIsSavingPackage = true;
 	}
 	else
 	{
+		// Make sure the GC cannot run while we're saving
 		FGCCSyncObject::Get().LockAsync();
-	}
-
-	// Do not change GIsSavingPackage while saving concurrently. It should have been set before and after all packages are saved
-	if (!bSavingConcurrent)
-	{
-		GIsSavingPackage = true;
+		GIsSavingPackageTLS = true;
 	}
 
 	// Mark the package as being saved 
@@ -1255,16 +1260,13 @@ FScopedSavingFlag::FScopedSavingFlag(bool InSavingConcurrent, UPackage* InSavedP
 
 FScopedSavingFlag::~FScopedSavingFlag()
 {
-	if (!bSavingConcurrent)
-	{
-		GIsSavingPackage = false;
-	}
 	if (IsInGameThread())
 	{
-		FGCCSyncObject::Get().GCUnlock();
+		GIsSavingPackage = false;
 	}
 	else
 	{
+		GIsSavingPackageTLS = false;
 		FGCCSyncObject::Get().UnlockAsync();
 	}
 
@@ -1272,7 +1274,6 @@ FScopedSavingFlag::~FScopedSavingFlag()
 	{
 		SavedPackage->ClearPackageFlags(PKG_IsSaving);
 	}
-
 }
 
 FCanSkipEditorReferencedPackagesWhenCooking::FCanSkipEditorReferencedPackagesWhenCooking()
