CoreUObject: Simplify serialization of ClassFlags in UClass::Serialize.
CLASS_ShouldNeverBeLoaded is only used for one single scenario, i.e. loading and saving UClass::ClassFlags.
Remove the explicit exclusion of the CLASS_TokenStreamAssembled flag since it is already part of CLASS_ShouldNeverBeLoaded.
Move the explicitly excluded CLASS_ReplicationDataIsSetUp flag into the composite flags CLASS_ShouldNeverBeLoaded.
Remove the double clearing of excluded flags when loading (once upon a time there was SetTokenStreamMaybeDirty() call between the two places).

Motivation:
Use CLASS_ShouldNeverBeLoaded to implement consistent and streamlined saving and loading in the new serialization stack, and to reduce the risk of the systems drifting apart.

#tests BuildCookRun Lyra

#Synced-CL 50970586

diff --git a/Engine/Source/Runtime/CoreUObject/Private/UObject/Class.cpp b/Engine/Source/Runtime/CoreUObject/Private/UObject/Class.cpp
--- a/Engine/Source/Runtime/CoreUObject/Private/UObject/Class.cpp
+++ b/Engine/Source/Runtime/CoreUObject/Private/UObject/Class.cpp
@@ -6003,13 +6003,16 @@ void UClass::Serialize( FArchive& Ar )
 	if (Ar.IsSaving())
 	{
 		uint32 SavedClassFlags = ClassFlags;
-		SavedClassFlags &= ~(CLASS_ShouldNeverBeLoaded | CLASS_TokenStreamAssembled | CLASS_ReplicationDataIsSetUp);
+		SavedClassFlags &= ~CLASS_ShouldNeverBeLoaded;
 		Ar << SavedClassFlags;
 	}
 	else if (Ar.IsLoading())
 	{
 		Ar << (uint32&)ClassFlags;
-		ClassFlags &= ~(CLASS_ShouldNeverBeLoaded | CLASS_TokenStreamAssembled | CLASS_ReplicationDataIsSetUp);
+		ClassFlags &= ~CLASS_ShouldNeverBeLoaded;
+
+		checkf(!HasAnyClassFlags(CLASS_Native), TEXT("Class %s loaded with CLASS_Native....we should not be loading any native classes."), *GetFullName());
+		checkf(!HasAnyClassFlags(CLASS_Intrinsic), TEXT("Class %s loaded with CLASS_Intrinsic....we should not be loading any intrinsic classes."), *GetFullName());
 	}
 	else
 	{
@@ -6072,15 +6075,9 @@ void UClass::Serialize( FArchive& Ar )
 		Ar << ClassGeneratedBy;
 	}
 
-	if(Ar.IsLoading())
+	if (Ar.IsLoading() && !(Ar.GetPortFlags() & PPF_Duplicate))
 	{
-		checkf(!HasAnyClassFlags(CLASS_Native), TEXT("Class %s loaded with CLASS_Native....we should not be loading any native classes."), *GetFullName());
-		checkf(!HasAnyClassFlags(CLASS_Intrinsic), TEXT("Class %s loaded with CLASS_Intrinsic....we should not be loading any intrinsic classes."), *GetFullName());
-		ClassFlags &= ~(CLASS_ShouldNeverBeLoaded | CLASS_TokenStreamAssembled);
-		if (!(Ar.GetPortFlags() & PPF_Duplicate))
-		{
-			Link(Ar, true);
-		}
+		Link(Ar, true);
 	}
 
 	if(Ar.IsLoading())
diff --git a/Engine/Source/Runtime/CoreUObject/Public/UObject/ObjectMacros.h b/Engine/Source/Runtime/CoreUObject/Public/UObject/ObjectMacros.h
--- a/Engine/Source/Runtime/CoreUObject/Public/UObject/ObjectMacros.h
+++ b/Engine/Source/Runtime/CoreUObject/Public/UObject/ObjectMacros.h
@@ -294,8 +294,8 @@ ENUM_CLASS_FLAGS(EClassFlags);
 /** These flags will be cleared by the compiler when the class is parsed during script compilation */
 #define CLASS_RecompilerClear ((EClassFlags)(CLASS_Inherit | CLASS_Abstract | CLASS_Native | CLASS_Intrinsic | CLASS_TokenStreamAssembled))
 
-/** These flags will be cleared by the compiler when the class is parsed during script compilation */
-#define CLASS_ShouldNeverBeLoaded ((EClassFlags)(CLASS_Native | CLASS_Optional | CLASS_Intrinsic | CLASS_TokenStreamAssembled))
+/** These flags will never be loaded or saved */
+#define CLASS_ShouldNeverBeLoaded ((EClassFlags)(CLASS_Native | CLASS_Optional | CLASS_Intrinsic | CLASS_TokenStreamAssembled | CLASS_ReplicationDataIsSetUp))
 
 /** These flags will be inherited from the base class only for non-intrinsic classes */
 #define CLASS_ScriptInherit ((EClassFlags)(CLASS_Inherit | CLASS_EditInlineNew | CLASS_CollapseCategories))
