CoreUObject: Add ENUM_CLASS_FLAGS to EStructFlags and cleanup explicit casts in composite flag defines for CoreUObject enums.
Delete some composite EObjectFlags that were deprecated in 5.4.

#Synced-CL 50970586

diff --git a/Engine/Source/Runtime/CoreUObject/Public/UObject/Class.h b/Engine/Source/Runtime/CoreUObject/Public/UObject/Class.h
--- a/Engine/Source/Runtime/CoreUObject/Public/UObject/Class.h
+++ b/Engine/Source/Runtime/CoreUObject/Public/UObject/Class.h
@@ -1130,6 +1130,8 @@ enum EStructFlags
 	// !!!!!!!!!!!!!!!!!!!!!!!!!!!!!!!!!!!!!!!!!!!!!!!!!!!!!!!!!!!!!!!!!!!!!!!!!!!!!!!!!!!!!!!!!!!!!!!!!!!!!!!!!!!!!!!!!!!!!!!!!!!!!
 };
 
+ENUM_CLASS_FLAGS(EStructFlags)
+
 /** Struct flags that are automatically inherited */
 #define STRUCT_Inherit			(STRUCT_HasInstancedReference | STRUCT_Atomic)
 
@@ -2359,7 +2361,7 @@ public:
 	/** Resets currently assigned CppStructOps, called when loading a struct */
 	void ClearCppStructOps()
 	{
-		StructFlags = EStructFlags(StructFlags & ~STRUCT_ComputedFlags);
+		StructFlags = StructFlags & ~STRUCT_ComputedFlags;
 		bPrepareCppStructOpsCompleted = false;
 		CppStructOps = nullptr;
 	}
diff --git a/Engine/Source/Runtime/CoreUObject/Public/UObject/ObjectMacros.h b/Engine/Source/Runtime/CoreUObject/Public/UObject/ObjectMacros.h
--- a/Engine/Source/Runtime/CoreUObject/Public/UObject/ObjectMacros.h
+++ b/Engine/Source/Runtime/CoreUObject/Public/UObject/ObjectMacros.h
@@ -171,7 +171,7 @@ enum EPackageFlags
 #define PKG_TransientFlags	(PKG_NewlyCreated | PKG_IsSaving | PKG_ReloadingForCooker)
 
 /** Flag mask that indicates if this package is a package that exists in memory only. */
-#define PKG_InMemoryOnly	(EPackageFlags)(PKG_CompiledIn | PKG_NewlyCreated)
+#define PKG_InMemoryOnly	(PKG_CompiledIn | PKG_NewlyCreated)
 
 ENUM_CLASS_FLAGS(EPackageFlags);
 
@@ -288,20 +288,20 @@ enum EClassFlags
 ENUM_CLASS_FLAGS(EClassFlags);
 
 /** Flags to inherit from base class */
-#define CLASS_Inherit ((EClassFlags)(CLASS_Transient | CLASS_Optional | CLASS_DefaultConfig | CLASS_Config | CLASS_PerObjectConfig | CLASS_ConfigDoNotCheckDefaults | CLASS_NotPlaceable \
-						| CLASS_Const | CLASS_HasInstancedReference | CLASS_Deprecated | CLASS_DefaultToInstanced | CLASS_GlobalUserConfig | CLASS_ProjectUserConfig | CLASS_PerPlatformConfig | CLASS_NeedsDeferredDependencyLoading))
+#define CLASS_Inherit (CLASS_Transient | CLASS_Optional | CLASS_DefaultConfig | CLASS_Config | CLASS_PerObjectConfig | CLASS_ConfigDoNotCheckDefaults | CLASS_NotPlaceable \
+						| CLASS_Const | CLASS_HasInstancedReference | CLASS_Deprecated | CLASS_DefaultToInstanced | CLASS_GlobalUserConfig | CLASS_ProjectUserConfig | CLASS_PerPlatformConfig | CLASS_NeedsDeferredDependencyLoading)
 
 /** These flags will be cleared by the compiler when the class is parsed during script compilation */
-#define CLASS_RecompilerClear ((EClassFlags)(CLASS_Inherit | CLASS_Abstract | CLASS_Native | CLASS_Intrinsic | CLASS_TokenStreamAssembled))
+#define CLASS_RecompilerClear (CLASS_Inherit | CLASS_Abstract | CLASS_Native | CLASS_Intrinsic | CLASS_TokenStreamAssembled)
 
 /** These flags will never be loaded or saved */
-#define CLASS_ShouldNeverBeLoaded ((EClassFlags)(CLASS_Native | CLASS_Optional | CLASS_Intrinsic | CLASS_TokenStreamAssembled | CLASS_ReplicationDataIsSetUp))
+#define CLASS_ShouldNeverBeLoaded (CLASS_Native | CLASS_Optional | CLASS_Intrinsic | CLASS_TokenStreamAssembled | CLASS_ReplicationDataIsSetUp)
 
 /** These flags will be inherited from the base class only for non-intrinsic classes */
-#define CLASS_ScriptInherit ((EClassFlags)(CLASS_Inherit | CLASS_EditInlineNew | CLASS_CollapseCategories))
+#define CLASS_ScriptInherit (CLASS_Inherit | CLASS_EditInlineNew | CLASS_CollapseCategories)
 
 /** This is used as a mask for the flags put into generated code for "compiled in" classes. */
-#define CLASS_SaveInCompiledInClasses ((EClassFlags)(\
+#define CLASS_SaveInCompiledInClasses ( \
 	CLASS_Abstract | \
 	CLASS_DefaultConfig | \
 	CLASS_GlobalUserConfig | \
@@ -328,9 +328,9 @@ ENUM_CLASS_FLAGS(EClassFlags);
 	CLASS_MinimalAPI | \
 	CLASS_RequiredAPI | \
 	CLASS_MatchedSerializers | \
-	CLASS_NeedsDeferredDependencyLoading))
+	CLASS_NeedsDeferredDependencyLoading)
 
-#define CLASS_AllFlags ((EClassFlags)0xFFFFFFFFu)
+#define CLASS_AllFlags ((EClassFlags)0xFFFFFFFF)
 
 
 /**
@@ -635,23 +635,14 @@ enum EObjectFlags
 	// !!!!!!!!!!!!!!!!!!!!!!!!!!!!!!!!!!!!!!!!!!!!!!!!!!!!!!!!!!!!!!!!!!!!!!!!!!!!!!!!!!!!!!!!!!!!!!!!!!!!!!!!!!!!!!!!!!!!!!!!!!!
 };
 
-UE_DEPRECATED(5.4, "RF_InternalPendingKill should no longer be used. Use IsValid(Object) instead.")
-inline constexpr EObjectFlags RF_InternalPendingKill = RF_MirroredGarbage;
-UE_DEPRECATED(5.4, "RF_InternalGarbage should no longer be used. Use IsValid(Object) instead.")
-inline constexpr EObjectFlags RF_InternalGarbage = RF_MirroredGarbage;
-UE_DEPRECATED(5.4, "RF_InternalMirroredFlags should no longer be used. Use IsValid(Object) instead.")
-inline constexpr EObjectFlags RF_InternalMirroredFlags = RF_MirroredGarbage;
-UE_DEPRECATED(5.5, "RF_HasPlaceholderType is no longer in use.")
-inline constexpr EObjectFlags RF_HasPlaceholderType = EObjectFlags(0x20000000);
-
 /** Mask for all object flags */
-#define RF_AllFlags				(EObjectFlags)0xffffffff	///< All flags, used mainly for error checking
+#define RF_AllFlags				((EObjectFlags)0xFFFFFFFF)	///< All flags, used mainly for error checking
 
 /** Flags to load from unreal asset files */
-#define RF_Load						((EObjectFlags)(RF_Public | RF_Standalone | RF_Transactional | RF_ClassDefaultObject | RF_ArchetypeObject | RF_GeneratePublicHash | RF_DefaultSubObject | RF_TextExportTransient | RF_InheritableComponentTemplate | RF_DuplicateTransient | RF_NonPIEDuplicateTransient | RF_MigratingAsset | RF_HasDynamicImports)) 
+#define RF_Load						(RF_Public | RF_Standalone | RF_Transactional | RF_ClassDefaultObject | RF_ArchetypeObject | RF_GeneratePublicHash | RF_DefaultSubObject | RF_TextExportTransient | RF_InheritableComponentTemplate | RF_DuplicateTransient | RF_NonPIEDuplicateTransient | RF_MigratingAsset | RF_HasDynamicImports)
 
 /** Sub-objects will inherit these flags from their SuperObject */
-#define RF_PropagateToSubObjects	((EObjectFlags)(RF_Public | RF_ArchetypeObject | RF_Transactional | RF_Transient))
+#define RF_PropagateToSubObjects	(RF_Public | RF_ArchetypeObject | RF_Transactional | RF_Transient)
 
 ENUM_CLASS_FLAGS(EObjectFlags);
 
diff --git a/Engine/Source/Runtime/CoreUObject/Public/UObject/Script.h b/Engine/Source/Runtime/CoreUObject/Public/UObject/Script.h
--- a/Engine/Source/Runtime/CoreUObject/Public/UObject/Script.h
+++ b/Engine/Source/Runtime/CoreUObject/Public/UObject/Script.h
@@ -181,11 +181,11 @@ UE_FORCEINLINE_HINT FArchive& operator<<(FArchive& Ar, EFunctionFlags& Flags)
 ENUM_CLASS_FLAGS(EFunctionFlags)
 
 // Combinations of flags.
-#define FUNC_FuncInherit       ((EFunctionFlags)(FUNC_Exec | FUNC_Event | FUNC_BlueprintCallable | FUNC_BlueprintEvent | FUNC_BlueprintAuthorityOnly | FUNC_BlueprintCosmetic | FUNC_Const))
-#define FUNC_FuncOverrideMatch ((EFunctionFlags)(FUNC_Exec | FUNC_Final | FUNC_Static | FUNC_Public | FUNC_Protected | FUNC_Private))
-#define FUNC_NetFuncFlags      ((EFunctionFlags)(FUNC_Net | FUNC_NetReliable | FUNC_NetServer | FUNC_NetClient | FUNC_NetMulticast))
-#define FUNC_AccessSpecifiers  ((EFunctionFlags)(FUNC_Public | FUNC_Private | FUNC_Protected))
-#define FUNC_AllFlags          ((EFunctionFlags)(0xFFFFFFFF))
+#define FUNC_FuncInherit       (FUNC_Exec | FUNC_Event | FUNC_BlueprintCallable | FUNC_BlueprintEvent | FUNC_BlueprintAuthorityOnly | FUNC_BlueprintCosmetic | FUNC_Const)
+#define FUNC_FuncOverrideMatch (FUNC_Exec | FUNC_Final | FUNC_Static | FUNC_Public | FUNC_Protected | FUNC_Private)
+#define FUNC_NetFuncFlags      (FUNC_Net | FUNC_NetReliable | FUNC_NetServer | FUNC_NetClient | FUNC_NetMulticast)
+#define FUNC_AccessSpecifiers  (FUNC_Public | FUNC_Private | FUNC_Protected)
+#define FUNC_AllFlags          ((EFunctionFlags)0xFFFFFFFF)
 
 //
 // Evaluatable expression item types.
