Avoid resolving remote outers when sorting serialized remote object headers by outer

#Synced-CL 51010198

diff --git a/Engine/Source/Runtime/CoreUObject/Private/UObject/RemoteObjectSerialization.cpp b/Engine/Source/Runtime/CoreUObject/Private/UObject/RemoteObjectSerialization.cpp
--- a/Engine/Source/Runtime/CoreUObject/Private/UObject/RemoteObjectSerialization.cpp
+++ b/Engine/Source/Runtime/CoreUObject/Private/UObject/RemoteObjectSerialization.cpp
@@ -1151,8 +1151,8 @@ FORCENOINLINE void SortRemoteObjectHeadersByOuter(FArchiveRemoteObjectWriter& Ar
 
 		bool operator()(const FRemoteObjectHeader& A, const FRemoteObjectHeader& B) const
 		{
-			int32 OuterALevel = OuterLevelMap.FindChecked(A.Outer);
-			int32 OuterBLevel = OuterLevelMap.FindChecked(B.Outer);
+			int32 OuterALevel = OuterLevelMap.FindChecked(NoResolveObjectHandleNoRead(A.Outer.GetHandle()));
+			int32 OuterBLevel = OuterLevelMap.FindChecked(NoResolveObjectHandleNoRead(B.Outer.GetHandle()));
 			return OuterALevel > OuterBLevel;
 		}
 	} Predicate;
