Queue up RPCs called on borrowed objects so they can be processed when running to the owning server.

#Synced-CL 50977143

diff --git a/Engine/Source/Runtime/CoreUObject/Private/UObject/RemoteObjectSerialization.cpp b/Engine/Source/Runtime/CoreUObject/Private/UObject/RemoteObjectSerialization.cpp
--- a/Engine/Source/Runtime/CoreUObject/Private/UObject/RemoteObjectSerialization.cpp
+++ b/Engine/Source/Runtime/CoreUObject/Private/UObject/RemoteObjectSerialization.cpp
@@ -6,8 +6,10 @@
 #include "UObject/RemoteObjectPrivate.h"
 #include "UObject/ObjectPtr.h"
 #include "UObject/Class.h"
+#include "UObject/Field.h"
 #include "UObject/PropertyPortFlags.h"
 #include "UObject/ObjectHandle.h"
+#include "UObject/UObjectAnnotation.h"
 #include "UObject/UObjectHash.h"
 #include "UObject/UObjectMigrationContext.h"
 #include "UObject/UObjectThreadContext.h"
@@ -1568,4 +1570,180 @@ UObject* FindCanonicalRootObjectForSerialization(UObject* Object)
 	return Result;
 }
 
-} // namespace UE::RemoteObject::Serialization
\ No newline at end of file
+} // namespace UE::RemoteObject::Serialization
+
+namespace UE::RemoteObject::Serialization::Network
+{
+	// List of RPCs while borrowed per-object (Annotations are temporary data you can attach to UObjects)
+	COREUOBJECT_API FUObjectAnnotationSparse<FBorrowedRpcAnnotations, /*bAutoRemove=*/true> RemoteObjects_QueuedRPCs;
+
+	class FRemoteObjectReferenceWriter : public FMemoryWriter
+	{
+	public:
+		FRemoteObjectReferenceWriter(TArray<uint8>& InBytes)
+			: FMemoryWriter(InBytes)
+		{
+		}
+
+		virtual FArchive& operator<<(UObject*& Obj) override
+		{
+			FRemoteObjectReference RemoteReference { FObjectPtr(Obj) };
+			RemoteReference.Serialize(*this);
+			return *this;
+		}
+
+		virtual FArchive& operator<<(FObjectPtr& Value) override 
+		{
+			FRemoteObjectReference RemoteReference(Value);
+			RemoteReference.Serialize(*this);
+			return *this;
+		}
+
+		virtual FArchive& operator<<(FWeakObjectPtr& Value) override
+		{
+			FRemoteObjectReference RemoteReference(Value);
+			RemoteReference.Serialize(*this);
+			return *this;
+		}
+
+		virtual FArchive& operator<<(FLazyObjectPtr& Value) override { return FArchiveUObject::SerializeLazyObjectPtr(*this, Value); }
+		virtual FArchive& operator<<(FSoftObjectPtr& Value) override { return FArchiveUObject::SerializeSoftObjectPtr(*this, Value); }
+		virtual FArchive& operator<<(FSoftObjectPath& Value) override { return FArchiveUObject::SerializeSoftObjectPath(*this, Value); }
+	};
+
+	class FRemoteObjectReferenceReader : public FMemoryReader
+	{
+	public:
+		FRemoteObjectReferenceReader(TArray<uint8>& InBytes)
+			: FMemoryReader(InBytes)
+		{
+		}
+
+		virtual FArchive& operator<<(UObject*& Obj) override
+		{
+			FRemoteObjectReference RemoteReference;
+			RemoteReference.Serialize(*this);
+			Obj = RemoteReference.Resolve();
+			return *this;
+		}
+
+		virtual FArchive& operator<<(FObjectPtr& Value) override 
+		{
+			FRemoteObjectReference RemoteReference;
+			RemoteReference.Serialize(*this);
+			Value = RemoteReference.ToObjectPtr();
+			return *this;
+		}
+
+		virtual FArchive& operator<<(FWeakObjectPtr& Value) override
+		{
+			FRemoteObjectReference RemoteReference;
+			RemoteReference.Serialize(*this);
+			Value = RemoteReference.ToWeakPtr();
+			return *this;
+		}
+
+		virtual FArchive& operator<<(FLazyObjectPtr& Value) override { return FArchiveUObject::SerializeLazyObjectPtr(*this, Value); }
+		virtual FArchive& operator<<(FSoftObjectPtr& Value) override { return FArchiveUObject::SerializeSoftObjectPtr(*this, Value); }
+		virtual FArchive& operator<<(FSoftObjectPath& Value) override { return FArchiveUObject::SerializeSoftObjectPath(*this, Value); }
+	};
+
+	void EnqueueRPC(UObject* RootObject, UObject* SubObject, UFunction* Function, void* Parameters)
+	{
+		// Serialize Parameters, taking into account objects need to be stored as remote object references
+		TArray<uint8> SerializedParams;
+		{
+			FRemoteObjectReferenceWriter RemoteObjectReferenceWriter(SerializedParams);
+			FStructuredArchiveFromArchive StructuredRemoteObjectReferenceArchive(RemoteObjectReferenceWriter);
+
+			for (TFieldIterator<FProperty> PropertyIt(Function); PropertyIt && PropertyIt->HasAnyPropertyFlags(CPF_Parm); ++PropertyIt)
+			{
+				uint8* Param = PropertyIt->ContainerPtrToValuePtr<uint8>(Parameters);
+				PropertyIt->SerializeItem(StructuredRemoteObjectReferenceArchive.GetSlot(), Param);
+			}
+		}
+
+		// Append the data to the existing annotation (or create a new one)
+		FBorrowedRpcAnnotations AnnotationInstance = RemoteObjects_QueuedRPCs.GetAndRemoveAnnotation(RootObject);
+		FBorrowedRpcAnnotations::FBorrowedSerializedRPC& BorrowedRPC = AnnotationInstance.SerializedRPCs.Emplace_GetRef();
+
+		// Create an entry in the annotations for this RPC
+		BorrowedRPC.SubObject = SubObject;
+		BorrowedRPC.FunctionName = Function->GetFName();
+		BorrowedRPC.Params = MoveTemp(SerializedParams);
+
+		// Store the modified annotation back
+		RemoteObjects_QueuedRPCs.AddAnnotation(RootObject, MoveTemp(AnnotationInstance));
+
+		UE_LOGF(LogRemoteSerialization, Verbose, "Enqueing RPC %ls on root object %ls and sub-object %ls", *Function->GetName(), *GetNameSafe(RootObject), *GetNameSafe(SubObject));
+	}
+
+	void SerializeRPCQueue(UObject* RootObject, FArchive& Ar)
+	{
+		// Remove if saving (or create if none exists / loading).
+		FBorrowedRpcAnnotations AnnotationInstance = RemoteObjects_QueuedRPCs.GetAndRemoveAnnotation(RootObject);
+		AnnotationInstance.Serialize(Ar);
+
+		// If we're loading, add our annotation back for PostMigrate to handle it correctly
+		if (Ar.IsLoading() && !AnnotationInstance.IsDefault())
+		{
+			RemoteObjects_QueuedRPCs.AddAnnotation(RootObject, MoveTemp(AnnotationInstance));
+		}
+	}
+
+	void ProcessRPC(UObject* TargetObject, FBorrowedRpcAnnotations::FBorrowedSerializedRPC& RPC)
+	{
+		UFunction* Function = TargetObject->FindFunction(RPC.FunctionName);
+		if (!Function)
+		{
+			UE_LOGF(LogRemoteSerialization, Warning, "Cannot find function %ls on object %ls", *RPC.FunctionName.ToString(), *TargetObject->GetName());
+			return;
+		}
+
+		uint8* Params = static_cast<uint8*>(FMemory_Alloca(Function->ParmsSize));
+		FMemory::Memzero(Params, Function->ParmsSize);
+
+		for (TFieldIterator<FProperty> It(Function); It && It->HasAnyPropertyFlags(CPF_Parm); ++It)
+		{
+			It->InitializeValue_InContainer(Params);
+		}
+
+		// Deserailize RPC arguments into memory.
+		FRemoteObjectReferenceReader RemoteObjectReferenceReader(RPC.Params);
+		FStructuredArchiveFromArchive StructuredRemoteObjectReferenceArchive(RemoteObjectReferenceReader);
+
+		for (TFieldIterator<FProperty> It(Function); It && It->HasAnyPropertyFlags(CPF_Parm); ++It)
+		{
+			uint8* PropertyMemory = It->ContainerPtrToValuePtr<uint8>(Params);
+			It->SerializeItem(StructuredRemoteObjectReferenceArchive.GetSlot(), PropertyMemory);
+		}
+
+		// Call the RPC function.
+		UE_LOGF(LogRemoteSerialization, Verbose, "Processing queued RPC %ls on %ls", *Function->GetName(), *TargetObject->GetName());
+		TargetObject->ProcessEvent(Function, Params);
+
+		for (TFieldIterator<FProperty> It(Function); It && It->HasAnyPropertyFlags(CPF_Parm); ++It)
+		{
+			It->DestroyValue_InContainer(Params);
+		}
+	}
+
+	void ProcessRPCQueue(UObject* RootObject)
+	{
+		FBorrowedRpcAnnotations AnnotationInstance = RemoteObjects_QueuedRPCs.GetAndRemoveAnnotation(RootObject);
+
+		for (FBorrowedRpcAnnotations::FBorrowedSerializedRPC& RPC : AnnotationInstance.SerializedRPCs)
+		{
+			UObject* SubObject = RPC.SubObject.Get();
+			UObject* TargetObject = SubObject ? SubObject : RootObject;
+
+			if (!TargetObject)
+			{
+				UE_LOGF(LogRemoteSerialization, Error, "Cannot find the target object for function %ls on root object %ls", *RPC.FunctionName.ToString(), *GetNameSafe(RootObject));
+				continue;
+			}
+
+			ProcessRPC(TargetObject, RPC);
+		}
+	}
+}
\ No newline at end of file
diff --git a/Engine/Source/Runtime/CoreUObject/Private/UObject/RemoteObjectSerialization.h b/Engine/Source/Runtime/CoreUObject/Private/UObject/RemoteObjectSerialization.h
--- a/Engine/Source/Runtime/CoreUObject/Private/UObject/RemoteObjectSerialization.h
+++ b/Engine/Source/Runtime/CoreUObject/Private/UObject/RemoteObjectSerialization.h
@@ -163,4 +163,55 @@ namespace UE::RemoteObject::Serialization::Disk
 {
 	void LoadObjectFromDisk(const FUObjectMigrationContext& MigrationContext);
 	void SaveObjectToDisk(const UE::RemoteObject::Transfer::FMigrateSendParams& Params);
-}
\ No newline at end of file
+}
+
+namespace UE::RemoteObject::Serialization::Network
+{
+	/*
+	* We want to be able to store "annotations" (metadata) about an Object while it's borrowed in order to
+	* execute it on the owning server when the loan is returned.
+	*/
+	struct FBorrowedRpcAnnotations
+	{
+		// The data for an individual RPC call
+		struct FBorrowedSerializedRPC
+		{
+			/** The object to execute the RPC on */
+			TWeakObjectPtr<UObject> SubObject;
+
+			/** The function name of the RPC */
+			FName FunctionName;
+
+			/** The serialized parameters of the RPC */
+			TArray<uint8> Params;
+
+			friend FArchive& operator<<(FArchive& Ar, FBorrowedSerializedRPC& Value)
+			{
+				Ar << Value.SubObject;
+				Ar << Value.FunctionName;
+				Ar << Value.Params;
+
+				return Ar;
+			}
+		};
+
+		// Serialized RPC Data that were processed (not executed) while borrowed
+		TArray<FBorrowedSerializedRPC> SerializedRPCs;
+
+		// Annotation is default (empty) and removable if it has no RPCs
+		bool IsDefault() const
+		{
+			return SerializedRPCs.IsEmpty();
+		}
+
+		void Serialize(FArchive& Ar)
+		{
+			TArrayPrivateFriend::Serialize(Ar, SerializedRPCs);
+		}
+	};
+}
+
+template<> struct TStructOpsTypeTraits<UE::RemoteObject::Serialization::Network::FBorrowedRpcAnnotations::FBorrowedSerializedRPC> : public TStructOpsTypeTraitsBase2<UE::RemoteObject::Serialization::Network::FBorrowedRpcAnnotations::FBorrowedSerializedRPC>
+{
+	enum { WithSerializer = true };
+};
diff --git a/Engine/Source/Runtime/CoreUObject/Public/UObject/RemoteObjectTransfer.h b/Engine/Source/Runtime/CoreUObject/Public/UObject/RemoteObjectTransfer.h
--- a/Engine/Source/Runtime/CoreUObject/Public/UObject/RemoteObjectTransfer.h
+++ b/Engine/Source/Runtime/CoreUObject/Public/UObject/RemoteObjectTransfer.h
@@ -169,6 +169,31 @@ inline UClass* FSerializedRemoteObjectIterator::GetClass() const
 	return Cast<UClass>(Data.PathNames[Data.SerializedObjectHeaders[Index].ClassIndex].Resolve(Data.Tables));
 }
 
+namespace UE::RemoteObject::Serialization::Network
+{
+	/** 
+	 * Add a RPC, and its arguments, for an object to a queue to be processed later.
+	 *
+	 * There is a queue per root object that contains all RPCs for that root object or any of its sub-objects.
+	 *
+	 * If SubObject is null then its assumed the RPC is for the actor.
+	 */
+	COREUOBJECT_API void EnqueueRPC(UObject* RootObject, UObject* SubObject, UFunction* Function, void* Parameters);
+
+	/**
+	 * Serialize and deserialize and RPC queue for a root object.
+	 *
+	 * Each RPC that is serialized will be removed from the root object's queue, while each RPC that is deserialized
+	 * will be added to the root object's queue.
+	 */
+	COREUOBJECT_API void SerializeRPCQueue(UObject* RootObject, FArchive& Ar);
+
+	/** 
+	 * Process all RPCs queued up for a root object and its components and sub-objects.
+	 */
+	COREUOBJECT_API void ProcessRPCQueue(UObject* RootObject);
+}
+
 namespace UE::RemoteObject::Transfer
 {
 	/** Information for performing a migration (send) an object to a remote server */
diff --git a/Engine/Source/Runtime/Engine/Private/Actor.cpp b/Engine/Source/Runtime/Engine/Private/Actor.cpp
--- a/Engine/Source/Runtime/Engine/Private/Actor.cpp
+++ b/Engine/Source/Runtime/Engine/Private/Actor.cpp
@@ -1089,6 +1089,8 @@ void AActor::Serialize(FArchive& Ar)
 		Ar << MigrationFlags;
 		Ar << OwnedComponents;
 
+		UE::RemoteObject::Serialization::Network::SerializeRPCQueue(this, Ar);
+
 		if (Ar.IsLoading())
 		{
 			// When migrating actors are in a state after BeginPlay has been called but we need to call SetLifeSpan to initialize local timer
@@ -1402,6 +1404,7 @@ void AActor::PostMigrate(const FUObjectMigrationContext& MigrationContext)
 		{
 			ensureMsgf(!bActorIsBeingDestroyed, TEXT("We do not want to be assigned ownership of Actors that have already been destroyed. We do not yet know what to do with it."));
 
+			// Any queued up RPCs will be processed when the deferred start replication is processed.
 			UE::Net::FDeferredReplicationSystemCalls::StartReplicatingActorDefer(this);
 
 			// Each component will be in charge of registering itself in its own PostMigrate, but we must register
@@ -1414,6 +1417,9 @@ void AActor::PostMigrate(const FUObjectMigrationContext& MigrationContext)
 		}
 		else if (MigrationRecvType == EObjectMigrationRecvType::ReturnedLoan)
 		{
+			// Process any RPCs for this actor and its sub-objects (assumes the actor and its sub-objects have already been registered with Iris).
+			UE::RemoteObject::Serialization::Network::ProcessRPCQueue(this);
+
 			FActorLoanedState LoanedState = RemoteObjects_ActorLoanedStates.GetAndRemoveAnnotation(this);
 			if (bActorIsBeingDestroyed)
 			{
diff --git a/Engine/Source/Runtime/Engine/Private/ActorReplication.cpp b/Engine/Source/Runtime/Engine/Private/ActorReplication.cpp
--- a/Engine/Source/Runtime/Engine/Private/ActorReplication.cpp
+++ b/Engine/Source/Runtime/Engine/Private/ActorReplication.cpp
@@ -1088,6 +1088,11 @@ void AActor::FillReplicationParams(const FFillReplicationParamsContext& Context,
 void AActor::OnReplicationStartedForIris(const FOnReplicationStartedParams& Params)
 {
 	UpdateOwningNetConnection();
+
+#if UE_WITH_REMOTE_OBJECT_HANDLE
+	// Process any RPCs for this actor and its sub-objects (assumes the actor and its sub-objects have already been registered with Iris).
+	UE::RemoteObject::Serialization::Network::ProcessRPCQueue(this);
+#endif
 }
 
 void AActor::OnStopReplicationForIris(const FOnStopReplicationParams&)
diff --git a/Engine/Source/Runtime/Engine/Private/NetDriver.cpp b/Engine/Source/Runtime/Engine/Private/NetDriver.cpp
--- a/Engine/Source/Runtime/Engine/Private/NetDriver.cpp
+++ b/Engine/Source/Runtime/Engine/Private/NetDriver.cpp
@@ -8210,6 +8210,14 @@ void UNetDriver::ProcessRemoteFunction(
 			LocalOutParms = UE::Net::Private::CopyOutParametersToLocalParameters(Function, OutParms, Parameters, TargetObj);
 		}
 
+#if UE_WITH_REMOTE_OBJECT_HANDLE
+		if (!UE::RemoteObject::Handle::IsOwned(Actor) || !UE::RemoteObject::Handle::IsOwned(SubObject))
+		{
+			UE::RemoteObject::Serialization::Network::EnqueueRPC(Actor, SubObject, Function, Parameters);
+			return;
+		}
+#endif
+
 		if (ReplicationSystem)
 		{
 			if (bIsServerMulticast)
diff --git a/Engine/Source/Runtime/Engine/Private/PlayerController.cpp b/Engine/Source/Runtime/Engine/Private/PlayerController.cpp
--- a/Engine/Source/Runtime/Engine/Private/PlayerController.cpp
+++ b/Engine/Source/Runtime/Engine/Private/PlayerController.cpp
@@ -5104,9 +5104,13 @@ void APlayerController::PostMigrate(const FUObjectMigrationContext& MigrationCon
 
 		const bool bAttachControllerToConnection = (MigrationRecvType == EObjectMigrationRecvType::AssignedOwnership);
 
-		// Make sure to update NetConnection of a player controller that was loaned and is now returning, because the pointer will 
-		// be set to null when serialized.
-		const bool bUpdateNetConnection = (bAttachControllerToConnection || MigrationRecvType == EObjectMigrationRecvType::ReturnedLoan);
+		// Update the NetConnection on the player controller because it will be set to null when serialized. It's important that the NetConnection variable
+		// is always updated, even when borrowed, because there is code that depends on that value and won't work if it's null.
+		const bool bUpdateNetConnection = (
+			bAttachControllerToConnection || 
+			MigrationRecvType == EObjectMigrationRecvType::ReturnedLoan || 
+			MigrationRecvType == EObjectMigrationRecvType::Borrowed
+		);
 
 		if (bUpdateNetConnection)
 		{
