Reinstate s.AsyncCreatePackagesFromQueueMaxTimeLimitMs, but disable by default in the editor (due to PIE autotest crash)
- Make it possible for time limit scopes to be conditional so that disabling the time limit is non-intrusive
#Synced-CL 50901525

diff --git a/Engine/Source/Runtime/CoreUObject/Private/Serialization/AsyncLoading2.cpp b/Engine/Source/Runtime/CoreUObject/Private/Serialization/AsyncLoading2.cpp
--- a/Engine/Source/Runtime/CoreUObject/Private/Serialization/AsyncLoading2.cpp
+++ b/Engine/Source/Runtime/CoreUObject/Private/Serialization/AsyncLoading2.cpp
@@ -541,6 +541,20 @@ static FAutoConsoleVariableRef CVarGAllowMultithreadedLoading(
 	ECVF_Default
 );
 
+#if WITH_EDITOR
+static float GAsyncCreatePackagesFromQueueMaxTimeLimitMs = 0.0f;
+#else
+static float GAsyncCreatePackagesFromQueueMaxTimeLimitMs = 2.5f;
+#endif
+static FAutoConsoleVariableRef CVarAsyncCreatePackagesFromQueueMaxTimeLimitMs(
+	TEXT("s.AsyncCreatePackagesFromQueueMaxTimeLimitMs"),
+	GAsyncCreatePackagesFromQueueMaxTimeLimitMs,
+	TEXT("Limit the time spent in CreateAsyncPackagesFromQueue when running on the ALT.\n")
+	TEXT("This can improve throughput when s.AllowMultithreadedLoading is enabled by allowing MT work to kick off earlier."),
+	ECVF_Default
+);
+
+
 static bool GOnlyProcessRequiredPackagesWhenSyncLoading = true;
 static FAutoConsoleVariableRef CVarGOnlyProcessRequiredPackagesWhenSyncLoading(
 	TEXT("s.OnlyProcessRequiredPackagesWhenSyncLoading"),
@@ -3636,22 +3650,25 @@ struct FAsyncLoadingThreadState2
 		double StartTime = 0.0;
 		FAsyncLoadingThreadState2& ThreadState;
 
-		FTimeLimitScope(FAsyncLoadingThreadState2& InThreadState, bool bInUseTimeLimit, double InTimeLimit)
+		FTimeLimitScope(FAsyncLoadingThreadState2& InThreadState, bool bInUseTimeLimit, double InTimeLimit, bool bCondition=true)
 			: ThreadState(InThreadState)
 		{
 			bUseTimeLimit = ThreadState.bUseTimeLimit;
 			TimeLimit = ThreadState.TimeLimit;
 			StartTime = ThreadState.StartTime;
 
-			ThreadState.bUseTimeLimit = bInUseTimeLimit;
-			ThreadState.TimeLimit = InTimeLimit;
-			if (bInUseTimeLimit)
-			{
-				ThreadState.StartTime = FPlatformTime::Seconds();
-			}
-			else
+			if (bCondition)
 			{
-				ThreadState.StartTime = 0.0;
+				ThreadState.bUseTimeLimit = bInUseTimeLimit;
+				ThreadState.TimeLimit = InTimeLimit;
+				if (bInUseTimeLimit)
+				{
+					ThreadState.StartTime = FPlatformTime::Seconds();
+				}
+				else
+				{
+					ThreadState.StartTime = 0.0;
+				}
 			}
 		}
 
@@ -11327,6 +11344,12 @@ uint32 FAsyncLoadingThread2::Run()
 
 					if (QueuedPackagesCounter || !PendingPackages.IsEmpty())
 					{
+						TRACE_CPUPROFILER_EVENT_SCOPE_CONDITIONAL(TimeLimitedCreateAsyncPackages, GAsyncCreatePackagesFromQueueMaxTimeLimitMs > .0f);
+						// Use a time limit scope to restore the time limit to the old values when exiting the scope
+						// This is required to ensure a reentrant call here doesn't overwrite time limits permanently.
+						// Note: this is only active when GAsyncCreatePackagesFromQueueMaxTimeLimitMs > 0
+						FAsyncLoadingThreadState2::FTimeLimitScope TimeLimitScope( ThreadState, true, double(GAsyncCreatePackagesFromQueueMaxTimeLimitMs/1000.0f), GAsyncCreatePackagesFromQueueMaxTimeLimitMs > .0f);
+
 						if (CreateAsyncPackagesFromQueue(ThreadState, GAsyncLoading2_AllowPreemptingPackagesDuringGC))
 						{
 							// Fall through to FAsyncLoadEventQueue2 processing unless we need to suspend
