Disable s.AsyncCreatePackagesFromQueueMaxTimeLimitMs in PGO jobs by default (speculative fix for a crash)
#Synced-CL 51003583

diff --git a/Engine/Source/Runtime/CoreUObject/Private/Serialization/AsyncLoading2.cpp b/Engine/Source/Runtime/CoreUObject/Private/Serialization/AsyncLoading2.cpp
--- a/Engine/Source/Runtime/CoreUObject/Private/Serialization/AsyncLoading2.cpp
+++ b/Engine/Source/Runtime/CoreUObject/Private/Serialization/AsyncLoading2.cpp
@@ -541,7 +541,7 @@ static FAutoConsoleVariableRef CVarGAllowMultithreadedLoading(
 	ECVF_Default
 );
 
-#if WITH_EDITOR
+#if WITH_EDITOR || ENABLE_PGO_PROFILE
 static float GAsyncCreatePackagesFromQueueMaxTimeLimitMs = 0.0f;
 #else
 static float GAsyncCreatePackagesFromQueueMaxTimeLimitMs = 2.5f;
