Cooker: CookArtifact: Split IsReferenced into an enum with None,Cook,Stage; the development AssetRegistry artifact needs to know about Cook references but other artifacts continue to need to know about Stage references.
#rnx
#Synced-CL 51045425

diff --git a/Engine/Source/Editor/UnrealEd/Private/CookOnTheFlyServer.cpp b/Engine/Source/Editor/UnrealEd/Private/CookOnTheFlyServer.cpp
--- a/Engine/Source/Editor/UnrealEd/Private/CookOnTheFlyServer.cpp
+++ b/Engine/Source/Editor/UnrealEd/Private/CookOnTheFlyServer.cpp
@@ -9465,16 +9465,16 @@ void UCookOnTheFlyServer::CookByTheBookFinishedInternal()
 			for (FPackageData* PackageData : CookedPackageDatas)
 			{
 				FOplogPackageData& Data = OplogContext.OplogPackages.FindOrAdd(PackageData->GetPackageName());
-				Data.bReferenced = true;
+				Data.ReferenceType = UE::Cook::EPackageReferenceType::Stage;
 				const FPackagePlatformData& PackagePlatformData = PackageData->FindOrAddPlatformData(TargetPlatform);
 				Data.bIncrementallyCooked = !PackagePlatformData.IsIncrementallySkipped();
 			}
 			for (FName PackageName : Generator.GetPreviousIncrementalCookPackages())
 			{
-				// We rely on Data.bReferenced being initialized to false for newly added elements. In that case,
-				// calling FindOrAdd is sufficient to add the entry we need because we only need to add entries for
-				// any non-referenced-set packages; all of the referenced-set packages were added in the
-				// for (PackageData : CookedPackageDatas) loop above.
+				// We rely on ReferenceType being initialized to EPackageReferenceType::Unreferenced for newly added
+				// elements. In that case, calling FindOrAdd is sufficient to add the entry we need because we only
+				// need to add entries for any non-referenced-set packages; all of the referenced-set packages were
+				// added in the for (PackageData : CookedPackageDatas) loop above.
 				OplogContext.OplogPackages.FindOrAdd(PackageName);
 			}
 
diff --git a/Engine/Source/Runtime/CoreUObject/Public/Cooker/CookArtifact.h b/Engine/Source/Runtime/CoreUObject/Public/Cooker/CookArtifact.h
--- a/Engine/Source/Runtime/CoreUObject/Public/Cooker/CookArtifact.h
+++ b/Engine/Source/Runtime/CoreUObject/Public/Cooker/CookArtifact.h
@@ -15,6 +15,31 @@ namespace UE::Cook { class ICookArtifactReader; }
 namespace UE::Cook { class ICookInfo; }
 namespace UE::Cook::Artifact { struct FUpdateOplogPackagesContext; }
 
+namespace UE::Cook
+{
+
+/**
+ * To what degree is a package referenced in a cook; this information is provided to CookArtifacts about all of the
+ * packages reported from the current cook and from previous incremental cooks.
+ */
+enum class EPackageReferenceType : uint8
+{
+	/**
+	 * Unreferenced in the cook. This value occurs for packages referenced in a previous incremental cook that were
+	 * not referenced in the current cook due to e.g. commandline options such as -cooksinglepackagenorefs.
+	 */
+	Unreferenced,
+	/**
+	 * The cook needed to read and store information about this package, but it is not staged to runtime. This value
+	 * occurs for build dependencies and editoronly packages.
+	 */
+	Cook,
+	/** Cook referenced and should be staged to runtime. */
+	Stage,
+};
+
+}
+
 namespace UE::Cook::Artifact
 {
 
@@ -69,24 +94,21 @@ private:
 struct FOplogPackageData
 {
 public:
+	/** Degree to which the package is referenced in the cook, @see EPackageReferenceType. */
+	UE::Cook::EPackageReferenceType GetReferenceType() const;
 	/**
-	 * True if the package was referenced in the current cook (and for DLC cooks is not a basegame package cooked
-	 * elsewhere), and will therefore be staged during staging.
-	 */
-	bool IsReferenced() const;
-	/**
-	 * True if the package was referenced, and was either not present in an earlier incremental cook (NewCooked) or
-	 * was found to be modified and was Recooked.
+	 * True if the package was referenced in the cook (GetReferenceType != EPackageReferenceType::Unreferenced), and
+	 * was either not present in an earlier incremental cook (NewCooked) or was found to be modified and was Recooked.
 	 */
 	bool WasIncrementallyCooked() const;
 	/**
-	 * True if the package was referenced, and was found to be cooked in a previous cook and not modified so it was
-	 * not Recooked.
+	 * True if the package was referenced in the cook (GetReferenceType != EPackageReferenceType::Unreferenced), and
+	 * was found to be cooked in a previous cook and not modified so it was not Recooked.
 	 */
 	bool WasIncrementallySkipped() const;
 
 private:
-	bool bReferenced = false;
+	UE::Cook::EPackageReferenceType ReferenceType = EPackageReferenceType::Unreferenced;
 	bool bIncrementallyCooked = false;
 
 	friend FUpdateOplogPackagesContext;
@@ -270,19 +292,19 @@ inline const TMap<FName, FOplogPackageData>& FUpdateOplogPackagesContext::GetOpl
 	return OplogPackages;
 }
 
-inline bool FOplogPackageData::IsReferenced() const
+inline UE::Cook::EPackageReferenceType FOplogPackageData::GetReferenceType() const
 {
-	return bReferenced;
+	return ReferenceType;
 }
 
 inline bool FOplogPackageData::WasIncrementallyCooked() const
 {
-	return bReferenced && bIncrementallyCooked;
+	return (ReferenceType != UE::Cook::EPackageReferenceType::Unreferenced) && bIncrementallyCooked;
 }
 
 inline bool FOplogPackageData::WasIncrementallySkipped() const
 {
-	return bReferenced && !bIncrementallyCooked;
+	return (ReferenceType != UE::Cook::EPackageReferenceType::Unreferenced) && !bIncrementallyCooked;
 }
 
 } // namespace UE::Cook::Artifact
diff --git a/Engine/Source/Runtime/RenderCore/Private/ShaderLibrary/EditorShaderCodeArchive.cpp b/Engine/Source/Runtime/RenderCore/Private/ShaderLibrary/EditorShaderCodeArchive.cpp
--- a/Engine/Source/Runtime/RenderCore/Private/ShaderLibrary/EditorShaderCodeArchive.cpp
+++ b/Engine/Source/Runtime/RenderCore/Private/ShaderLibrary/EditorShaderCodeArchive.cpp
@@ -865,7 +865,8 @@ void FEditorShaderCodeArchive::UpdateOplogPackages(UE::Cook::Artifact::FUpdateOp
 			if (OplogData)
 			{
 				Pair.Value.bReferencedByOplog = true;
-				Pair.Value.bReferencedByStaging |= OplogData->IsReferenced();
+				Pair.Value.bReferencedByStaging |=
+					OplogData->GetReferenceType() == UE::Cook::EPackageReferenceType::Stage;
 			}
 		}
 	}
