[PluginManager] Optimize plugin leak detection by traversing the packages once per unmounted batch

#tests local cooked client
#rnx

#Synced-CL 51003583

diff --git a/Engine/Source/Runtime/CoreUObject/Private/Misc/CoreUObjectPluginManager.cpp b/Engine/Source/Runtime/CoreUObject/Private/Misc/CoreUObjectPluginManager.cpp
--- a/Engine/Source/Runtime/CoreUObject/Private/Misc/CoreUObjectPluginManager.cpp
+++ b/Engine/Source/Runtime/CoreUObject/Private/Misc/CoreUObjectPluginManager.cpp
@@ -1,7 +1,12 @@
 // Copyright Epic Games, Inc. All Rights Reserved.
 
 #include "CoreUObjectPluginManager.h"
+#include "Containers/ContainerAllocationPolicies.h"
+#include "Containers/ContainersFwd.h"
+#include "Containers/UnrealString.h"
 #include "Interfaces/IPluginManager.h"
+#include "Misc/PathViews.h"
+#include "UObject/NameTypes.h"
 #include "UObject/UObjectAllocator.h"
 #include "UObject/Package.h"
 #include "Templates/Casts.h"
@@ -9,9 +14,6 @@
 #include "Logging/LogMacros.h"
 #include "UObject/ObjectRename.h"
 #include "HAL/IConsoleManager.h"
-#include "Misc/AsciiSet.h"
-#include "String/ParseTokens.h"
-#include "Misc/WildcardString.h"
 #include "UObject/ReferenceChainSearch.h"
 
 namespace UE::CoreUObject::Private
@@ -143,7 +145,8 @@ namespace UE::CoreUObject::Private
 
 	// Check if any assets from the plugin mount point have leaked, and if so trace them.
 	// Then rename the assets and mark them as garbage to allow new copies of them to be loaded. 
-	void HandlePossibleAssetLeaks(const FString& PluginName)
+	template<typename TSetAllocatorType>
+	void HandlePossibleAssetLeaks(const TSet<FString, DefaultKeyFuncs<FString>, TSetAllocatorType>& PluginNames)
 	{
 		TRACE_CPUPROFILER_EVENT_SCOPE(HandlePossibleAssetLeaks);
 
@@ -164,10 +167,6 @@ namespace UE::CoreUObject::Private
 		}
 
 		TArray<UPackage*> LeakedPackages;
-		TStringBuilder<512> Prefix;
-		Prefix << '/' << PluginName << '/';
-
-		TStringBuilder<NAME_SIZE> NameBuffer;
 		{
 			// If the UObject hash knew about package mount roots, we could avoid this loop
 			TRACE_CPUPROFILER_EVENT_SCOPE(PackageLoop);
@@ -180,9 +179,9 @@ namespace UE::CoreUObject::Private
 						{
 							return;
 						}
-						NameBuffer.Reset();
-						Package->GetFName().GetDisplayNameEntry()->AppendNameToString(NameBuffer);
-						if (NameBuffer.ToView().StartsWith(Prefix, ESearchCase::IgnoreCase))
+						FNameBuilder NameBuffer(Package->GetFName());
+						FStringView PluginNameOfPackage = FPathViews::GetMountPointNameFromPath(NameBuffer.ToView());
+						if (PluginNames.ContainsByHash(GetTypeHash(PluginNameOfPackage), PluginNameOfPackage))
 						{
 							LeakedPackages.Add(Package);
 						}
@@ -231,10 +230,9 @@ namespace UE::CoreUObject::Private
 			TSet<UObject*> ObjectsThatHadStandaloneCleared;
 			TArray<UObject*> ObjectsInPackage;
 
-			// Sort even if we don't limit the count, so that high priority leaks appear first 
-			int32 OmittedCount = FMath::Max(0, LeakedPackages.Num() - GLeakedAssetTrace_MaxReportCount);
 			TArray<UPackage*> PackagesToSearchFor;
-			int32 i = 0;
+			TMap<FString, int32> PackageCountPerMountPoint;
+			PackageCountPerMountPoint.Reserve(PluginNames.Num());
 			for (UPackage* Package : LeakedPackages)
 			{
 				ObjectsInPackage.Reset();
@@ -252,13 +250,17 @@ namespace UE::CoreUObject::Private
 					}
 				}
 
-				if (i++ < GLeakedAssetTrace_MaxReportCount)
+				const FNameBuilder PackageName(Package->GetFName());
+				const FString MountPointName(FPathViews::GetMountPointNameFromPath(PackageName.ToView()));
+
+				if (PackageCountPerMountPoint.FindOrAdd(MountPointName)++ < GLeakedAssetTrace_MaxReportCount)
 				{
 					PackagesToSearchFor.Add(Package);
 				}
 			}
 
-			UE_LOG(PluginHandlerLog, Display, TEXT("Searching for references to %d leaked packages (%d omitted for speed) from plugin %s"), LeakedPackages.Num(), OmittedCount, *PluginName);
+			const int32 OmittedCount = FMath::Max(0, LeakedPackages.Num() - PackagesToSearchFor.Num());
+			UE_LOG(PluginHandlerLog, Display, TEXT("Searching for references to %d leaked packages (%d omitted for speed)"), LeakedPackages.Num(), OmittedCount);
 			FReferenceChainSearch::FindAndPrintStaleReferencesToObjects(MakeArrayView((UObject**)PackagesToSearchFor.GetData(), PackagesToSearchFor.Num()), Options);
 
 			for (UObject* Object : ObjectsThatHadStandaloneCleared)
@@ -297,7 +299,7 @@ namespace UE::CoreUObject::Private
 
 		if (SuppressGCRefCount > 0)
 		{
-			DeferredPluginsToGC.RemoveSwap(Plugin.GetName());
+			DeferredPluginsToGC.Remove(Plugin.GetName());
 		}
 	}
 
@@ -313,7 +315,8 @@ namespace UE::CoreUObject::Private
 		else
 		{
 			CollectGarbage(GARBAGE_COLLECTION_KEEPFLAGS, true);
-			HandlePossibleAssetLeaks(Plugin.GetName());
+			using FOneElementInlineSet = TSet<FString, DefaultKeyFuncs<FString>, TInlineSetAllocator<1>>;
+			HandlePossibleAssetLeaks(FOneElementInlineSet{Plugin.GetName()});
 		}
 	}
 
@@ -339,10 +342,7 @@ namespace UE::CoreUObject::Private
 			{
 				CollectGarbage(GARBAGE_COLLECTION_KEEPFLAGS, true);
 
-				for (FString& PluginName : DeferredPluginsToGC)
-				{
-					HandlePossibleAssetLeaks(PluginName);
-				}
+				HandlePossibleAssetLeaks(DeferredPluginsToGC);
 
 				DeferredPluginsToGC.Empty();
 			}
diff --git a/Engine/Source/Runtime/CoreUObject/Private/Misc/CoreUObjectPluginManager.h b/Engine/Source/Runtime/CoreUObject/Private/Misc/CoreUObjectPluginManager.h
--- a/Engine/Source/Runtime/CoreUObject/Private/Misc/CoreUObjectPluginManager.h
+++ b/Engine/Source/Runtime/CoreUObject/Private/Misc/CoreUObjectPluginManager.h
@@ -2,7 +2,7 @@
 
 #pragma once
 #include "ICoreUObjectPluginManager.h"
-#include "Containers/Array.h"
+#include "Containers/Set.h"
 #include "Containers/UnrealString.h"
 
 namespace UE::CoreUObject::Private
@@ -22,7 +22,7 @@ namespace UE::CoreUObject::Private
 
 	private:
 
-		TArray<FString> DeferredPluginsToGC;
+		TSet<FString> DeferredPluginsToGC;
 
 		/** Ref count for deferring calls to OnPluginUnload. When the ref count reaches 0 we GC and leak test all deferred plugins */
 		int32 SuppressGCRefCount = 0;
