Disable GEnableTimeoutOnPendingDestroyedObjectGC when build is compiled with ASAN

ASAN builds can be extremely slow on some platforms. During world trimming, GC may wait too long while the render thread is still flushing commands, which can trigger a timeout crash even though it?s a false positive.

#Synced-CL 51088023

diff --git a/Engine/Source/Runtime/CoreUObject/Private/UObject/GarbageCollection.cpp b/Engine/Source/Runtime/CoreUObject/Private/UObject/GarbageCollection.cpp
--- a/Engine/Source/Runtime/CoreUObject/Private/UObject/GarbageCollection.cpp
+++ b/Engine/Source/Runtime/CoreUObject/Private/UObject/GarbageCollection.cpp
@@ -4862,8 +4862,14 @@ void IncrementalPurgeGarbage(bool bUseTimeLimit, double TimeLimit)
 }
 
 static bool GWarningTimeOutHasBeenDisplayedGC = false;
+#if !USING_ADDRESS_SANITISER
 static bool GEnableTimeoutOnPendingDestroyedObjectGC = true;
-#if UE_BUILD_SHIPPING || USING_ADDRESS_SANITISER
+#else
+// ASAN builds can be extremely slow on some platforms. During world trimming, GC may wait too long while the render thread 
+// is still flushing commands, which can trigger a timeout crash even though it’s a false positive.
+static bool GEnableTimeoutOnPendingDestroyedObjectGC = false;
+#endif
+#if UE_BUILD_SHIPPING
 static FAutoConsoleVariableRef CVarGCEnableTimeoutOnPendingDestroyedObjectInShipping(
 	TEXT("gc.EnableTimeoutOnPendingDestroyedObjectInShipping"),
 	GEnableTimeoutOnPendingDestroyedObjectGC,
