Improved an assert message when storing a migrated object data but the object data already exists locally

#Synced-CL 51045578

diff --git a/Engine/Source/Runtime/CoreUObject/Private/UObject/RemoteObjectTransfer.cpp b/Engine/Source/Runtime/CoreUObject/Private/UObject/RemoteObjectTransfer.cpp
--- a/Engine/Source/Runtime/CoreUObject/Private/UObject/RemoteObjectTransfer.cpp
+++ b/Engine/Source/Runtime/CoreUObject/Private/UObject/RemoteObjectTransfer.cpp
@@ -973,7 +973,15 @@ TRefCountPtr<FTransactionalMigrationData> StoreMigratedObjectData(FRemoteObjectD
 		FRemoteObjectId SerializedId = It.GetId();
 		FRemoteObjectStub* Stub = FindOrAddRemoteObjectStub(SerializedId, FRemoteServerId::GetLocalServerId());
 
-		checkf(!Stub->MigratedData, TEXT("Received remote object %s bundled with %s but the latter object's data already exists locally"), *MigratedData->Context.ObjectId.ToString(), *SerializedId.ToString());
+		checkf(!Stub->MigratedData, TEXT("Received remote object %s (%s) from server %s bundled with %s (%s) but the latter object's data already exists locally (received from server %s and bundled with %s %s)"), 
+			*MigratedData->Context.ObjectId.ToString(), 
+			*MigratedData->Data.GetName(MigratedData->Context.ObjectId).ToString(),
+			*MigratedData->Context.RemoteServerId.ToString(),
+			*SerializedId.ToString(),
+			*MigratedData->Data.GetName(SerializedId).ToString(),
+			*Stub->MigratedData->Context.RemoteServerId.ToString(),
+			*Stub->MigratedData->Context.ObjectId.ToString(),
+			*Stub->MigratedData->Data.GetName(Stub->MigratedData->Context.ObjectId).ToString());
 
 		Stub->OwningServerId = MigratedData->Context.OwnerServerId;
 		Stub->MigratedData = MigratedData;
@@ -1077,7 +1085,7 @@ void DeserializeObjectsFromMigratedData(FTransactionalMigrationData& InOutMigrat
 			ReceivedButNotInstanced.RemoveAtSwap(0);
 
 			auto WrappedGetNameSafe = [](const UObject* O) { return FString::Printf(TEXT("%s (%s)"), *GetNameSafe(O), *(FObjectPtr{ const_cast<UObject*>(O) }.GetRemoteId().ToString())); };
-			UE_CLOG(!ReceivedButNotInstanced.IsEmpty(), LogRemoteObject, Warning, TEXT(" Received objects but they weren't instanced: %s"), *FString::JoinBy(ReceivedButNotInstanced, TEXT(", "), WrappedGetNameSafe));
+			UE_CLOG(!ReceivedButNotInstanced.IsEmpty(), LogRemoteObject, Verbose, TEXT(" Received objects but they weren't instanced: %s"), *FString::JoinBy(ReceivedButNotInstanced, TEXT(", "), WrappedGetNameSafe));
 			UE_CLOG(!AllSubObjects.IsEmpty(), LogRemoteObject, Warning, TEXT(" New (or Existing) instances not received: %s"), *FString::JoinBy(AllSubObjects, TEXT(", "), WrappedGetNameSafe));
 
 			if (CVarRemoveMissingMigratedComponents.GetValueOnGameThread())
diff --git a/Engine/Source/Runtime/CoreUObject/Public/UObject/RemoteObjectTransfer.h b/Engine/Source/Runtime/CoreUObject/Public/UObject/RemoteObjectTransfer.h
--- a/Engine/Source/Runtime/CoreUObject/Public/UObject/RemoteObjectTransfer.h
+++ b/Engine/Source/Runtime/CoreUObject/Public/UObject/RemoteObjectTransfer.h
@@ -130,11 +130,26 @@ struct FRemoteObjectData
 	}
 
 	/**
-	* Returns a unique FName based on its index the id tables
+	* Returns a unique FName based on its index the name tables
 	*/
-	inline FName GetName(FRemoteObjectTables::FNameIndexType IdIndex) const
+	inline FName GetName(FRemoteObjectTables::FNameIndexType NameIndex) const
 	{
-		return Tables.Names[IdIndex];
+		return Tables.Names[NameIndex];
+	}
+
+	/**
+	* Attempts to return an object name associated with a remote id. 'None' if the id does not represent an object serialized in this object data.
+	*/
+	inline FName GetName(FRemoteObjectId RemoteId) const
+	{
+		for (int32 ObjectIndex = 0; ObjectIndex < SerializedObjectHeaders.Num(); ++ObjectIndex)
+		{
+			if (GetSerializedObjectId(ObjectIndex) == RemoteId)
+			{
+				return GetName(SerializedObjectHeaders[ObjectIndex].NameIndex);
+			}
+		}
+		return FName();
 	}
 
 	// Support for iterating over serialized object ids:
