[ZenLoader/ALT]
- Fix possible race between linker attachment to a new package and detach linker to an old package
- Small cleanups

#Synced-CL 51209060

diff --git a/Engine/Source/Runtime/CoreUObject/Private/Serialization/AsyncLoading2.cpp b/Engine/Source/Runtime/CoreUObject/Private/Serialization/AsyncLoading2.cpp
--- a/Engine/Source/Runtime/CoreUObject/Private/Serialization/AsyncLoading2.cpp
+++ b/Engine/Source/Runtime/CoreUObject/Private/Serialization/AsyncLoading2.cpp
@@ -4965,15 +4965,10 @@ public:
 
 	virtual void NotifyRegistrationComplete() override;
 
-	FORCEINLINE FAsyncPackage2* FindAsyncPackage(FPackageId PackageId)
-	{
-		//checkSlow(IsInAsyncLoadThread());
-		return AsyncPackageLookup.FindRef(PackageId);
-	}
-
+	// This is only safe to call from ALT since we're not adding any refs to the package
+	// but that's OK on ALT since we destroy the packages from ALT.
 	FORCEINLINE FAsyncPackage2* GetAsyncPackage(const FPackageId& PackageId)
 	{
-		// TRACE_CPUPROFILER_EVENT_SCOPE(GetAsyncPackage);
 		return AsyncPackageLookup.FindRef(PackageId);
 	}
 
@@ -7247,7 +7242,7 @@ void FAsyncPackage2::ImportPackagesRecursiveInner(FAsyncLoadingThreadState2& Thr
 			const bool bIsMissing = ImportedPackageStatus == EPackageStoreEntryStatus::Missing;
 			if (bIsFullyLoaded)
 			{
-				ImportedPackage = AsyncLoadingThread.FindAsyncPackage(ImportedPackageId);
+				ImportedPackage = AsyncLoadingThread.GetAsyncPackage(ImportedPackageId);
 				if (!ImportedPackage)
 				{
 					continue;
@@ -10752,8 +10747,6 @@ EAsyncPackageState::Type FAsyncLoadingThread2::ProcessLoadedPackagesFromGameThre
 			Package->FinishUPackage();
 
 			{
-				AsyncPackageLookup.Remove(Package->Desc.UPackageId);
-
 #if WITH_EDITOR
 				if (!Package->HasLoadFailed())
 				{
@@ -10773,6 +10766,11 @@ EAsyncPackageState::Type FAsyncLoadingThread2::ProcessLoadedPackagesFromGameThre
 				Package->ClearConstructedObjects();
 
 				Package->DetachLinker();
+
+				// If we were to perform the removal before DetachLinker, it would be possible
+				// for a new async package to be created, and we could race trying to attach
+				// the linker to the new package before its detached by this one.
+				AsyncPackageLookup.Remove(Package->Desc.UPackageId);
 			}
 
 			// Remove the package from the list before we trigger the callbacks, 
@@ -11525,15 +11523,7 @@ void FAsyncLoadingThread2::ResumeLoading()
 
 float FAsyncLoadingThread2::GetAsyncLoadPercentage(const FName& PackageName)
 {
-	float LoadPercentage = -1.0f;
-	/*
-	FAsyncPackage2* Package = FindAsyncPackage(PackageName);
-	if (Package)
-	{
-		LoadPercentage = Package->GetLoadPercentage();
-	}
-	*/
-	return LoadPercentage;
+	return -1.0f;
 }
 
 static void VerifyObjectLoadFlagsWhenFinishedLoading()
