Fix UPackage::GetExternalPackages to return unique packages.

#Synced-CL 51061797

diff --git a/Engine/Source/Runtime/CoreUObject/Private/UObject/Package.cpp b/Engine/Source/Runtime/CoreUObject/Private/UObject/Package.cpp
--- a/Engine/Source/Runtime/CoreUObject/Private/UObject/Package.cpp
+++ b/Engine/Source/Runtime/CoreUObject/Private/UObject/Package.cpp
@@ -279,21 +279,22 @@ UObject* UPackage::FindAssetInPackage(EObjectFlags RequiredTopLevelFlags) const
 
 TArray<UPackage*> UPackage::GetExternalPackages() const
 {
-	TArray<UPackage*> Result;
-	TArray<UObject*> TopLevelObjects;
-	GetObjectsWithPackage(const_cast<UPackage*>(this), TopLevelObjects, false);
-	for (UObject* Object : TopLevelObjects)
-	{
-		ForEachObjectWithOuter(Object, [&Result, ThisPackage = this](UObject* InObject)
-			{
-				UPackage* ObjectPackage = InObject->GetExternalPackage();
-				if (ObjectPackage && ObjectPackage != ThisPackage)
+	TSet<UPackage*> Result;
+	ForEachObjectWithPackage(const_cast<UPackage*>(this), 
+		[this, &Result](UObject* Object)
+		{
+			ForEachObjectWithOuter(Object, 
+				[this, &Result](UObject* InObject)
 				{
-					Result.Add(ObjectPackage);
-				}
-			});
-	}
-	return Result;
+					UPackage* ObjectPackage = InObject->GetExternalPackage();
+					if (ObjectPackage && ObjectPackage != this)
+					{
+						Result.Add(ObjectPackage);
+					}
+				});
+			return true;
+		}, /*bIncludeNestedObjects*/false);	
+	return Result.Array();
 }
 
 #if WITH_METADATA
