Avoid data duplication when storing remote objects to database

Introduced a concept of database redirector files for suobjects stored with their outer hierarchy

#Synced-CL 51045578

diff --git a/Engine/Source/Runtime/CoreUObject/Private/UObject/RemoteObjectSerialization.cpp b/Engine/Source/Runtime/CoreUObject/Private/UObject/RemoteObjectSerialization.cpp
--- a/Engine/Source/Runtime/CoreUObject/Private/UObject/RemoteObjectSerialization.cpp
+++ b/Engine/Source/Runtime/CoreUObject/Private/UObject/RemoteObjectSerialization.cpp
@@ -66,6 +66,39 @@ FArchive& operator<<(FArchive& Ar, FPackedRemoteObjectHeader& Header)
 namespace UE::RemoteObject::Serialization::Disk
 {
 
+enum class ERemoteObjectFileType : int32
+{
+	Unknown = 0,
+	Data = 1,
+	Redirector = 2,
+
+	MaxPlusOne,
+	Max = MaxPlusOne - 1
+};
+
+struct FRemoteObjectFileHeader
+{
+	ERemoteObjectFileType FileType = ERemoteObjectFileType::Unknown;
+	FRemoteObjectId RootObjectId;
+};
+
+void SerializeFileHeader(FArchive& Ar, FRemoteObjectFileHeader& InHeader)
+{
+	int32 SerializedFileType = (int32)InHeader.FileType;
+	Ar << SerializedFileType;
+	Ar << InHeader.RootObjectId;
+}
+
+void DeserializeFileHeader(FArchive& Ar, FRemoteObjectFileHeader& OutHeader)
+{
+	int32 SerializedFileType = -1;
+	Ar << SerializedFileType;
+	checkf(SerializedFileType >= (int32)ERemoteObjectFileType::Data && SerializedFileType <= (int32)ERemoteObjectFileType::Max, TEXT("Unknown remote object file type %d"), SerializedFileType);
+	
+	OutHeader.FileType = (ERemoteObjectFileType)SerializedFileType;
+	Ar << OutHeader.RootObjectId;
+}
+
 void SerializeNameTables(FArchive& Ar, const FRemoteObjectData& InObjectData)
 {
 	FNameIndexType NumNames = IntCastChecked<FNameIndexType>(InObjectData.Tables.Names.Num());
@@ -101,17 +134,46 @@ FString GenerateRemoteObjectFilename(FRemoteObjectId ObjectId, FRemoteServerId O
 	return FPaths::Combine(*FPaths::ProjectSavedDir(), *(FString::Printf(TEXT("%s-%s_%s.remote"), *FRemoteServerId::GetLocalServerId().ToString(), *ObjectId.ToString(ERemoteIdToStringVerbosity::Id), *OwnerServerId.ToString())));
 }
 
-void LoadObjectFromDisk(const FUObjectMigrationContext& MigrationContext)
+void LoadObjectFromDiskInternal(const FUObjectMigrationContext& MigrationContext, FRemoteObjectId RootObjectId, FRemoteObjectData& OutData, ERemoteObjectFileType PreviousFileType)
 {
-	const FString Filename(GenerateRemoteObjectFilename(MigrationContext.ObjectId, MigrationContext.OwnerServerId));
+	FString Filename(GenerateRemoteObjectFilename(RootObjectId, MigrationContext.OwnerServerId));
 	TUniquePtr<FArchive> FileReader{ IFileManager::Get().CreateFileReader(*Filename) };
 	checkf(FileReader.IsValid(), TEXT("Unable to create file reader for remote object %s"), *MigrationContext.ObjectId.ToString());
 
-	FRemoteObjectData ObjectData;	
-	DeserializeNameTables(*FileReader, ObjectData);
-	*FileReader << ObjectData.Bytes;
+	FRemoteObjectFileHeader Header;
+	DeserializeFileHeader(*FileReader, Header);
+
+	if (Header.FileType == ERemoteObjectFileType::Redirector)
+	{
+		// We allow only one level of recursion: a redirector must point to an actual data file
+		checkf(PreviousFileType != ERemoteObjectFileType::Redirector, TEXT("A remote object file redirector redirected to another redirector (Context Object ID: %s, RootObjectId: %s)"),
+			*MigrationContext.ObjectId.ToString(), *RootObjectId.ToString());
+		LoadObjectFromDiskInternal(MigrationContext, Header.RootObjectId, OutData, Header.FileType);
+	}
+	else
+	{
+		DeserializeNameTables(*FileReader, OutData);
+		*FileReader << OutData.Bytes;
+	}
+
 	FileReader->Close();
 	IFileManager::Get().Delete(*Filename, false, true, true);
+}
+
+void LoadObjectFromDisk(const FUObjectMigrationContext& MigrationContext)
+{
+	FRemoteObjectData ObjectData;
+	LoadObjectFromDiskInternal(MigrationContext, MigrationContext.ObjectId, ObjectData, ERemoteObjectFileType::Unknown);
+
+	FRemoteObjectId RootObjectId(ObjectData.GetRootSerializedObjectId());
+	for (FRemoteObjectId SerializedObjectId : ObjectData)
+	{
+		if (SerializedObjectId != RootObjectId)
+		{
+			FString Filename(GenerateRemoteObjectFilename(SerializedObjectId, MigrationContext.OwnerServerId));
+			IFileManager::Get().Delete(*Filename, false, true, true);
+		}
+	}
 
 	// We have transferred ownership from the Database to the local server
 	FRemoteServerId LocalServerId = FRemoteServerId::GetLocalServerId();
@@ -121,12 +183,31 @@ void LoadObjectFromDisk(const FUObjectMigrationContext& MigrationContext)
 
 void SaveObjectToDisk(const UE::RemoteObject::Transfer::FMigrateSendParams& Params)
 {
-	TUniquePtr<FArchive> FileWriter{ IFileManager::Get().CreateFileWriter(*GenerateRemoteObjectFilename(Params.MigrationContext.ObjectId, Params.MigrationContext.OwnerServerId)) };
-	checkf(FileWriter.IsValid(), TEXT("Unable to create file writer for remote object %s"), *Params.MigrationContext.ObjectId.ToString());
-	
-	SerializeNameTables(*FileWriter, Params.ObjectData);
-	*FileWriter << const_cast<TArray<FRemoteObjectBytes>&>(Params.ObjectData.Bytes);
-	FileWriter->Close();
+	checkf(Params.MigrationContext.ObjectId == Params.ObjectData.GetRootSerializedObjectId(), TEXT("Saving a subobject %s to disk. Only its root %s can be saved directly."),
+		*Params.MigrationContext.ObjectId.ToString(), *Params.ObjectData.GetRootSerializedObjectId().ToString());
+
+	{
+		TUniquePtr<FArchive> RootFileWriter{ IFileManager::Get().CreateFileWriter(*GenerateRemoteObjectFilename(Params.MigrationContext.ObjectId, Params.MigrationContext.OwnerServerId)) };
+		checkf(RootFileWriter.IsValid(), TEXT("Unable to create file writer for remote object %s"), *Params.MigrationContext.ObjectId.ToString());
+
+		FRemoteObjectFileHeader RootHeader{ ERemoteObjectFileType::Data, Params.MigrationContext.ObjectId };
+		SerializeFileHeader(*RootFileWriter, RootHeader);
+		SerializeNameTables(*RootFileWriter, Params.ObjectData);
+		*RootFileWriter << const_cast<TArray<FRemoteObjectBytes>&>(Params.ObjectData.Bytes);
+		RootFileWriter->Close();
+	}
+
+	FRemoteObjectFileHeader SubobjectHeader{ ERemoteObjectFileType::Redirector, Params.MigrationContext.ObjectId };
+	for (FRemoteObjectId SerializedObjectId : Params.ObjectData)
+	{
+		if (SerializedObjectId != Params.MigrationContext.ObjectId)
+		{
+			TUniquePtr<FArchive> SubobjectFileWriter{ IFileManager::Get().CreateFileWriter(*GenerateRemoteObjectFilename(SerializedObjectId, Params.MigrationContext.OwnerServerId)) };
+			checkf(SubobjectFileWriter.IsValid(), TEXT("Unable to create file writer for remote object redirector %s"), *SerializedObjectId.ToString());
+			SerializeFileHeader(*SubobjectFileWriter, SubobjectHeader);
+			SubobjectFileWriter->Close();
+		}
+	}
 }
 
 } // namespace UE::RemoteObject::Serialization::Disk
@@ -1291,10 +1372,10 @@ void SerializeObjectDataInternal(FArchiveRemoteObjectWriter& Ar, UObject* Reques
 	}
 }
 
-FRemoteObjectData SerializeObjectData(UObject* InObject, TSet<UObject*>& OutObjects, TSet<UObject*>& OutReferencedObjects, const FUObjectMigrationContext* MigrationContext)
+FRemoteObjectData SerializeObjectData(UObject* InObject, TSet<UObject*>& OutObjects, TSet<UObject*>& OutReferencedObjects, const FUObjectMigrationContext* MigrationContext, ERemoteObjectSerializationFlags InFlags)
 {
 	UObject* Object = InObject;
-	if (GRemoteObjectsMigrateFullHierarchy)
+	if (GRemoteObjectsMigrateFullHierarchy && !EnumHasAnyFlags(InFlags, ERemoteObjectSerializationFlags::SkipCanonicalRootSearch))
 	{
 		Object = FindCanonicalRootObjectForSerialization(Object);
 	}
@@ -1305,7 +1386,7 @@ FRemoteObjectData SerializeObjectData(UObject* InObject, TSet<UObject*>& OutObje
 	{
 		FArchiveRemoteObjectWriter Ar(Object, ObjectData, MigrationContext, nullptr, &OutReferencedObjects);
 		Ar.SetMigratingRemoteObjects(true);
-		SerializeObjectDataInternal(Ar, InObject, RequestedObjectId, OutObjects, ERemoteObjectSerializationFlags::None);
+		SerializeObjectDataInternal(Ar, InObject, RequestedObjectId, OutObjects, InFlags);
 	}
 	checkf(ObjectData.SerializedObjectHeaders.Num(), TEXT("Serialized object data for object %s is empty"), *RequestedObjectId.ToString());
 
diff --git a/Engine/Source/Runtime/CoreUObject/Private/UObject/RemoteObjectSerialization.h b/Engine/Source/Runtime/CoreUObject/Private/UObject/RemoteObjectSerialization.h
--- a/Engine/Source/Runtime/CoreUObject/Private/UObject/RemoteObjectSerialization.h
+++ b/Engine/Source/Runtime/CoreUObject/Private/UObject/RemoteObjectSerialization.h
@@ -17,7 +17,8 @@ namespace UE::RemoteObject::Serialization
 		None = 0,
 		UseExistingObjects = (1 << 0),			// If possible re-use existing UObjects and don't reconstruct them when deserializing object data
 		PreserveRemoteReferences = (1 << 1),		// Don't overwrite references to objects that are remote
-		Resetting = (1 << 2)						// Indicates that the serialization process is resetting an object to its archetype state
+		Resetting = (1 << 2),						// Indicates that the serialization process is resetting an object to its archetype state
+		SkipCanonicalRootSearch = (1 << 3),			// Don't attempt to find a canonical root for serialization
 	};
 	ENUM_CLASS_FLAGS(ERemoteObjectSerializationFlags);
 
@@ -134,9 +135,10 @@ namespace UE::RemoteObject::Serialization
 	* @param OutObjects All objects that have been serialized (Object and its subobjects and/or parent)
 	* @param OutReferencedObjects Keeps track of all objects that need to be tagged with RemoteReference
 	* @param MigrationContext Contains the meta data of the current migration request
+	* @param InFlags Serialization flags
 	* @return Remote object data representing the serialized objects
 	*/
-	FRemoteObjectData SerializeObjectData(UObject* InObject, TSet<UObject*>& OutObjects, TSet<UObject*>& OutReferencedObjects, const FUObjectMigrationContext* MigrationContext);
+	FRemoteObjectData SerializeObjectData(UObject* InObject, TSet<UObject*>& OutObjects, TSet<UObject*>& OutReferencedObjects, const FUObjectMigrationContext* MigrationContext, ERemoteObjectSerializationFlags InFlags = ERemoteObjectSerializationFlags::None);
 
 	/**
 	* Deserializes remote object data
diff --git a/Engine/Source/Runtime/CoreUObject/Private/UObject/RemoteObjectTransfer.cpp b/Engine/Source/Runtime/CoreUObject/Private/UObject/RemoteObjectTransfer.cpp
--- a/Engine/Source/Runtime/CoreUObject/Private/UObject/RemoteObjectTransfer.cpp
+++ b/Engine/Source/Runtime/CoreUObject/Private/UObject/RemoteObjectTransfer.cpp
@@ -1641,8 +1641,21 @@ void StoreObjectToDatabase(UObject* Object, FRemoteObjectId ObjectId)
 			};
 			FScopedObjectMigrationContext ScopedMigrationContext(StoreToDatabaseContext);
 
-			// Do the actual serialization
-			FRemoteObjectData ObjectData = SerializeObjectData(OutermostUnreachableRemoteReference, SentObjects, ReferencedObjects, &StoreToDatabaseContext);
+			// Do the actual serialization but skip canonical root object search since we've already done this with the outermost unreachable reference and
+			// we don't want to find the root again if we're serializing a subobject that's unreachable from its root (Outer)
+			FRemoteObjectData ObjectData = SerializeObjectData(OutermostUnreachableRemoteReference, SentObjects, ReferencedObjects, &StoreToDatabaseContext, ERemoteObjectSerializationFlags::SkipCanonicalRootSearch);
+#if DO_CHECK
+			// Make sure none of the sent objects is marked as remote at this point, otherwise we either stored a remote (an outdated) object memory to database or we stored the object to database twice
+			for (UObject* SentObject : SentObjects)
+			{
+				FRemoteObjectId SentObjectId(SentObject);
+				FRemoteObjectStub* Stub = FindRemoteObjectStub(SentObjectId);
+				checkf(Stub, TEXT("Stub for sent object %s does not exist"), *SentObjectId.ToString());
+				checkf(GetResidence(SentObject) != EResidence::Remote, TEXT("Double store to database for %s? Object is remote on server %s"), 
+					*SentObjectId.ToString(ERemoteIdToStringVerbosity::PathName),
+					*Stub->ResidentServerId.ToString());
+			}
+#endif
 
 			const FMigrateSendParams Params { .MigrationContext = StoreToDatabaseContext, .ObjectData = ObjectData };
 			StoreRemoteObjectDataDelegate.Execute(Params);
@@ -1657,10 +1670,12 @@ void StoreObjectToDatabase(UObject* Object, FRemoteObjectId ObjectId)
 		{
 			MarkAsRemoteReference(ReferencedObject);
 		}
-
-		// MarkAsRemote clears EInternalObjectFlags::RemoteReference flag which will prevent re-enetring this loop with the same outermost object but 
-		// we still want to separately process its subobjects (so that they're stored as separate entries in the database) so only mark the root object as remote
-		SetResidence(OutermostUnreachableRemoteReference, EResidence::Remote, DatabaseId);
+		
+		// Mark all stored objects as remote to make sure we don't accidentally store them again
+		for (UObject* SentObject : SentObjects)
+		{
+			SetResidence(SentObject, EResidence::Remote, DatabaseId);
+		}
 
 		// It's possible that OutermostUnreachableRemoteReference no longer referenced this Object and Object hasn't been stored to disk
 		// in which case we need to repeat this loop until Object is marked as remote (and is no longer marked as RemoteReference).
