Fixing a cvar bound to the wrong global

#Synced-CL 51045578

diff --git a/Engine/Source/Runtime/CoreUObject/Private/UObject/RemoteObjectSerialization.cpp b/Engine/Source/Runtime/CoreUObject/Private/UObject/RemoteObjectSerialization.cpp
--- a/Engine/Source/Runtime/CoreUObject/Private/UObject/RemoteObjectSerialization.cpp
+++ b/Engine/Source/Runtime/CoreUObject/Private/UObject/RemoteObjectSerialization.cpp
@@ -33,7 +33,7 @@ static FAutoConsoleVariableRef CVarRemoteObjectsMigrateFullHierarchy(
 int32 GResetBorrowedObjects = 1;
 static FAutoConsoleVariableRef CVarResetBorrowedObjects(
 	TEXT("ro.ResetBorrowedObjects"),
-	GRemoteObjectsMigrateFullHierarchy,
+	GResetBorrowedObjects,
 	TEXT("Whether remote objects that were borrowed should be reset upon returning to their owner server instead of being reconstructed"));
 
 int32 GUseImmutableArchetypes = 1;
