Fixing verification code incorrectly assuming a stub for a serialized remote object must always exist

#Synced-CL 51045578

diff --git a/Engine/Source/Runtime/CoreUObject/Private/UObject/RemoteObjectTransfer.cpp b/Engine/Source/Runtime/CoreUObject/Private/UObject/RemoteObjectTransfer.cpp
--- a/Engine/Source/Runtime/CoreUObject/Private/UObject/RemoteObjectTransfer.cpp
+++ b/Engine/Source/Runtime/CoreUObject/Private/UObject/RemoteObjectTransfer.cpp
@@ -1649,11 +1649,12 @@ void StoreObjectToDatabase(UObject* Object, FRemoteObjectId ObjectId)
 			for (UObject* SentObject : SentObjects)
 			{
 				FRemoteObjectId SentObjectId(SentObject);
-				FRemoteObjectStub* Stub = FindRemoteObjectStub(SentObjectId);
-				checkf(Stub, TEXT("Stub for sent object %s does not exist"), *SentObjectId.ToString());
-				checkf(GetResidence(SentObject) != EResidence::Remote, TEXT("Double store to database for %s? Object is remote on server %s"), 
-					*SentObjectId.ToString(ERemoteIdToStringVerbosity::PathName),
-					*Stub->ResidentServerId.ToString());
+				if (FRemoteObjectStub* Stub = FindRemoteObjectStub(SentObjectId))
+				{
+					checkf(GetResidence(SentObject) != EResidence::Remote, TEXT("Double store to database for %s? Object is remote on server %s"),
+						*SentObjectId.ToString(ERemoteIdToStringVerbosity::PathName),
+						*Stub->ResidentServerId.ToString());
+				}
 			}
 #endif
 
