Restoring accidentally changed log verbosity

#Synced-CL 51045578

diff --git a/Engine/Source/Runtime/CoreUObject/Private/UObject/RemoteObjectTransfer.cpp b/Engine/Source/Runtime/CoreUObject/Private/UObject/RemoteObjectTransfer.cpp
--- a/Engine/Source/Runtime/CoreUObject/Private/UObject/RemoteObjectTransfer.cpp
+++ b/Engine/Source/Runtime/CoreUObject/Private/UObject/RemoteObjectTransfer.cpp
@@ -1085,7 +1085,7 @@ void DeserializeObjectsFromMigratedData(FTransactionalMigrationData& InOutMigrat
 			ReceivedButNotInstanced.RemoveAtSwap(0);
 
 			auto WrappedGetNameSafe = [](const UObject* O) { return FString::Printf(TEXT("%s (%s)"), *GetNameSafe(O), *(FObjectPtr{ const_cast<UObject*>(O) }.GetRemoteId().ToString())); };
-			UE_CLOG(!ReceivedButNotInstanced.IsEmpty(), LogRemoteObject, Verbose, TEXT(" Received objects but they weren't instanced: %s"), *FString::JoinBy(ReceivedButNotInstanced, TEXT(", "), WrappedGetNameSafe));
+			UE_CLOG(!ReceivedButNotInstanced.IsEmpty(), LogRemoteObject, Warning, TEXT(" Received objects but they weren't instanced: %s"), *FString::JoinBy(ReceivedButNotInstanced, TEXT(", "), WrappedGetNameSafe));
 			UE_CLOG(!AllSubObjects.IsEmpty(), LogRemoteObject, Warning, TEXT(" New (or Existing) instances not received: %s"), *FString::JoinBy(AllSubObjects, TEXT(", "), WrappedGetNameSafe));
 
 			if (CVarRemoveMissingMigratedComponents.GetValueOnGameThread())
