[Engine] Disable garbage reference elimination if compact sets are enabled
* Remove default ini setting so defaults can be used
* Added warning if user attempts to enable it via console/cvar or manually requesting it to change.
* Note: Ideally GC would automatically trigger map/set rehash when changing underlying data, so this is a temporary measure that will prevent bugs while figure it the path to make that work.

#tests Ran engine tests against cl that adds more compact set validation to see it passing
#Synced-CL 51169782

diff --git a/Engine/Config/BaseEngine.ini b/Engine/Config/BaseEngine.ini
--- a/Engine/Config/BaseEngine.ini
+++ b/Engine/Config/BaseEngine.ini
@@ -1712,7 +1712,9 @@ gc.MinGCClusterSize=5
 gc.AssetClustreringEnabled=False
 gc.ActorClusteringEnabled=False
 gc.VerifyUObjectsAreNotFGCObjects=False
-gc.GarbageEliminationEnabled=True
+
+; Example only, commented out so we will Use the CoreUObject default value
+; gc.GarbageEliminationEnabled=True
 
 [Internationalization]
 ; These are the paths for localizing the core engine itself, the localized game content settings are in *Game.ini
diff --git a/Engine/Source/Runtime/CoreUObject/Private/UObject/ObjectBaseUtility.cpp b/Engine/Source/Runtime/CoreUObject/Private/UObject/ObjectBaseUtility.cpp
--- a/Engine/Source/Runtime/CoreUObject/Private/UObject/ObjectBaseUtility.cpp
+++ b/Engine/Source/Runtime/CoreUObject/Private/UObject/ObjectBaseUtility.cpp
@@ -167,7 +167,7 @@ void FScopeCycleCounterUObject::StartObjectTrace(const UObjectBaseUtility* Objec
 
 
 // Console variable so that GarbageCollectorSettings work in the editor but we don't want to use it in runtime as we can't support changing its value from console
-int32 GGarbageEliminationEnabled = 1;
+int32 GGarbageEliminationEnabled = !UE_USE_COMPACT_SET_AS_DEFAULT;
 static FAutoConsoleVariableRef CVarGarbageEliminationEnabled(
 	TEXT("gc.GarbageEliminationEnabled"),
 	GGarbageEliminationEnabled,
@@ -179,7 +179,7 @@ bool UObjectBaseUtility::bGarbageEliminationEnabled = !!GGarbageEliminationEnabl
 void InitGarbageElimination()
 {
 	check(GConfig);
-	bool bGarbageEliminationEnabled = true;
+	bool bGarbageEliminationEnabled = UObjectBaseUtility::IsGarbageEliminationEnabled();
 	if (GConfig->GetBool(TEXT("/Script/Engine.GarbageCollectionSettings"), TEXT("gc.PendingKillEnabled"), bGarbageEliminationEnabled, GEngineIni))
 	{
 		// The old gc.PendingKillEnabled will always take precedence over gc.GarbageEliminationEnabled because the new setting is always present in BaseEngine.ini
@@ -201,7 +201,25 @@ void InitGarbageElimination()
 	{
 		bGarbageEliminationEnabled = true;
 	}
-	// Try to sync even though we're not gonna use the console var
-	UObjectBaseUtility::bGarbageEliminationEnabled = bGarbageEliminationEnabled;
-	GGarbageEliminationEnabled = bGarbageEliminationEnabled;
+
+	UObjectBaseUtility::SetGarbageEliminationEnabled(bGarbageEliminationEnabled);
 }
+
+void UObjectBaseUtility::SetGarbageEliminationEnabled(bool bEnabled)
+{
+#if !UE_USE_COMPACT_SET_AS_DEFAULT
+	// Try to sync even though we're not gonna use the console var
+	UObjectBaseUtility::bGarbageEliminationEnabled = bEnabled;
+	GGarbageEliminationEnabled = bEnabled;
+#else
+	if (bEnabled)
+	{
+		// Garbage elimination is not compatible when compact sets are default since it can modify key data which invalidates
+		// the lookup table used to manage the set which can lead to crashes or infinite loops on subsequent operations.
+		UE_LOG(LogObj, Warning, TEXT("Request made to enable Garbage Elimination, however this feature is not currently compatible with compact sets. Leaving it disabled."));
+	}
+
+	UObjectBaseUtility::bGarbageEliminationEnabled = false;
+	GGarbageEliminationEnabled = 0;
+#endif
+}
\ No newline at end of file
diff --git a/Engine/Source/Runtime/CoreUObject/Public/UObject/UObjectBaseUtility.h b/Engine/Source/Runtime/CoreUObject/Public/UObject/UObjectBaseUtility.h
--- a/Engine/Source/Runtime/CoreUObject/Public/UObject/UObjectBaseUtility.h
+++ b/Engine/Source/Runtime/CoreUObject/Public/UObject/UObjectBaseUtility.h
@@ -46,7 +46,6 @@ class UObjectBaseUtility : public UObjectBase
 	/** If true references to objects marked as Garbage will be automatically eliminated by Garbage Collector*/
 	static COREUOBJECT_API bool bGarbageEliminationEnabled;
 
-	friend void InitGarbageElimination();
 	friend struct FInternalUObjectBaseUtilityIsValidFlagsChecker;
 
 public:
@@ -407,8 +406,11 @@ public:
 		return IsGarbageEliminationEnabled();
 	}
 
-	/** Helper function to set the private bGarbageEliminationEnabled variable. */
-	static inline void SetGarbageEliminationEnabled(bool bEnabled)
+	/** Helper function to try to set the private bGarbageEliminationEnabled variable. */
+	static COREUOBJECT_API void SetGarbageEliminationEnabled(bool bEnabled);
+
+	/** Helper function to set the private bGarbageEliminationEnabled variable for testing purposes only as garbage elimination may not be available for general use based on compile settings. */
+	static inline void ForceSetGarbageEliminationEnabledForTest(bool bEnabled)
 	{
 		bGarbageEliminationEnabled = bEnabled;
 	}
