[CoreUObject]

StaticDuplicateObjectEx will now regather default sub-objects (DSOs) when duplicating with the FDuplicateDataWriter. FDuplicateDataWriter initially gathers DSOs before creating duplicate objects, however it can end up duplicating objects like ActorComponents that themselves contain DSOs that won't be found in the initial gather. As a result, FDuplicateDataWriter would reconstruct these objects in-place as the FDuplicateDataWriter iterates over objects duplicated during serialization. In the case of DSOs we do not want this reconstruction to occur, since it's wasteful, but it is also not threadsafe when duplicating from the AsyncLoadingThread.

#rnx

#Synced-CL 51168131

diff --git a/Engine/Source/Runtime/CoreUObject/Private/UObject/UObjectGlobals.cpp b/Engine/Source/Runtime/CoreUObject/Private/UObject/UObjectGlobals.cpp
--- a/Engine/Source/Runtime/CoreUObject/Private/UObject/UObjectGlobals.cpp
+++ b/Engine/Source/Runtime/CoreUObject/Private/UObject/UObjectGlobals.cpp
@@ -2996,6 +2996,15 @@ UObject* StaticDuplicateObjectEx( FObjectDuplicationParameters& Parameters )
 		Object->Serialize(Writer);
 		SerializedObjects.Emplace(Object, ObjectOffset);
 
+		// After serialization we may have created default sub-objects (DSO) that we don't want to duplicate on subsequent calls.
+		// e.g. Duplicating an ActorComponent that contains a DSO will have appended the ActorComponent to the UnserializedObjects but not accounted for the 
+		// DSO in the annotation map. We do that here now.
+		FDuplicatedObject DupObj = DuplicatedObjectAnnotation.GetAnnotation(Object);
+		if (!DupObj.IsDefault())
+		{
+			GatherDefaultSubobjectsForDuplication(Object, DupObj.DuplicatedObject.GetEvenIfUnreachable(), DuplicatedObjectAnnotation, Writer);
+		}
+
 #if WITH_VERSE_VM || defined(__INTELLISENSE__)
 		DuplicatedCells.WriteUnserializedCells(Context, Writer);
 #endif
