[CoreUObject]

Changing Fatal World Leak assert to be a log error so that it can be logged by gauntlet in full. Fatal asserts are truncated to 10 lines currently. Emit a new fatal log after the LogError of the leak info so we still terminate

#rnx

#Synced-CL 51236738

diff --git a/Engine/Source/Runtime/CoreUObject/Private/UObject/ReferenceChainSearch.cpp b/Engine/Source/Runtime/CoreUObject/Private/UObject/ReferenceChainSearch.cpp
--- a/Engine/Source/Runtime/CoreUObject/Private/UObject/ReferenceChainSearch.cpp
+++ b/Engine/Source/Runtime/CoreUObject/Private/UObject/ReferenceChainSearch.cpp
@@ -2050,12 +2050,15 @@ TArray<FString> FReferenceChainSearch::FindAndPrintStaleReferencesToObjects(TCon
 	if ((Options & EPrintStaleReferencesOptions::Fatal) == EPrintStaleReferencesOptions::Fatal)
 	{
 		// If you change the format of this string, update the corresponding AutomationUtils/Matcher/WorldLeakEventMatcher.cs
-		UE_LOG(LogLoad, Fatal, 
+		
+		// Emit an error and then a fatal error since Gauntlet can truncate Fatal Error messages.
+		UE_LOG(LogLoad, Error, 
 			TEXT("====Fatal World Leaks====") LINE_TERMINATOR
 			TEXT("Logging first error, check logs for additional information") LINE_TERMINATOR 
 			TEXT("%s")
 			TEXT("====End Fatal World Leaks====") LINE_TERMINATOR,
 			Paths.Num() ? *Paths[0] : TEXT("No paths to leaked objects found!"));
+		UE_LOGF(LogLoad, Fatal, "Crashing due to uobject leaks (world leaks). See log above for details.");
 	}
 
 	return Paths;
