Dstm: Add call to TouchResidentObject when object pointer is dereferenced

#Synced-CL 51166368

diff --git a/Engine/Source/Runtime/CoreUObject/Public/UObject/ObjectHandle.h b/Engine/Source/Runtime/CoreUObject/Public/UObject/ObjectHandle.h
--- a/Engine/Source/Runtime/CoreUObject/Public/UObject/ObjectHandle.h
+++ b/Engine/Source/Runtime/CoreUObject/Public/UObject/ObjectHandle.h
@@ -709,7 +709,7 @@ namespace UE::CoreUObject::Private
 		{
 			Handle = MakeObjectHandle(ResolvedObject);
 		}
-
+		UE::RemoteObject::Handle::TouchResidentObject(ResolvedObject);
 		return ResolvedObject;
 #else
 		return ReadObjectHandlePointerNoCheck(Handle);
