Serialize all subobjects of the canonical root object (not only the subobjects that are actually referenced) when migrating objects

#Synced-CL 51290442

diff --git a/Engine/Source/Runtime/CoreUObject/Internal/UObject/UObjectMigrationContext.h b/Engine/Source/Runtime/CoreUObject/Internal/UObject/UObjectMigrationContext.h
--- a/Engine/Source/Runtime/CoreUObject/Internal/UObject/UObjectMigrationContext.h
+++ b/Engine/Source/Runtime/CoreUObject/Internal/UObject/UObjectMigrationContext.h
@@ -62,7 +62,7 @@ struct FUObjectMigrationContext final
 {
 	struct FMigratedObjectStatus
 	{
-		bool bIsOwned = false;
+		FRemoteServerId OwningServerId;
 		EResidence Residence = EResidence::Local;
 	};
 
@@ -92,7 +92,7 @@ struct FUObjectMigrationContext final
 
 	inline bool IsOwned(FRemoteObjectId InObjectId) const
 	{
-		return CachedObjectStatus.FindChecked(InObjectId).bIsOwned;
+		return CachedObjectStatus.FindChecked(InObjectId).OwningServerId.IsLocal();
 	}
 
 	inline EResidence GetResidence(FRemoteObjectId InObjectId) const
@@ -100,6 +100,15 @@ struct FUObjectMigrationContext final
 		return CachedObjectStatus.FindChecked(InObjectId).Residence;
 	}
 
+	inline FRemoteServerId TryGetOwningServerId(FRemoteObjectId InObjectId) const
+	{
+		if (FMigratedObjectStatus* Status = CachedObjectStatus.Find(InObjectId))
+		{
+			return Status->OwningServerId;
+		}
+		return FRemoteServerId();
+	}
+
 	/**
 	* During migration, figure out the role of the object being received
 	*/
diff --git a/Engine/Source/Runtime/CoreUObject/Private/UObject/RemoteObject.cpp b/Engine/Source/Runtime/CoreUObject/Private/UObject/RemoteObject.cpp
--- a/Engine/Source/Runtime/CoreUObject/Private/UObject/RemoteObject.cpp
+++ b/Engine/Source/Runtime/CoreUObject/Private/UObject/RemoteObject.cpp
@@ -719,27 +719,29 @@ bool IsValid(FRemoteObjectId ObjectId)
 	return false;
 }
 
-FRemoteServerId GetOwnerServerId(const UObject* Object)
+FRemoteServerId GetOwnerServerId(FRemoteObjectId ObjectId)
 {
 	using namespace UE::CoreUObject::Private;
 	using namespace UE::RemoteObject::Private;
 
 	FRemoteServerId Result;
-#if UE_WITH_REMOTE_OBJECT_HANDLE
-	UE::RemoteObject::Handle::FRemoteObjectStub* RemoteStub = FindRemoteObjectStub(FObjectHandleUtils::GetRemoteId(Object));
-	if (RemoteStub)
+	if (FRemoteObjectStub* RemoteStub = FindRemoteObjectStub(ObjectId))
 	{
 		Result = RemoteStub->OwningServerId;
 	}
 	else
 	{
-		// if the object wasn't received or ever migrated, we own it locally
-		Result = FRemoteServerId::GetLocalServerId();
+		// if the object wasn't received or ever migrated, use the encoded server id. This may not reflect the actual state of the object since we don't track ownership across servers if we don't own the object
+		Result = ObjectId.GetServerId();
 	}
-#endif
 	return Result;
 }
 
+FRemoteServerId GetOwnerServerId(const UObject* Object)
+{
+	return GetOwnerServerId(FRemoteObjectId(Object));
+}
+
 void ChangeOwnerServerId(const UObject* Object, FRemoteServerId NewOwnerServerId)
 {
 	using namespace UE::CoreUObject::Private;
diff --git a/Engine/Source/Runtime/CoreUObject/Private/UObject/RemoteObjectSerialization.cpp b/Engine/Source/Runtime/CoreUObject/Private/UObject/RemoteObjectSerialization.cpp
--- a/Engine/Source/Runtime/CoreUObject/Private/UObject/RemoteObjectSerialization.cpp
+++ b/Engine/Source/Runtime/CoreUObject/Private/UObject/RemoteObjectSerialization.cpp
@@ -30,6 +30,12 @@ static FAutoConsoleVariableRef CVarRemoteObjectsMigrateFullHierarchy(
 	GRemoteObjectsMigrateFullHierarchy,
 	TEXT("Whether remote objects that are default subobjects should be always migrated with their parent objects"));
 
+int32 GRemoteObjectsMigrateAllSubobjects = 1;
+static FAutoConsoleVariableRef CVarRemoteObjectsMigrateAllSubobjects(
+	TEXT("ro.MigrateAllSubobjects"),
+	GRemoteObjectsMigrateAllSubobjects,
+	TEXT("Whether to collect all objects with the canonical root as outer for serialization (and not only the subobjects that are actually referenced by their root)"));
+
 int32 GResetBorrowedObjects = 1;
 static FAutoConsoleVariableRef CVarResetBorrowedObjects(
 	TEXT("ro.ResetBorrowedObjects"),
@@ -529,7 +535,7 @@ public:
 		using namespace UE::CoreUObject::Private;
 
 		OutHeader.Name = AddNameToNameMap(Object->GetFName());
-		OutHeader.RemoteId = AddRemoteIdToIdMap(FObjectHandleUtils::GetRemoteId(Object));
+		OutHeader.RemoteId = AddRemoteIdToIdMap(FRemoteObjectId(Object));
 		OutHeader.Class = Object->GetClass();
 		UObject* Outer = FObjectHandleUtils::GetNonAccessTrackedOuterNoResolve(Object);
 		OutHeader.Outer = Outer;
@@ -1386,6 +1392,10 @@ FRemoteObjectData SerializeObjectData(UObject* InObject, TSet<UObject*>& OutObje
 	{
 		FArchiveRemoteObjectWriter Ar(Object, ObjectData, MigrationContext, nullptr, &OutReferencedObjects);
 		Ar.SetMigratingRemoteObjects(true);
+		if (GRemoteObjectsMigrateAllSubobjects)
+		{
+			GetObjectsWithOuter(Ar.GetRootObject(), Ar.GetObjectsToSerialize(),  /*bIncludeNestedObjects =*/ true);
+		}
 		SerializeObjectDataInternal(Ar, InObject, RequestedObjectId, OutObjects, InFlags);
 	}
 	checkf(ObjectData.SerializedObjectHeaders.Num(), TEXT("Serialized object data for object %s is empty"), *RequestedObjectId.ToString());
@@ -1564,20 +1574,38 @@ int32 DeserializeObjectData(FRemoteObjectData& ObjectData,
 			return RequestedObjectIndex;
 		}
 
-		// If we're already resetting a borrowed object we don't want to change the ownership until the object is deserialized using remote server data
+		// If we're already resetting a borrowed object we don't want to change the ownership (see below for an exception) until the object is deserialized using remote server data
 		// and we also don't want to recursively reset the object
-		if (!bResetting)
+		if (!bResetting && RootObject)
 		{
+			TArray<UObject*> CreatedObjects;
+			CreatedObjects.Reserve(OutObjects.Num());
+			CreatedObjects.Add(RootObject);
+			GetObjectsWithOuter(RootObject, CreatedObjects, /*bIncludeNestedObjects =*/ true);
+
 			// ensure remote object stubs are created (ownership will be assigned after PostMigrate)
-			for (int32 ObjectIndex = 0; ObjectIndex < OutObjects.Num(); ++ObjectIndex)
+			for (UObject* Object : CreatedObjects)
 			{
-				if (UObject* Object = OutObjects[ObjectIndex])
+				FRemoteObjectId ObjectId(Object);
+				// Call FindOrAddRemoteObjectStub instead of RegisterRemoteObjectId to always update the resident server id
+				if (FRemoteObjectStub* Stub = FindOrAddRemoteObjectStub(ObjectId, LocalServerId))
 				{
-					// Call FindOrAddRemoteObjectStub instead of RegisterRemoteObjectId to always update the resident server id
-					UE::RemoteObject::Private::FindOrAddRemoteObjectStub(FRemoteObjectId(Object), LocalServerId);
+					// It's possible we created more objects than we received (maybe one of the subobjects has already been GC'd on the other side)
+					// In this case make sure to assign the new object ownership to the server that owns the root object
+					if (MigrationContext && !MigrationContext->TryGetOwningServerId(ObjectId).IsValid())
+					{
+						Stub->OwningServerId = MigrationContext->TryGetOwningServerId(OutObjectRemoteIds[0]);
+					}
 				}
 			}
 
+#if DO_CHECK
+			for (UObject* DeserializedObject : OutObjects)
+			{
+				checkf(!DeserializedObject || CreatedObjects.Contains(DeserializedObject), TEXT("Deserialized an object (%s) that's outside of the root object hierarchy"), *GetPathNameSafe(DeserializedObject));
+			}
+#endif
+
 			if (bReturningBorrowedObject)
 			{
 				ResetRemoteObject(OutObjects[0]);
diff --git a/Engine/Source/Runtime/CoreUObject/Private/UObject/RemoteObjectTransfer.cpp b/Engine/Source/Runtime/CoreUObject/Private/UObject/RemoteObjectTransfer.cpp
--- a/Engine/Source/Runtime/CoreUObject/Private/UObject/RemoteObjectTransfer.cpp
+++ b/Engine/Source/Runtime/CoreUObject/Private/UObject/RemoteObjectTransfer.cpp
@@ -1063,26 +1063,39 @@ void DeserializeObjectsFromMigratedData(FTransactionalMigrationData& InOutMigrat
 		}
 
 		// Let's make sure we received data for all of the objects
-		const UObject* RootObject = ReceivedObjects[0];
+		UObject* RootObject = ReceivedObjects[0];
 		TArray<UObject*> AllSubObjects;
 		constexpr bool bIncludeNestedObjects = true;
 		constexpr EObjectFlags ExclusionFlags = EObjectFlags::RF_MirroredGarbage;
 		GetObjectsWithOuter(RootObject, AllSubObjects, bIncludeNestedObjects, ExclusionFlags);
 
-		const bool bSameNumberOfObjects = (1 + AllSubObjects.Num() == ReceivedObjects.Num());
+		// Apply the same filter to ReceivedObjects as when calling GetObjectsWithOuter()
+		TArray<UObject*> FilteredReceivedObjects;
+		FilteredReceivedObjects.Reserve(ReceivedObjects.Num());
+		for (UObject* Obj : ReceivedObjects)
+		{
+			if (!Obj->HasAnyFlags(ExclusionFlags))
+			{
+				FilteredReceivedObjects.Add(Obj);
+			}
+		}
+		const bool bSameNumberOfObjects = (1 + AllSubObjects.Num() == FilteredReceivedObjects.Num());
 		if (!bSameNumberOfObjects)
 		{
 			UE_LOG(LogRemoteObject, Warning, TEXT("While migrating '%s' [%s -> %s] we received %d Objects whereas %d Objects got instantiated after deserialization.  Differences:"), *GetNameSafe(RootObject), *MigrationContext.RemoteServerId.ToString(), *LocalServerId.ToString(), ReceivedObjects.Num(), 1 + AllSubObjects.Num());
 			TArray<UObject*> ReceivedButNotInstanced;
-			for (UObject* Obj : ReceivedObjects)
+			for (UObject* Obj : FilteredReceivedObjects)
 			{
 				if (AllSubObjects.Remove(Obj) < 1)
 				{
 					ReceivedButNotInstanced.Add(Obj);
 				}
+			}			
+			if (!ReceivedButNotInstanced.IsEmpty())
+			{
+				// We'll always have the root object in this list due to AllSubObjects not containing the root object
+				ReceivedButNotInstanced.Remove(RootObject);
 			}
-			// We'll always have the root object in this list due to AllSubObjects not containing the root object
-			ReceivedButNotInstanced.RemoveAtSwap(0);
 
 			auto WrappedGetNameSafe = [](const UObject* O) { return FString::Printf(TEXT("%s (%s)"), *GetNameSafe(O), *(FObjectPtr{ const_cast<UObject*>(O) }.GetRemoteId().ToString())); };
 			UE_CLOG(!ReceivedButNotInstanced.IsEmpty(), LogRemoteObject, Warning, TEXT(" Received objects but they weren't instanced: %s"), *FString::JoinBy(ReceivedButNotInstanced, TEXT(", "), WrappedGetNameSafe));
diff --git a/Engine/Source/Runtime/CoreUObject/Private/UObject/UObjectMigrationContext.cpp b/Engine/Source/Runtime/CoreUObject/Private/UObject/UObjectMigrationContext.cpp
--- a/Engine/Source/Runtime/CoreUObject/Private/UObject/UObjectMigrationContext.cpp
+++ b/Engine/Source/Runtime/CoreUObject/Private/UObject/UObjectMigrationContext.cpp
@@ -20,16 +20,16 @@ const FUObjectMigrationContext::FMigratedObjectStatus& FUObjectMigrationContext:
 {
 	if (const FMigratedObjectStatus* CachedStatus = CachedObjectStatus.Find(InObjectId))
 	{
-		checkf(CachedStatus->bIsOwned == UE::RemoteObject::Handle::IsOwned(InObjectId), TEXT("Ownership changed after it was first cached (was owned: %s, is owned: %s"),
-			CachedStatus->bIsOwned ? TEXT("yes") : TEXT("no"),
-			UE::RemoteObject::Handle::IsOwned(InObjectId) ? TEXT("yes") : TEXT("no"));
+		checkf(CachedStatus->OwningServerId == UE::RemoteObject::Handle::GetOwnerServerId(InObjectId), TEXT("Ownership changed after it was first cached (cached owner: %s, current owner: %s"),
+			*CachedStatus->OwningServerId.ToString(),
+			*UE::RemoteObject::Handle::GetOwnerServerId(InObjectId).ToString());
 		checkf(CachedStatus->Residence == UE::RemoteObject::Handle::GetResidence(InObjectId), TEXT("Residence changed after it was first cached (was: %s, now: %s)"),
 			EnumToString(CachedStatus->Residence),
 			EnumToString(UE::RemoteObject::Handle::GetResidence(InObjectId)));
 		return *CachedStatus;
 	}
 
-	return CachedObjectStatus.Add(InObjectId, { UE::RemoteObject::Handle::IsOwned(InObjectId), UE::RemoteObject::Handle::GetResidence(InObjectId) });
+	return CachedObjectStatus.Add(InObjectId, { UE::RemoteObject::Handle::GetOwnerServerId(InObjectId), UE::RemoteObject::Handle::GetResidence(InObjectId)});
 }
 
 void FUObjectMigrationContext::CacheObjectStatus(const FRemoteObjectData& ObjectData)
@@ -52,7 +52,7 @@ EObjectMigrationRecvType FUObjectMigrationContext::GetObjectMigrationRecvType(FR
 	const bool bAlreadyOwns = IsOwned(InObjectId);
 
 	const bool bAssignOwnership = !bAlreadyOwns && (OwnerServerId == LocalServerId);
-	const bool bReturnLoan = bAlreadyOwns && ensureMsgf(OwnerServerId == LocalServerId, TEXT("We are receiving object %s that we own, so the data we received should affirm we are the owner."), *InObjectId.ToString());
+	const bool bReturnLoan = bAlreadyOwns && ensureMsgf(OwnerServerId == LocalServerId, TEXT("We are receiving object %s that we own, so the data we received should affirm we are the owner. (OwnerServerId: %s)"), *InObjectId.ToString(), *OwnerServerId.ToString());
 
 	if (bAssignOwnership)
 	{
@@ -87,7 +87,7 @@ EObjectMigrationSendType FUObjectMigrationContext::GetObjectMigrationSendType(FR
 	}
 
 	const FRemoteServerId LocalServerId = FRemoteServerId::GetLocalServerId();
-	const bool bCurrentlyOwns = CacheObjectStatus(InObjectId).bIsOwned;
+	const bool bCurrentlyOwns = CacheObjectStatus(InObjectId).OwningServerId.IsLocal();
 	const bool bOwnershipTransfer = (OwnerServerId == RemoteServerId) && bCurrentlyOwns;
 	const bool bReturnBorrowed = (OwnerServerId == RemoteServerId) && !bCurrentlyOwns;
 
diff --git a/Engine/Source/Runtime/CoreUObject/Public/UObject/RemoteObject.h b/Engine/Source/Runtime/CoreUObject/Public/UObject/RemoteObject.h
--- a/Engine/Source/Runtime/CoreUObject/Public/UObject/RemoteObject.h
+++ b/Engine/Source/Runtime/CoreUObject/Public/UObject/RemoteObject.h
@@ -212,6 +212,11 @@ namespace UE::RemoteObject::Handle
 	*/
 	COREUOBJECT_API FRemoteServerId GetOwnerServerId(const UObject* Object);
 
+	/**
+	* Get the owner server id for the specified object id (may not report the correct owner if the object is not local)
+	*/
+	COREUOBJECT_API FRemoteServerId GetOwnerServerId(FRemoteObjectId ObjectId);
+
 	/**
 	* Sets the owner server id for a locally resident object
 	*/
