UPackage: Fix typo in memcpy, sizeof(&SavedHash.GetBytes()) == 8, while sizeof(SavedHash.GetBytes()) == 20 was the intended value.
It didn't matter with current code, which coincidentally had the same 0 values already after byte 8, but it would break if that ever changes.
#rnx

#Synced-CL 51330552

diff --git a/Engine/Source/Runtime/CoreUObject/Private/UObject/UObjectGlobals.cpp b/Engine/Source/Runtime/CoreUObject/Private/UObject/UObjectGlobals.cpp
--- a/Engine/Source/Runtime/CoreUObject/Private/UObject/UObjectGlobals.cpp
+++ b/Engine/Source/Runtime/CoreUObject/Private/UObject/UObjectGlobals.cpp
@@ -6517,7 +6517,7 @@ namespace UECodeGen_Private
 		// Set the initial saved hash to a value based on the CRCs; this is needed for script packages.
 		FIoHash SavedHash;
 		FMemory::Memcpy(&SavedHash.GetBytes(), &DeterministicGuid,
-			FMath::Min(sizeof(&SavedHash.GetBytes()), sizeof(DeterministicGuid)));
+			FMath::Min(sizeof(SavedHash.GetBytes()), sizeof(DeterministicGuid)));
 		NewPackage->SetSavedHash(SavedHash);
 #endif
 
