UPackage: Silence static analysis warning V568 "It's odd that the argument of sizeof() operator...".
sizeof(SavedHash.GetBytes()) is the correct value: GetBytes() returns ByteArray& (sizeof == 20), and the call results in the correct instructions that copy the 4 ints (16 bytes) from DeterministicGuid to SavedHash.
#rnx

#Synced-CL 51330552

diff --git a/Engine/Source/Runtime/CoreUObject/Private/UObject/UObjectGlobals.cpp b/Engine/Source/Runtime/CoreUObject/Private/UObject/UObjectGlobals.cpp
--- a/Engine/Source/Runtime/CoreUObject/Private/UObject/UObjectGlobals.cpp
+++ b/Engine/Source/Runtime/CoreUObject/Private/UObject/UObjectGlobals.cpp
@@ -6516,8 +6516,9 @@ namespace UECodeGen_Private
 		NewPackage->SetPersistentGuid(DeterministicGuid);
 		// Set the initial saved hash to a value based on the CRCs; this is needed for script packages.
 		FIoHash SavedHash;
+		constexpr size_t SavedHashSize = sizeof(SavedHash.GetBytes()); // -V568
 		FMemory::Memcpy(&SavedHash.GetBytes(), &DeterministicGuid,
-			FMath::Min(sizeof(SavedHash.GetBytes()), sizeof(DeterministicGuid)));
+			FMath::Min(SavedHashSize, sizeof(DeterministicGuid)));
 		NewPackage->SetSavedHash(SavedHash);
 #endif
 
