Fix a bug where packages that failed to load within a transaction would forever fail to load.

#Synced-CL 51291051

#Synced-CL 51334491

diff --git a/Engine/Source/Programs/AutoRTFMTests/Private/UPackageTests.cpp b/Engine/Source/Programs/AutoRTFMTests/Private/UPackageTests.cpp
--- a/Engine/Source/Programs/AutoRTFMTests/Private/UPackageTests.cpp
+++ b/Engine/Source/Programs/AutoRTFMTests/Private/UPackageTests.cpp
@@ -755,4 +755,63 @@ TEST_CASE("UPackage.AsyncLoading")
 			});
 		});
 	}
+
+	SECTION("Packaged Load Fails First Time")
+	{
+		FString Name(FString::Printf(TEXT("/Game/%dAutoRTFMTestPackage%d"), FPlatformProcess::GetCurrentProcessId(), __LINE__));
+
+		{
+			bool bHit = false;
+
+			FLoadPackageAsyncDelegate CompletionDelegate;
+			CompletionDelegate.BindLambda([&](const FName&, UPackage*, EAsyncLoadingResult::Type Type) { bHit = true; REQUIRE(Type == EAsyncLoadingResult::FailedMissing); });
+
+			FlushAsyncLoading(LoadPackageAsync(FPackagePath::FromPackageNameUnchecked(*Name), NAME_None, CompletionDelegate));
+
+			REQUIRE(bHit);
+		}
+
+		FPackageScopedMaker _(Name);
+
+		{
+			bool bHit = false;
+
+			FLoadPackageAsyncDelegate CompletionDelegate;
+			CompletionDelegate.BindLambda([&](const FName&, UPackage*, EAsyncLoadingResult::Type Type) { bHit = true; REQUIRE(Type == EAsyncLoadingResult::Succeeded); });
+
+			FlushAsyncLoading(LoadPackageAsync(FPackagePath::FromPackageNameUnchecked(*Name), NAME_None, CompletionDelegate));
+
+			REQUIRE(bHit);
+		}
+	}
+
+	SECTION("Packaged Load Fails First Time, Second Try in Transaction")
+	{
+		FString Name(FString::Printf(TEXT("/Game/%dAutoRTFMTestPackage%d"), FPlatformProcess::GetCurrentProcessId(), __LINE__));
+
+		{
+			bool bHit = false;
+
+			FLoadPackageAsyncDelegate CompletionDelegate;
+			CompletionDelegate.BindLambda([&](const FName&, UPackage*, EAsyncLoadingResult::Type Type) { bHit = true; REQUIRE(Type == EAsyncLoadingResult::FailedMissing); });
+
+			FlushAsyncLoading(LoadPackageAsync(FPackagePath::FromPackageNameUnchecked(*Name), NAME_None, CompletionDelegate));
+
+			REQUIRE(bHit);
+		}
+
+		FPackageScopedMaker _(Name);
+
+		AutoRTFM::Transact([&]
+		{
+			bool bHit = false;
+
+			FLoadPackageAsyncDelegate CompletionDelegate;
+			CompletionDelegate.BindLambda([&](const FName&, UPackage*, EAsyncLoadingResult::Type Type) { bHit = true; REQUIRE(Type == EAsyncLoadingResult::Succeeded); });
+
+			FlushAsyncLoading(LoadPackageAsync(FPackagePath::FromPackageNameUnchecked(*Name), NAME_None, CompletionDelegate));
+
+			REQUIRE(bHit);
+		});
+	}
 }
diff --git a/Engine/Source/Runtime/CoreUObject/Private/Serialization/TransactionallySafeAsyncLoading.cpp b/Engine/Source/Runtime/CoreUObject/Private/Serialization/TransactionallySafeAsyncLoading.cpp
--- a/Engine/Source/Runtime/CoreUObject/Private/Serialization/TransactionallySafeAsyncLoading.cpp
+++ b/Engine/Source/Runtime/CoreUObject/Private/Serialization/TransactionallySafeAsyncLoading.cpp
@@ -207,6 +207,43 @@ public:
 				}
 				else
 				{
+					// If we've got a package cache hit on a load that did not succeed we *only* want
+					// to preserve this entry in the package cache until the transaction nest completes
+					// Otherwise we could have a package load that fails, which a subsequent load could
+					// have succeeded on (say the package became available on disk), and we'd never be
+					// able to load that package.
+					if (MapPayload->Result != EAsyncLoadingResult::Succeeded)
+					{
+						const bool bFirst = TransactionallyLoadedPackagesThatDidNotSucceed.IsEmpty();
+
+						TransactionallyLoadedPackagesThatDidNotSucceed.Add(PackagePathAsString);
+
+						if (bFirst)
+						{
+							AutoRTFM::OnComplete([this]
+							{
+								UE::TScopeLock _(Mutex);
+
+								for (const FString& PackagePathAsString : TransactionallyLoadedPackagesThatDidNotSucceed)
+								{
+									FMapPayload MapPayload;
+
+									// This can return false if we've issued multiple async loads to the same failing
+									// package within a transaction. The first would remove the element from the cache
+									// then all subsequent wouldn't find a map payload.
+									if (PackageCache.RemoveAndCopyValue(PackagePathAsString, MapPayload))
+									{
+										// We'll remove this later - but do an additional check that we definitely only
+										// removed a package that did not succeed (as we intended).
+										ensure(MapPayload.Result != EAsyncLoadingResult::Succeeded);
+									}
+								}
+
+								TransactionallyLoadedPackagesThatDidNotSucceed.Empty();
+							});
+						}
+					}
+
 					// Whether we cause an abort in `InCompletionDelegate`, all we are doing in the open
 					// after this call is to return from the wrapped open, which will do the right thing
 					// and continue with the abort. That is why it is safe to ignore the result from the
@@ -494,6 +531,7 @@ private:
 	};
 
 	TArray<FTransactionallyDeferredLoadPackagePayload> TransactionallyDeferredLoadPackages;
+	TArray<FString> TransactionallyLoadedPackagesThatDidNotSucceed;
 
 	struct FMapPayload final
 	{
@@ -602,7 +640,15 @@ private:
 		WrapperDelegate.BindLambda([this, PackagePathStr, CustomPackageName, CompletionDelegate = ::MoveTemp(InCompletionDelegate)](
 			const FName& Name, UPackage* const Package, EAsyncLoadingResult::Type Type)
 		{
-			if (EAsyncLoadingResult::Canceled != Type)
+			const bool bRetrying = AutoRTFM::IsRetrying();
+
+			// If we are retrying the transaction we will cache everything but cancelled loads.
+			const bool bShouldCacheIfRetrying = bRetrying && (EAsyncLoadingResult::Canceled != Type);
+
+			// But if we are not (EG. we are outside a transaction nest), we only cache succeeded loads.
+			const bool bShouldCacheIfNotRetrying = !bRetrying && (EAsyncLoadingResult::Succeeded == Type);
+
+			if (bShouldCacheIfRetrying || bShouldCacheIfNotRetrying)
 			{
 				UE::TScopeLock _(Mutex);
 
@@ -610,7 +656,7 @@ private:
 				{
 					// If we are not retrying it means we've done a package load outside of a transaction,
 					// so we want to just forget all previous loaded packages, and start from a blank array.
-					if (!AutoRTFM::IsRetrying())
+					if (!bRetrying)
 					{
 						Payload->LoadedPackages.Empty();
 					}
