Fix build with merged modules

#rnx

#Synced-CL 51330515

diff --git a/Engine/Source/Runtime/CoreUObject/Public/UObject/PerPlatformProperties.h b/Engine/Source/Runtime/CoreUObject/Public/UObject/PerPlatformProperties.h
--- a/Engine/Source/Runtime/CoreUObject/Public/UObject/PerPlatformProperties.h
+++ b/Engine/Source/Runtime/CoreUObject/Public/UObject/PerPlatformProperties.h
@@ -61,9 +61,9 @@ namespace PerPlatformProperty::Private
 #if WITH_PREVIEW_PPX_DATA
 struct PPXPreviewOverrides
 {
-	static FName PerPlatformOverride;
-	static FName PerPlatformGroupOverride;
-	static void SetPPXPreviewOverrides(FName& InPerPlatformOverride, FName& InPerPlatformGroupOverride);
+	static COREUOBJECT_API FName PerPlatformOverride;
+	static COREUOBJECT_API FName PerPlatformGroupOverride;
+	static COREUOBJECT_API void SetPPXPreviewOverrides(FName& InPerPlatformOverride, FName& InPerPlatformGroupOverride);
 };
 #endif
 
