Missed a test and forgot an `AutoRTFM::Close` in the transactionally safe loader.

#Synced-CL 51334642

diff --git a/Engine/Source/Programs/AutoRTFMTests/Private/UPackageTests.cpp b/Engine/Source/Programs/AutoRTFMTests/Private/UPackageTests.cpp
--- a/Engine/Source/Programs/AutoRTFMTests/Private/UPackageTests.cpp
+++ b/Engine/Source/Programs/AutoRTFMTests/Private/UPackageTests.cpp
@@ -814,4 +814,35 @@ TEST_CASE("UPackage.AsyncLoading")
 			REQUIRE(bHit);
 		});
 	}
+
+	SECTION("Packaged Load Fails First Time In Transaction, Second Try in Transaction")
+	{
+		FString Name(FString::Printf(TEXT("/Game/%dAutoRTFMTestPackage%d"), FPlatformProcess::GetCurrentProcessId(), __LINE__));
+
+		AutoRTFM::Transact([&]
+		{
+			bool bHit = false;
+
+			FLoadPackageAsyncDelegate CompletionDelegate;
+			CompletionDelegate.BindLambda([&](const FName&, UPackage*, EAsyncLoadingResult::Type Type) { bHit = true; REQUIRE(Type == EAsyncLoadingResult::FailedMissing); });
+
+			FlushAsyncLoading(LoadPackageAsync(FPackagePath::FromPackageNameUnchecked(*Name), NAME_None, CompletionDelegate));
+
+			REQUIRE(bHit);
+		});
+
+		FPackageScopedMaker _(Name);
+
+		AutoRTFM::Transact([&]
+		{
+			bool bHit = false;
+
+			FLoadPackageAsyncDelegate CompletionDelegate;
+			CompletionDelegate.BindLambda([&](const FName&, UPackage*, EAsyncLoadingResult::Type Type) { bHit = true; REQUIRE(Type == EAsyncLoadingResult::Succeeded); });
+
+			FlushAsyncLoading(LoadPackageAsync(FPackagePath::FromPackageNameUnchecked(*Name), NAME_None, CompletionDelegate));
+
+			REQUIRE(bHit);
+		});
+	}
 }
diff --git a/Engine/Source/Runtime/CoreUObject/Private/Serialization/TransactionallySafeAsyncLoading.cpp b/Engine/Source/Runtime/CoreUObject/Private/Serialization/TransactionallySafeAsyncLoading.cpp
--- a/Engine/Source/Runtime/CoreUObject/Private/Serialization/TransactionallySafeAsyncLoading.cpp
+++ b/Engine/Source/Runtime/CoreUObject/Private/Serialization/TransactionallySafeAsyncLoading.cpp
@@ -220,27 +220,31 @@ public:
 
 						if (bFirst)
 						{
-							AutoRTFM::OnComplete([this]
+							AutoRTFM::ETransactionStatus Status = AutoRTFM::Close([&]
 							{
-								UE::TScopeLock _(Mutex);
-
-								for (const FString& PackagePathAsString : TransactionallyLoadedPackagesThatDidNotSucceed)
+								AutoRTFM::OnComplete([this]
 								{
-									FMapPayload MapPayload;
+									UE::TScopeLock _(Mutex);
 
-									// This can return false if we've issued multiple async loads to the same failing
-									// package within a transaction. The first would remove the element from the cache
-									// then all subsequent wouldn't find a map payload.
-									if (PackageCache.RemoveAndCopyValue(PackagePathAsString, MapPayload))
+									for (const FString& PackagePathAsString : TransactionallyLoadedPackagesThatDidNotSucceed)
 									{
-										// We'll remove this later - but do an additional check that we definitely only
-										// removed a package that did not succeed (as we intended).
-										ensure(MapPayload.Result != EAsyncLoadingResult::Succeeded);
+										FMapPayload MapPayload;
+
+										// This can return false if we've issued multiple async loads to the same failing
+										// package within a transaction. The first would remove the element from the cache
+										// then all subsequent wouldn't find a map payload.
+										if (PackageCache.RemoveAndCopyValue(PackagePathAsString, MapPayload))
+										{
+											// We'll remove this later - but do an additional check that we definitely only
+											// removed a package that did not succeed (as we intended).
+											ensure(MapPayload.Result != EAsyncLoadingResult::Succeeded);
+										}
 									}
-								}
 
-								TransactionallyLoadedPackagesThatDidNotSucceed.Empty();
+									TransactionallyLoadedPackagesThatDidNotSucceed.Empty();
+								});
 							});
+							ensure(Status == AutoRTFM::ETransactionStatus::Executing);
 						}
 					}
 
