CoreUObject: Removed deprecated members from FPropertyTag

#rnx
#Synced-CL 51340251

diff --git a/Engine/Source/Runtime/CoreUObject/Private/UObject/PropertyTag.cpp b/Engine/Source/Runtime/CoreUObject/Private/UObject/PropertyTag.cpp
--- a/Engine/Source/Runtime/CoreUObject/Private/UObject/PropertyTag.cpp
+++ b/Engine/Source/Runtime/CoreUObject/Private/UObject/PropertyTag.cpp
@@ -84,9 +84,7 @@ FPropertyTag::FPropertyTag(FProperty* Property, int32 InIndex, uint8* Value)
 
 void FPropertyTag::SetProperty(FProperty* Property)
 {
-	PRAGMA_DISABLE_DEPRECATION_WARNINGS
 	Prop = Property;
-	PRAGMA_ENABLE_DEPRECATION_WARNINGS
 }
 
 void FPropertyTag::SetPropertyGuid(const FGuid& InPropertyGuid)
@@ -99,38 +97,6 @@ void FPropertyTag::SetType(UE::FPropertyTypeName InFullType)
 {
 	TypeName = InFullType;
 	Type = TypeName.GetName();
-
-	PRAGMA_DISABLE_DEPRECATION_WARNINGS
-	if (const EName* TagType = Type.ToEName(); TagType && Type.GetNumber() == NAME_NO_NUMBER_INTERNAL)
-	{
-		switch (*TagType)
-		{
-		case NAME_StructProperty:
-			StructName = TypeName.GetParameterName(0);
-			if (FName StructGuidName = TypeName.GetParameterName(1);
-				StructGuidName.IsNone() || !FGuid::Parse(StructGuidName.ToString(), StructGuid))
-			{
-				StructGuid.Invalidate();
-			}
-			break;
-		case NAME_ByteProperty:
-		case NAME_EnumProperty:
-			EnumName = TypeName.GetParameterName(0);
-			break;
-		case NAME_ArrayProperty:
-		case NAME_OptionalProperty:
-		case NAME_SetProperty:
-			InnerType = TypeName.GetParameterName(0);
-			break;
-		case NAME_MapProperty:
-			InnerType = TypeName.GetParameterName(0);
-			ValueType = TypeName.GetParameterName(1);
-			break;
-		default:
-			break;
-		}
-	}
-	PRAGMA_ENABLE_DEPRECATION_WARNINGS
 }
 
 static EPropertyTagExtension CalculatePropertyExtensionFlags(FArchive& UnderlyingArchive, FPropertyTag& Tag)
@@ -294,10 +260,6 @@ FORCENOINLINE void LoadPropertyTagNoFullType(FStructuredArchive::FSlot Slot, FPr
 				}
 			}
 			TypeBuilder.EndParameters();
-			PRAGMA_DISABLE_DEPRECATION_WARNINGS
-			Tag.StructName = StructName;
-			Tag.StructGuid = StructGuid;
-			PRAGMA_ENABLE_DEPRECATION_WARNINGS
 		}
 		// only need to serialize this for bools
 		else if (TagType == NAME_BoolProperty && !UnderlyingArchive.IsTextFormat())
@@ -330,9 +292,6 @@ FORCENOINLINE void LoadPropertyTagNoFullType(FStructuredArchive::FSlot Slot, FPr
 				AddEnumPath(TypeBuilder, EnumName);
 				TypeBuilder.EndParameters();
 			}
-			PRAGMA_DISABLE_DEPRECATION_WARNINGS
-			Tag.EnumName = EnumName;
-			PRAGMA_ENABLE_DEPRECATION_WARNINGS
 		}
 		else if (TagType == NAME_EnumProperty)
 		{
@@ -342,9 +301,6 @@ FORCENOINLINE void LoadPropertyTagNoFullType(FStructuredArchive::FSlot Slot, FPr
 			AddEnumPath(TypeBuilder, EnumName);
 			TypeBuilder.AddName(NAME_ByteProperty);
 			TypeBuilder.EndParameters();
-			PRAGMA_DISABLE_DEPRECATION_WARNINGS
-			Tag.EnumName = EnumName;
-			PRAGMA_ENABLE_DEPRECATION_WARNINGS
 		}
 		// need to serialize the InnerType for arrays
 		else if (TagType == NAME_ArrayProperty)
@@ -357,9 +313,6 @@ FORCENOINLINE void LoadPropertyTagNoFullType(FStructuredArchive::FSlot Slot, FPr
 			TypeBuilder.BeginParameters();
 			TypeBuilder.AddName(InnerType);
 			TypeBuilder.EndParameters();
-			PRAGMA_DISABLE_DEPRECATION_WARNINGS
-			Tag.InnerType = InnerType;
-			PRAGMA_ENABLE_DEPRECATION_WARNINGS
 		}
 		// need to serialize the InnerType for optionals.
 		else if (TagType == NAME_OptionalProperty)
@@ -369,9 +322,6 @@ FORCENOINLINE void LoadPropertyTagNoFullType(FStructuredArchive::FSlot Slot, FPr
 			TypeBuilder.BeginParameters();
 			TypeBuilder.AddName(InnerType);
 			TypeBuilder.EndParameters();
-			PRAGMA_DISABLE_DEPRECATION_WARNINGS
-			Tag.InnerType = InnerType;
-			PRAGMA_ENABLE_DEPRECATION_WARNINGS
 		}
 		else if (Version >= VER_UE4_PROPERTY_TAG_SET_MAP_SUPPORT)
 		{
@@ -382,9 +332,6 @@ FORCENOINLINE void LoadPropertyTagNoFullType(FStructuredArchive::FSlot Slot, FPr
 				TypeBuilder.BeginParameters();
 				TypeBuilder.AddName(InnerType);
 				TypeBuilder.EndParameters();
-				PRAGMA_DISABLE_DEPRECATION_WARNINGS
-				Tag.InnerType = InnerType;
-				PRAGMA_ENABLE_DEPRECATION_WARNINGS
 			}
 			else if (TagType == NAME_MapProperty)
 			{
@@ -396,10 +343,6 @@ FORCENOINLINE void LoadPropertyTagNoFullType(FStructuredArchive::FSlot Slot, FPr
 				TypeBuilder.AddName(InnerType);
 				TypeBuilder.AddName(ValueType);
 				TypeBuilder.EndParameters();
-				PRAGMA_DISABLE_DEPRECATION_WARNINGS
-				Tag.InnerType = InnerType;
-				Tag.ValueType = ValueType;
-				PRAGMA_ENABLE_DEPRECATION_WARNINGS
 			}
 		}
 
diff --git a/Engine/Source/Runtime/CoreUObject/Public/UObject/PropertyTag.h b/Engine/Source/Runtime/CoreUObject/Public/UObject/PropertyTag.h
--- a/Engine/Source/Runtime/CoreUObject/Public/UObject/PropertyTag.h
+++ b/Engine/Source/Runtime/CoreUObject/Public/UObject/PropertyTag.h
@@ -36,30 +36,16 @@ enum class EPropertyTagSerializeType : uint8
  */
 struct FPropertyTag
 {
-	// Transient.
-	UE_DEPRECATED(5.4, "Prop has been deprecated. Use GetProperty()/SetProperty().")
-	FProperty* Prop = nullptr;
-
-	// Variables.
 private:
+	FProperty* Prop = nullptr; // Transient
 	UE::FPropertyTypeName TypeName;
 
 public:
 	FName	Type;		// Type of property
 	FName	Name;		// Name of property.
-	UE_DEPRECATED(5.4, "StructName has been deprecated and replaced by GetType(). Use GetType().IsStruct(StructName) to test, otherwise if Type is StructProperty then StructName is GetType().GetParameterName(0).")
-	FName	StructName;	// Struct name if FStructProperty.
-	UE_DEPRECATED(5.4, "EnumName has been deprecated and replaced by GetType(). Use GetType().IsEnum(EnumName) to test, otherwise if Type is EnumProperty or ByteProperty then EnumName is GetType().GetParameterName(0).")
-	FName	EnumName;	// Enum name if FByteProperty or FEnumProperty
-	UE_DEPRECATED(5.4, "InnerName has been deprecated and replaced by GetType(). If Type is ArrayProperty/SetProperty/MapProperty/OptionalProperty then InnerType is GetType().GetParameterName(0).")
-	FName	InnerType;	// Inner type if FArrayProperty, FSetProperty, FMapProperty, or OptionalProperty
-	UE_DEPRECATED(5.4, "ValueName has been deprecated and replaced by GetType(). If Type is MapProperty then ValueType is GetType().GetParameterName(1).")
-	FName	ValueType;	// Value type if UMapProperty
 	int32	Size = 0;   // Property size.
 	int32	ArrayIndex = INDEX_NONE; // Index if an array; else 0.
 	int64	SizeOffset = INDEX_NONE; // location in stream of tag size member
-	UE_DEPRECATED(5.4, "StructGuid has been deprecated and replaced by GetType(). If Type is StructProperty then StructGuid is GetType().GetParameterName(1).")
-	FGuid	StructGuid;
 	FGuid	PropertyGuid;
 	uint8	HasPropertyGuid = 0;
 	uint8	BoolVal = 0;// a boolean property's value (never need to serialize data for bool properties except here)
@@ -72,18 +58,9 @@ public:
 	FPropertyTag();
 	UE_INTERNAL COREUOBJECT_API FPropertyTag(FProperty* Property, int32 InIndex, uint8* Value);
 
-	PRAGMA_DISABLE_DEPRECATION_WARNINGS
-	FPropertyTag(FPropertyTag&&) = default;
-	FPropertyTag(const FPropertyTag&) = default;
-	FPropertyTag& operator=(FPropertyTag&&) = default;
-	FPropertyTag& operator=(const FPropertyTag&) = default;
-	PRAGMA_ENABLE_DEPRECATION_WARNINGS
-
 	inline FProperty* GetProperty() const
 	{
-		PRAGMA_DISABLE_DEPRECATION_WARNINGS
 		return Prop;
-		PRAGMA_ENABLE_DEPRECATION_WARNINGS
 	}
 
 	void SetProperty(FProperty* Property);
