CoreUObject: Removed deprecated functions from Class.h

#rnx

#Synced-CL 51340251

diff --git a/Engine/Source/Runtime/CoreUObject/Private/UObject/Class.cpp b/Engine/Source/Runtime/CoreUObject/Private/UObject/Class.cpp
--- a/Engine/Source/Runtime/CoreUObject/Private/UObject/Class.cpp
+++ b/Engine/Source/Runtime/CoreUObject/Private/UObject/Class.cpp
@@ -154,23 +154,6 @@ COREUOBJECT_API void InitializePrivateStaticClass(
 		TClass_PrivateStaticClass->Register(TClass_StaticClassFn, PackageName, Name);
 	}
 }
-
-COREUOBJECT_API void InitializePrivateStaticClass(
-	class UClass* TClass_Super_StaticClass,
-	class UClass* TClass_PrivateStaticClass,
-	class UClass* TClass_WithinClass_StaticClass,
-	const TCHAR* PackageName,
-	const TCHAR* Name
-)
-{
-	InitializePrivateStaticClass(
-		UClass::StaticClass,
-		TClass_Super_StaticClass,
-		TClass_PrivateStaticClass,
-		TClass_WithinClass_StaticClass,
-		PackageName,
-		Name);
-}
 #endif // !UE_WITH_CONSTINIT_UOBJECT
 
 void FNativeFunctionRegistrar::RegisterFunction(class UClass* Class, const ANSICHAR* InName, FNativeFuncPtr InPointer)
diff --git a/Engine/Source/Runtime/CoreUObject/Public/UObject/Class.h b/Engine/Source/Runtime/CoreUObject/Public/UObject/Class.h
--- a/Engine/Source/Runtime/CoreUObject/Public/UObject/Class.h
+++ b/Engine/Source/Runtime/CoreUObject/Public/UObject/Class.h
@@ -1849,18 +1849,6 @@ public:
 			return GetCapabilities().HasStructuredSerializer;
 		}
 
-		UE_DEPRECATED(5.6, "Use the overload of Serialize with defaults.")
-		bool Serialize(FArchive& Ar, void* Data)
-		{
-			return Serialize(Ar, Data, nullptr, nullptr);
-		}
-
-		UE_DEPRECATED(5.6, "Use the overload of Serialize with defaults.")
-		bool Serialize(FStructuredArchive::FSlot Slot, void* Data)
-		{
-			return Serialize(Slot, Data, nullptr, nullptr);
-		}
-
 		/**
 		 * Serialize this structure.
 		 * @return true if serialization was handled, false to fall back to standard script struct serialization.
@@ -4968,15 +4956,6 @@ COREUOBJECT_API void InitializePrivateStaticClass(
 	const TCHAR* Name
 );
 
-UE_DEPRECATED(5.6, "Use the new method above and specify UClass::StaticClass as the first argument")
-COREUOBJECT_API void InitializePrivateStaticClass(
-	class UClass* TClass_Super_StaticClass,
-	class UClass* TClass_PrivateStaticClass,
-	class UClass* TClass_WithinClass_StaticClass,
-	const TCHAR* PackageName,
-	const TCHAR* Name
-	);
-
 /**
  * Helper template allocate and construct a UClass
  *
