Cleaned up some deprecated features in UnrealType.h.

diff --git a/Engine/Source/Runtime/CoreUObject/Private/UObject/UObjectGlobals.cpp b/Engine/Source/Runtime/CoreUObject/Private/UObject/UObjectGlobals.cpp
--- a/Engine/Source/Runtime/CoreUObject/Private/UObject/UObjectGlobals.cpp
+++ b/Engine/Source/Runtime/CoreUObject/Private/UObject/UObjectGlobals.cpp
@@ -5254,7 +5254,7 @@ void CollectOptionalReference(FReferenceCollector& Collector, FOptionalProperty&
 }
 
 
-// Process FObjectProperty or FObjectPtrProperty reference
+// Process FObjectProperty reference
 template<EPropertyCollectFlags CollectFlags>
 FORCEINLINE_DEBUGGABLE void CollectObjectReference(FReferenceCollector& Collector, FProperty& Property, void* Value, const UObject* Referencer)
 {
diff --git a/Engine/Source/Runtime/CoreUObject/Public/UObject/UnrealType.h b/Engine/Source/Runtime/CoreUObject/Public/UObject/UnrealType.h
--- a/Engine/Source/Runtime/CoreUObject/Public/UObject/UnrealType.h
+++ b/Engine/Source/Runtime/CoreUObject/Public/UObject/UnrealType.h
@@ -2860,7 +2860,7 @@ public:
 	 */
 	UE_API static FString GetExportPath(FTopLevelAssetPath ClassPathName, const FString& ObjectPathName);
 
-	// Helper method for sharing code with FObjectPtrProperty even though one doesn't inherit from the other
+	// Helper method for deep comparing objects
 	UE_API static bool StaticIdentical(UObject* A, UObject* B, uint32 PortFlags);
 
 	virtual UObject* LoadObjectPropertyValue(const void* PropertyValueAddress) const
@@ -3233,8 +3233,6 @@ public:
 	UE_API virtual void EmitIntrusiveOptionalReferenceInfo(UE::GC::FSchemaBuilder& Schema, int32 BaseOffset, TArray<const FStructProperty*>& EncounteredStructProps, UE::GC::FPropertyStack& DebugPath) override;
 };
 
-using FObjectPtrProperty UE_DEPRECATED(5.4, "FObjectPtrProperty is deprecated using FObjectProperty instead.")  = FObjectProperty;
-
 //
 // Describes a reference variable to another object which may be nil, and may turn nil at any point
 //
@@ -3529,7 +3527,7 @@ public:
 #endif // USE_CIRCULAR_DEPENDENCY_LOAD_DEFERRING
 };
 
-using FClassPtrProperty UE_DEPRECATED(5.4, "FClassPtrProperty is deprecated use FClassProperty instead.") = FClassProperty;
+
 /*-----------------------------------------------------------------------------
 	FSoftClassProperty.
 -----------------------------------------------------------------------------*/
@@ -4716,12 +4714,6 @@ struct TScriptContainerIterator
 		return LogicalIndex;
 	}
 
-	UE_DEPRECATED(5.4, "Use Iterator directly, GetInternalIndex or GetLogicalIndex instead.")
-	int32 operator*() const
-	{
-		return InternalIndex;
-	}
-
 private:
 	const ContainerType& Container;
 	int32 InternalIndex = INDEX_NONE;
diff --git a/Engine/Source/Runtime/CoreUObject/Tests/ObjectPropertyTest.cpp b/Engine/Source/Runtime/CoreUObject/Tests/ObjectPropertyTest.cpp
--- a/Engine/Source/Runtime/CoreUObject/Tests/ObjectPropertyTest.cpp
+++ b/Engine/Source/Runtime/CoreUObject/Tests/ObjectPropertyTest.cpp
@@ -282,13 +282,9 @@ static void TestSerializeItem(FName ObjectName)
 	}
 }
 
-TEST_CASE("UE::CoreUObject::FObjectPtrProperty::StaticSerializeItem")
-{
-	TestSerializeItem<UObjectPtrTestClassWithRef>(TEXT("Object1"));
-}
-
 TEST_CASE("UE::CoreUObject::FObjectProperty::StaticSerializeItem")
 {
+	TestSerializeItem<UObjectPtrTestClassWithRef>(TEXT("Object1"));
 	TestSerializeItem<UObjectWithRawProperty>(TEXT("Object2"));
 }
 
