CoreUObject: Fixed a potential FPropertyTag issue when saving a static array with a property guid

#Synced-CL 50967057

diff --git a/Engine/Source/Runtime/CoreUObject/Private/UObject/PropertyTag.cpp b/Engine/Source/Runtime/CoreUObject/Private/UObject/PropertyTag.cpp
--- a/Engine/Source/Runtime/CoreUObject/Private/UObject/PropertyTag.cpp
+++ b/Engine/Source/Runtime/CoreUObject/Private/UObject/PropertyTag.cpp
@@ -467,7 +467,7 @@ void operator<<(FStructuredArchive::FSlot Slot, FPropertyTag& Tag)
 		}
 		if (Tag.HasPropertyGuid)
 		{
-			PropertyTagFlags = EPropertyTagFlags::HasPropertyGuid;
+			PropertyTagFlags |= EPropertyTagFlags::HasPropertyGuid;
 		}
 		if (PropertyTagExtensions != EPropertyTagExtension::NoExtension)
 		{
