[Backout] - CL50960451
#rnx
Original CL Desc
-----------------------------------------------------------------
[LinkerManager]
- Replace a global flush by flushing only the required linker

diff --git a/Engine/Source/Runtime/CoreUObject/Private/UObject/LinkerManager.cpp b/Engine/Source/Runtime/CoreUObject/Private/UObject/LinkerManager.cpp
--- a/Engine/Source/Runtime/CoreUObject/Private/UObject/LinkerManager.cpp
+++ b/Engine/Source/Runtime/CoreUObject/Private/UObject/LinkerManager.cpp
@@ -95,13 +95,15 @@ bool FLinkerManager::Exec_Dev(class UWorld* InWorld, const TCHAR* Cmd, FOutputDe
 	return false;
 }
 
-extern void ConditionalFlushAsyncLoadingForLinkers(TConstArrayView<FLinkerLoad*> InLinkers);
-
 void FLinkerManager::ResetLinkerExports(UPackage* InPackage)
 {
 	if (FLinkerLoad* LinkerToReset = FLinkerLoad::FindExistingLinkerForPackage(InPackage))
 	{
-		ConditionalFlushAsyncLoadingForLinkers({LinkerToReset});
+		// if the linker owner thread is not the main thread, we need to flush async loading (todo: for that package) before we can reset the linker
+		if (LinkerToReset->GetOwnerThreadId() != GGameThreadId)
+		{
+			FlushAsyncLoading();
+		}
 		LinkerToReset->DetachExports();
 	}
 }
