Dstm: Skip destroyed objects in BeginMultiServerCommit instead of ensure

#Synced-CL 51388279

diff --git a/Engine/Source/Runtime/CoreUObject/Private/UObject/RemoteObjectTransfer.cpp b/Engine/Source/Runtime/CoreUObject/Private/UObject/RemoteObjectTransfer.cpp
--- a/Engine/Source/Runtime/CoreUObject/Private/UObject/RemoteObjectTransfer.cpp
+++ b/Engine/Source/Runtime/CoreUObject/Private/UObject/RemoteObjectTransfer.cpp
@@ -596,14 +596,15 @@ public:
 				continue;
 			}
 
-			// if we used it this run and we're about to commit, we expect that the object is local
-			check(Handle::GetResidence(UsedObjectId) != EResidence::Remote);
-
+			// UsedObject may have been destroyed since it was touched, so we may not be able to find it.
 			UObject* UsedObject = StaticFindObjectFastInternal(UsedObjectId);
-			if (!ensure(UsedObject))
+			if (!UsedObject)
 			{
 				continue;
 			}
+
+			// If we used it this run and we're about to commit, we expect that the object is local
+			check(Handle::GetResidence(UsedObjectId) != EResidence::Remote);
 			FRemoteServerId OwnerServerId = Handle::GetOwnerServerId(UsedObject);
 
 			UObject* RootObject = FindCanonicalRootObjectForSerialization(UsedObject);
