RemoteObjects:  GetNameSafe override for TObjectPtr which can avoid migrations.

#rnx
#Synced-CL 51377531

diff --git a/Engine/Source/Runtime/CoreUObject/Public/UObject/UObjectBaseUtility.h b/Engine/Source/Runtime/CoreUObject/Public/UObject/UObjectBaseUtility.h
--- a/Engine/Source/Runtime/CoreUObject/Public/UObject/UObjectBaseUtility.h
+++ b/Engine/Source/Runtime/CoreUObject/Public/UObject/UObjectBaseUtility.h
@@ -916,6 +916,26 @@ inline FString GetNameSafe(const UObjectBaseUtility* Object)
 	}
 }
 
+#if UE_WITH_REMOTE_OBJECT_HANDLE
+/**
+ * Returns the name of this object (with no path information). This version can avoid migrations for remote objects.
+ * @param Object object to retrieve the name for; null gives "None".
+ * @return Name of the object.
+ */
+template<class T UE_REQUIRES(TIsDerivedFrom<T, UObjectBase>::Value)>
+inline FString GetNameSafe(const TObjectPtr<T>& ObjectPtr)
+{
+	if (ObjectPtr)
+	{
+		return ObjectPtr.GetName();
+	}
+	else
+	{
+		return TEXT("None");
+	}
+}
+#endif
+
 /**
  * Returns the path name of this object.
  * @param Object object to retrieve the path name for; null gives "None".
