RemoteExecutor logging changes to make aborts and retries more obvious.

#rnx

#Synced-CL 51377531

diff --git a/Engine/Source/Runtime/CoreUObject/Private/UObject/RemoteExecutor.cpp b/Engine/Source/Runtime/CoreUObject/Private/UObject/RemoteExecutor.cpp
--- a/Engine/Source/Runtime/CoreUObject/Private/UObject/RemoteExecutor.cpp
+++ b/Engine/Source/Runtime/CoreUObject/Private/UObject/RemoteExecutor.cpp
@@ -309,7 +309,7 @@ public:
 
 				if (bPrintStall)
 				{
-					UE_LOG(LogRemoteExec, Verbose, TEXT("ExecutePendingWork[%d] Waiting on handling remote multi server commit %s %s"),
+					UE_LOG(LogRemoteExec, Verbose, TEXT("ExecutePendingWork[LocalIter %d] Waiting on handling remote multi server commit %s %s"),
 						LocalIterationNumber,
 						*ActiveRemoteMultiServerCommitServerId.ToString(),
 						*ActiveRemoteMultiServerCommitRequestId.ToString());
@@ -353,8 +353,8 @@ public:
 				ConditionallyPumpNetwork(bForcePumpNetwork);
 			}
 
-			ExecutePendingWork(PendingMigrationWorks, LocalIterationNumber, LastStallPrintTime);
-			ExecutePendingWork(PendingWorks, LocalIterationNumber, LastStallPrintTime);
+			ExecutePendingWork(PendingMigrationWorks, LastStallPrintTime);
+			ExecutePendingWork(PendingWorks, LastStallPrintTime);
 		}
 
 		PendingWorkLists.Empty();
@@ -364,18 +364,22 @@ public:
 
 private:
 
-	void ExecutePendingWork(TArray<FRemoteExecutorWork>& WorksToExecute, int32 LocalIterationNumber, double& LastStallPrintTime)
+	void ExecutePendingWork(TArray<FRemoteExecutorWork>& WorksToExecute, double& LastStallPrintTime)
 	{
-		for (int32 PendingWorkIndex = 0; PendingWorkIndex < WorksToExecute.Num(); PendingWorkIndex++)
+		if (UE_LOG_ACTIVE(LogRemoteExec, VeryVerbose))
 		{
-			FRemoteExecutorWork& PendingWork = WorksToExecute[PendingWorkIndex];
-
-			UE_LOG(LogRemoteExec, VeryVerbose, TEXT("ExecutePendingWork[%d] : PendingWork[%d] '%s' RequestId %s %s"),
-				LocalIterationNumber,
-				PendingWorkIndex,
-				*PendingWork.Name.ToString(),
-				*PendingWork.RequestId.ToString(),
-				*PendingWork.Priority.ToString());
+			UE_CLOGF(WorksToExecute.Num() > 0, LogRemoteExec, VeryVerbose, "ExecutePendingWork Entries: ");
+
+			for (int32 PendingWorkIndex = 0; PendingWorkIndex < WorksToExecute.Num(); PendingWorkIndex++)
+			{
+				FRemoteExecutorWork& PendingWork = WorksToExecute[PendingWorkIndex];
+
+				UE_LOGF(LogRemoteExec, VeryVerbose, "   PendingWork[%d] '%ls' RequestId %ls %ls",
+					   PendingWorkIndex,
+					   *PendingWork.Name.ToString(),
+					   *PendingWork.RequestId.ToString(),
+					   *PendingWork.Priority.ToString());
+			}
 		}
 
 		bool bPrintStall = false;
@@ -387,8 +391,7 @@ private:
 			TGuardValue<FRemoteExecutorWork*> ExecutingWorkSave(ExecutingWork, &WorksToExecute[PendingWorkIndex]);
 			check(ExecutingWork);
 
-			UE_LOG(LogRemoteExec, VeryVerbose, TEXT("ExecutePendingWork[%d,%s] : Executing request %s %s --"),
-				LocalIterationNumber,
+			UE_LOGF(LogRemoteExec, VeryVerbose, "ExecutePendingWork[%ls] : Executing request %ls %ls --",
 				*ExecutingWork->RequestId.ToString(),
 				*ExecutingWork->Name.ToString(),
 				*ExecutingWork->Priority.ToString()); // -V522
@@ -403,8 +406,7 @@ private:
 					for (int32 SubsystemIndex = 0; SubsystemIndex < Subsystems.Num(); SubsystemIndex++)
 					{
 						FRemoteSubsystemBase* Subsystem = Subsystems[SubsystemIndex];
-						UE_LOG(LogRemoteExec, VeryVerbose, TEXT("ExecutePendingWork[%d,%s] : Creating request[%s]"),
-							LocalIterationNumber,
+						UE_LOGF(LogRemoteExec, VeryVerbose, "ExecutePendingWork[%ls] : Creating request[%ls]",
 							*ExecutingWork->RequestId.ToString(),
 							Subsystem->NameForDebug());
 						Subsystem->CreateRequest(ExecutingWork->RequestId, ExecutingWork->Priority);
@@ -449,8 +451,7 @@ private:
 					{
 						if (bPrintStall)
 						{
-							UE_LOG(LogRemoteExec, VeryVerbose, TEXT("ExecutePendingWork[%d,%s] : subsystem %s not ready..."),
-								LocalIterationNumber,
+							UE_LOGF(LogRemoteExec, VeryVerbose, "ExecutePendingWork[%ls] : subsystem %ls not ready...",
 								*ExecutingWork->RequestId.ToString(),
 								Subsystem->NameForDebug());
 						}
@@ -461,8 +462,7 @@ private:
 				{
 					if (bPrintStall)
 					{
-						UE_LOG(LogRemoteExec, VeryVerbose, TEXT("ExecutePendingWork[%d,%s] : %d of %d subsystems not ready..."),
-							LocalIterationNumber,
+						UE_LOGF(LogRemoteExec, VeryVerbose, "ExecutePendingWork[%ls] : %d of %d subsystems not ready...",
 							*ExecutingWork->RequestId.ToString(),
 							Subsystems.Num() - SubsystemReadyCount, Subsystems.Num());
 					}
@@ -478,7 +478,7 @@ private:
 				UE::RemoteExecutor::OnTransactionStartingDelegate.Broadcast(RequestId, ExecutingWork->Name);
 				AutoRTFM::ETransactionResult WorkTransactResult =
 					AutoRTFM::Transact(
-						[this, LocalIterationNumber]()
+						[this]()
 						{
 							UE_AUTORTFM_OPEN
 							{
@@ -494,8 +494,6 @@ private:
 							const uint32 ExecutionAttempts = ExecutingWork->ExecutionAttempts;
 							UE_AUTORTFM_ONABORT(ExecutionAttempts, this)
 							{
-								UE_LOG(LogRemoteExec, VeryVerbose, TEXT("Transaction %s Aborted: %s: %s"), *ExecutingWork->Name.ToString(), EnumToString(AbortReason), *AbortReasonDescription);
-
 								constexpr uint32 MaxExecutionAttemptsBeforeLivelockWarning = 500;
 								if ((ExecutionAttempts >= MaxExecutionAttemptsBeforeLivelockWarning) && (ExecutionAttempts % MaxExecutionAttemptsBeforeLivelockWarning) == 0)
 								{
@@ -510,7 +508,7 @@ private:
 							// Did we do something that should late-abort this transaction?
 							if (ExecutingWork->NumAbortsRequested > 0)
 							{
-								UE_LOG(LogRemoteExec, Warning, TEXT("We ended a transaction marked for abort. We expected to abort before this point (%d abort requests)."), ExecutingWork->NumAbortsRequested);
+								UE_LOGF(LogRemoteExec, Warning, "We ended a transaction marked for abort. We expected to abort before this point (%d abort requests).", ExecutingWork->NumAbortsRequested);
 								AutoRTFM::AbortTransaction();
 							}
 							else
@@ -526,8 +524,7 @@ private:
 									// first we need to send borrowed objects back
 									// to their owner and ask them if they are ready
 									// to commit or not...
-									UE_LOG(LogRemoteExec, VeryVerbose, TEXT("ExecutePendingWork[%d,%s] : STARTING multi-server commit because %s"),
-										LocalIterationNumber,
+									UE_LOGF(LogRemoteExec, VeryVerbose, "ExecutePendingWork[%ls] : STARTING multi-server commit because %ls",
 										*ExecutingWork->RequestId.ToString(),
 										*ExecutingWork->RequiresMultiServerCommitReason);
 
@@ -545,8 +542,7 @@ private:
 								TArray<FRemoteServerId> MultiServerCommitRemoteServers;
 								for (int32 SubsystemIndex = 0; SubsystemIndex < Subsystems.Num(); SubsystemIndex++)
 								{
-									UE_LOG(LogRemoteExec, VeryVerbose, TEXT("ExecutePendingWork[%d,%s] : BeginMultiServerCommit subsystem %d..."),
-										LocalIterationNumber,
+									UE_LOGF(LogRemoteExec, VeryVerbose, "ExecutePendingWork[%ls] : BeginMultiServerCommit subsystem %d...",
 										*ExecutingWork->RequestId.ToString(),
 										SubsystemIndex);
 
@@ -554,15 +550,13 @@ private:
 									Subsystem->BeginMultiServerCommit(MultiServerCommitRemoteServers);
 								}
 
-								UE_LOG(LogRemoteExec, VeryVerbose, TEXT("ExecutePendingWork[%d,%s] : BeginMultiServerCommit DONE (%d servers)"),
-									LocalIterationNumber,
+								UE_LOGF(LogRemoteExec, VeryVerbose, "ExecutePendingWork[%ls] : BeginMultiServerCommit DONE (%d servers)",
 									*ExecutingWork->RequestId.ToString(),
 									MultiServerCommitRemoteServers.Num());
 
 								for (FRemoteServerId ServerId : MultiServerCommitRemoteServers)
 								{
-									UE_LOG(LogRemoteExec, VeryVerbose, TEXT("ExecutePendingWork[%d,%s] : BeginMultiServerCommit server %s"),
-										LocalIterationNumber,
+									UE_LOGF(LogRemoteExec, VeryVerbose, "ExecutePendingWork[%ls] : BeginMultiServerCommit server %ls",
 										*ExecutingWork->RequestId.ToString(),
 										*ServerId.ToString());
 								}
@@ -576,8 +570,7 @@ private:
 								// ask each subsystem to send any necessary data as part of this commit
 								for (int32 SubsystemIndex = 0; SubsystemIndex < Subsystems.Num(); SubsystemIndex++)
 								{
-									UE_LOG(LogRemoteExec, VeryVerbose, TEXT("ExecutePendingWork[%d,%s] : ExecuteMultiServerCommit subsystem %d..."),
-										LocalIterationNumber,
+									UE_LOGF(LogRemoteExec, VeryVerbose, "ExecutePendingWork[%ls] : ExecuteMultiServerCommit subsystem %d...",
 										*ExecutingWork->RequestId.ToString(),
 										SubsystemIndex);
 
@@ -601,8 +594,7 @@ private:
 									if ((Now - LastPrintTime) > 1.0)
 									{
 										LastPrintTime = Now;
-										UE_LOG(LogRemoteExec, VeryVerbose, TEXT("ExecutePendingWork[%d,%s] : WAITING for multi-server commit..."),
-											LocalIterationNumber,
+										UE_LOGF(LogRemoteExec, VeryVerbose, "ExecutePendingWork[%ls] : WAITING for multi-server commit...",
 											*ExecutingWork->RequestId.ToString());
 									}
 
@@ -623,8 +615,7 @@ private:
 
 											if (ServerId != ReadyServerId)
 											{
-												UE_LOG(LogRemoteExec, VeryVerbose, TEXT("ExecutePendingWork[%d,%s] : Multi-server commit expected server %s but got %s"),
-													LocalIterationNumber,
+												UE_LOGF(LogRemoteExec, VeryVerbose, "ExecutePendingWork[%ls] : Multi-server commit expected server %ls but got %ls",
 													*ExecutingWork->RequestId.ToString(),
 													*ServerId.ToString(),
 													*ReadyServerId.ToString());
@@ -635,8 +626,7 @@ private:
 
 										if (bMultiServerCommitRequiresAbort)
 										{
-											UE_LOG(LogRemoteExec, Verbose, TEXT("ExecutePendingWork[%d,%s] : Multi-server commit ALL servers READY, but we are flagged with ABORT"),
-												LocalIterationNumber,
+											UE_LOGF(LogRemoteExec, Verbose, "ExecutePendingWork[%ls] : Multi-server commit ALL servers READY, but we are flagged with ABORT",
 												*ExecutingWork->RequestId.ToString());
 
 											UE_AUTORTFM_OPEN
@@ -647,12 +637,12 @@ private:
 													Subsystem->AbortMultiServerCommit();
 												}
 
-											// aborted, tell all servers to abandon this commit
-											UE::RemoteExecutor::AbandonMultiServerCommitDelegate.ExecuteIfBound(ExecutingWork->RequestId, MultiServerCommitRemoteServers);
+												// aborted, tell all servers to abandon this commit
+												UE::RemoteExecutor::AbandonMultiServerCommitDelegate.ExecuteIfBound(ExecutingWork->RequestId, MultiServerCommitRemoteServers);
 
-											MultiServerCommitRequestId = FRemoteTransactionId::Invalid();
-											MultiServerCommitReadyServers.Reset();
-											bMultiServerCommitRequiresAbort = false;
+												MultiServerCommitRequestId = FRemoteTransactionId::Invalid();
+												MultiServerCommitReadyServers.Reset();
+												bMultiServerCommitRequiresAbort = false;
 											};
 
 											using namespace UE::RemoteExecutor;
@@ -663,8 +653,7 @@ private:
 										else
 										{
 											// the commit was accepted by every server, notify each subsystem that we are committing
-											UE_LOG(LogRemoteExec, VeryVerbose, TEXT("ExecutePendingWork[%d,%s] : Multi-server commit ALL servers READY, COMMITTING"),
-												LocalIterationNumber,
+											UE_LOGF(LogRemoteExec, VeryVerbose, "ExecutePendingWork[%ls] : Multi-server commit ALL servers READY, COMMITTING",
 												*ExecutingWork->RequestId.ToString());
 
 											for (int32 SubsystemIndex = 0; SubsystemIndex < Subsystems.Num(); SubsystemIndex++)
@@ -683,9 +672,7 @@ private:
 									};
 								}
 
-								UE_LOG(LogRemoteExec, VeryVerbose, TEXT("ExecutePendingWork[%d,%s] : DONE with multi-server commit"),
-									LocalIterationNumber,
-									*ExecutingWork->RequestId.ToString());
+								UE_LOGF(LogRemoteExec, VeryVerbose, "ExecutePendingWork[%ls] : DONE with multi-server commit", *ExecutingWork->RequestId.ToString());
 
 								UE_AUTORTFM_OPEN
 								{
@@ -709,34 +696,21 @@ private:
 				{
 					using namespace UE::RemoteExecutor;
 
-					bWorkTransactionAborted = true;
+					UE_LOGF(LogRemoteExec, Verbose, "ExecutePendingWork[%ls] (%ls): Attempt %ux ABORTED. Reason: %ls (%ls)",
+							*ExecutingWork->RequestId.ToString(),
+							*ExecutingWork->Name.ToString(),
+							ExecutingWork->ExecutionAttempts,
+							EnumToString(AbortReason),
+							*AbortReasonDescription);
 
+					bWorkTransactionAborted = true;
 					OnTransactionAbortedDelegate.Broadcast(RequestId, ExecutingWork->ExecutionAttempts, AbortReasonDescription);
 
-					if (AbortReason == ERemoteExecutorAbortReason::RequiresDependencies) // -V547
-					{
-						UE_LOG(LogRemoteExec, Verbose, TEXT("ExecutePendingWork[%d,%s] : Work ABORTED for required dependency '%s'"),
-							LocalIterationNumber,
-							*ExecutingWork->RequestId.ToString(),
-							*AbortReasonDescription);
-					}
-					else if (AbortReason == ERemoteExecutorAbortReason::AbandonWork) // -V547
+					if (AbortReason == ERemoteExecutorAbortReason::AbandonWork) // -V547
 					{
-						UE_LOG(LogRemoteExec, Verbose, TEXT("ExecutePendingWork[%d,%s] : Work ABORTED for ABANDONMENT '%s'"),
-							LocalIterationNumber,
-							*ExecutingWork->RequestId.ToString(),
-							*AbortReasonDescription);
 						bWorkComplete = true;
 					}
-					else
-					{
-						UE_LOG(LogRemoteExec, Verbose, TEXT("ExecutePendingWork[%d,%s] : Work ABORTED for unknown reason '%s'"),
-							LocalIterationNumber,
-							*ExecutingWork->RequestId.ToString(),
-							*AbortReasonDescription);
-					}
-
-					if (!bWorkComplete)
+					else 
 					{
 						for (int32 SubsystemIndex = 0; SubsystemIndex < Subsystems.Num(); SubsystemIndex++)
 						{
@@ -749,10 +723,15 @@ private:
 				{
 					check(AbortReason == UE::RemoteExecutor::ERemoteExecutorAbortReason::Unspecified);
 
-					UE_LOG(LogRemoteExec, VeryVerbose, TEXT("ExecutePendingWork[%d,%s] : Work COMPLETED after %u attempts"),
-						LocalIterationNumber,
-						*ExecutingWork->RequestId.ToString(),
-						ExecutingWork->ExecutionAttempts);
+					UE_CLOGF(ExecutingWork->ExecutionAttempts > 1, LogRemoteExec, Verbose, "ExecutePendingWork[%ls] (%ls): COMPLETED after %u attempts",
+							*ExecutingWork->RequestId.ToString(),
+							*ExecutingWork->Name.ToString(),
+							ExecutingWork->ExecutionAttempts);
+
+					UE_CLOGF(ExecutingWork->ExecutionAttempts == 1, LogRemoteExec, VeryVerbose, "ExecutePendingWork[%ls] (%ls): COMPLETED (first attempt)",
+							 *ExecutingWork->RequestId.ToString(),
+							 *ExecutingWork->Name.ToString());
+
 					UE::RemoteExecutor::OnTransactionCompletedDelegate.Broadcast(RequestId, ExecutingWork->ExecutionAttempts);
 
 					bWorkComplete = true;
@@ -763,8 +742,7 @@ private:
 					for (int32 SubsystemIndex = 0; SubsystemIndex < Subsystems.Num(); SubsystemIndex++)
 					{
 						FRemoteSubsystemBase* Subsystem = Subsystems[SubsystemIndex];
-						UE_LOG(LogRemoteExec, VeryVerbose, TEXT("ExecutePendingWork[%d,%s] : End request[%s]"),
-							LocalIterationNumber,
+						UE_LOGF(LogRemoteExec, VeryVerbose, "ExecutePendingWork[%ls] : End request[%ls]",
 							*ExecutingWork->RequestId.ToString(),
 							Subsystem->NameForDebug());
 						Subsystem->EndRequest(!bWorkTransactionAborted);
