Fixed typo for "reinstantiation"

#Synced-CL 51467135

diff --git a/Engine/Source/Editor/UnrealEd/Private/EditorTransaction.cpp b/Engine/Source/Editor/UnrealEd/Private/EditorTransaction.cpp
--- a/Engine/Source/Editor/UnrealEd/Private/EditorTransaction.cpp
+++ b/Engine/Source/Editor/UnrealEd/Private/EditorTransaction.cpp
@@ -1224,7 +1224,7 @@ void UTransBuffer::OnObjectsReinstanced(const TMap<UObject*, UObject*>& OldToNew
 		return;
 	}
 
-	// Reinstanciation coming from the loading thread only concerns objects still being loaded and
+	// Reinstantiation coming from the loading thread only concerns objects still being loaded and
 	// shouldn't be referenced in any way by an existing transaction.
 	if (IsInActualLoadingThread())
 	{
diff --git a/Engine/Source/Runtime/CoreUObject/Private/StructUtils/InstancedStruct.cpp b/Engine/Source/Runtime/CoreUObject/Private/StructUtils/InstancedStruct.cpp
--- a/Engine/Source/Runtime/CoreUObject/Private/StructUtils/InstancedStruct.cpp
+++ b/Engine/Source/Runtime/CoreUObject/Private/StructUtils/InstancedStruct.cpp
@@ -209,10 +209,10 @@ bool FInstancedStruct::Serialize(FArchive& Ar, UStruct* DefaultsStruct, const vo
 		if (SerializedScriptStruct)
 		{
 #if WITH_EDITOR
-			// During reinstanciation, the serialized struct can be the the duplicated struct.
+			// During reinstantiation, the serialized struct can be the the duplicated struct.
 			//Replace it with the new version.
 			if (const UE::StructUtils::FStructReinstancer* Reinstancer = UE::StructUtils::FStructReinstancer::GetInstance();
-				Reinstancer && Reinstancer->IsSerializingForReinstanciation())
+				Reinstancer && Reinstancer->IsSerializingForReinstantiation())
 			{
 				if (const UScriptStruct* CompiledStruct = Reinstancer->GetCompiledReinstantingStruct(SerializedScriptStruct))
 				{
@@ -483,7 +483,7 @@ bool FInstancedStruct::Identical(const FInstancedStruct* Other, uint32 PortFlags
 		// Struct can't be identical when re-instancing. The struct has to serialize to get the new values.
 		if (const UE::StructUtils::FStructReinstancer* Reinstancer = UE::StructUtils::FStructReinstancer::GetInstance())
 		{
-			if (Reinstancer->IsSerializingForReinstanciation()
+			if (Reinstancer->IsSerializingForReinstantiation()
 				&& Reinstancer->IsReinstanting(StructTypePtr))
 			{
 				return false;
@@ -505,7 +505,7 @@ void FInstancedStruct::AddStructReferencedObjects(FReferenceCollector& Collector
 	{
 		if (ScriptStruct && Reinstancer->IsReinstanting(ScriptStruct))
 		{
-			Reinstancer->MarkAsRequiresReinstanciation();
+			Reinstancer->MarkAsRequiresReinstantiation();
 			// It is not necessary to change the struct to a duplicated one, but it helps with tracking mistakes and prevents memory stomps.
 			if (const UScriptStruct* Duplicated = Reinstancer->GetDuplicatedReinstantingStruct(ScriptStruct))
 			{
diff --git a/Engine/Source/Runtime/CoreUObject/Private/StructUtils/InstancedStructContainer.cpp b/Engine/Source/Runtime/CoreUObject/Private/StructUtils/InstancedStructContainer.cpp
--- a/Engine/Source/Runtime/CoreUObject/Private/StructUtils/InstancedStructContainer.cpp
+++ b/Engine/Source/Runtime/CoreUObject/Private/StructUtils/InstancedStructContainer.cpp
@@ -384,7 +384,7 @@ void FInstancedStructContainer::AddStructReferencedObjects(FReferenceCollector&
 			TObjectPtr<const UScriptStruct>& ScriptStruct = GetItem(Index).ScriptStruct;
 			if (ScriptStruct && Reinstancer->IsReinstanting(ScriptStruct))
 			{
-				Reinstancer->MarkAsRequiresReinstanciation();
+				Reinstancer->MarkAsRequiresReinstantiation();
 				// It is not necessary to change the struct to a duplicated one, but it helps with tracking mistakes and prevents memory stomps.
 				if (const UScriptStruct* Duplicated = Reinstancer->GetDuplicatedReinstantingStruct(ScriptStruct))
 				{
@@ -452,7 +452,7 @@ bool FInstancedStructContainer::Identical(const FInstancedStructContainer* Other
 #if WITH_EDITOR
 			// Struct can't be identical when re-instancing. The struct has to serialize to get the new values.
 			if (Reinstancer
-				&& Reinstancer->IsSerializingForReinstanciation()
+				&& Reinstancer->IsSerializingForReinstantiation()
 				&& Reinstancer->IsReinstanting(Item.ScriptStruct))
 			{
 				bResult = false;
@@ -581,9 +581,9 @@ bool FInstancedStructContainer::Serialize(FArchive& Ar, UStruct* DefaultsStruct,
 				if (SerializedScriptStruct)
 				{
 #if WITH_EDITOR
-					// During reinstanciation, the serialized struct can be the the duplicated struct.
+					// During reinstantiation, the serialized struct can be the the duplicated struct.
 					//Replace it with the new version.
-					if (Reinstancer && Reinstancer->IsSerializingForReinstanciation())
+					if (Reinstancer && Reinstancer->IsSerializingForReinstantiation())
 					{
 						if (const UScriptStruct* CompiledStruct = Reinstancer->GetCompiledReinstantingStruct(SerializedScriptStruct))
 						{
diff --git a/Engine/Source/Runtime/CoreUObject/Private/StructUtils/PropertyBag.cpp b/Engine/Source/Runtime/CoreUObject/Private/StructUtils/PropertyBag.cpp
--- a/Engine/Source/Runtime/CoreUObject/Private/StructUtils/PropertyBag.cpp
+++ b/Engine/Source/Runtime/CoreUObject/Private/StructUtils/PropertyBag.cpp
@@ -3213,7 +3213,7 @@ bool FInstancedPropertyBag::Identical(const FInstancedPropertyBag* Other, const
 #if WITH_EDITOR
 	// Struct can't be identical when re-instancing. The struct has to serialize to get the new values.
 	UE::StructUtils::FStructReinstancer* Reinstancer = UE::StructUtils::FStructReinstancer::GetInstance();
-	if (Other && Reinstancer && Reinstancer->IsSerializingForReinstanciation())
+	if (Other && Reinstancer && Reinstancer->IsSerializingForReinstantiation())
 	{
 		if (const UPropertyBag* BagStruct = GetPropertyBagStruct())
 		{
@@ -3291,7 +3291,7 @@ bool FInstancedPropertyBag::Serialize(FArchive& Ar)
 					if (Reinstancer && PropDesc.ValueType == EPropertyBagPropertyType::Struct)
 					{
 						if (const UScriptStruct* ValueStruct = Cast<const UScriptStruct>(PropDesc.ValueTypeObject);
-							ValueStruct && Reinstancer->IsSerializingForReinstanciation())
+							ValueStruct && Reinstancer->IsSerializingForReinstantiation())
 						{
 							if (const UScriptStruct* CompiledStruct = Reinstancer->GetCompiledReinstantingStruct(ValueStruct))
 							{
@@ -3308,7 +3308,7 @@ bool FInstancedPropertyBag::Serialize(FArchive& Ar)
 					if (Reinstancer && PropDesc.KeyType == EPropertyBagPropertyType::Struct)
 					{
 						if (const UScriptStruct* ValueStruct = Cast<const UScriptStruct>(PropDesc.KeyTypeObject);
-							ValueStruct && Reinstancer->IsSerializingForReinstanciation())
+							ValueStruct && Reinstancer->IsSerializingForReinstantiation())
 						{
 							if (const UScriptStruct* CompiledStruct = Reinstancer->GetCompiledReinstantingStruct(ValueStruct))
 							{
@@ -3371,7 +3371,7 @@ bool FInstancedPropertyBag::Serialize(FArchive& Ar)
 			// Destroy the memory. Making sure we are not using the memory with an invalid struct.
 			if (UE::StructUtils::FStructReinstancer* Reinstancer = UE::StructUtils::FStructReinstancer::GetInstance())
 			{
-				if (Ar.IsSaving() && Reinstancer->IsSerializingForReinstanciation())
+				if (Ar.IsSaving() && Reinstancer->IsSerializingForReinstantiation())
 				{
 					for (const FPropertyBagPropertyDesc& Desc : PropertyDescs)
 					{
@@ -3419,12 +3419,12 @@ void FInstancedPropertyBag::AddStructReferencedObjects(FReferenceCollector& Coll
 	{
 		if (UPropertyBag* BagStruct = Cast<UPropertyBag>(const_cast<UScriptStruct*>(Value.GetScriptStruct())))
 		{
-			auto MarkAsRequiresReinstanciation = [Reinstancer](TObjectPtr<const UObject>& TypeObject)
+			auto MarkAsRequiresReinstantiation = [Reinstancer](TObjectPtr<const UObject>& TypeObject)
 				{
 					const UScriptStruct* Struct = Cast<const UScriptStruct>(TypeObject);
 					if (Struct && Reinstancer->IsReinstanting(Struct))
 					{
-						Reinstancer->MarkAsRequiresReinstanciation();
+						Reinstancer->MarkAsRequiresReinstantiation();
 						// It is not necessary to change the struct to a duplicated one, but it helps with tracking mistakes and prevents memory stomps.
 						if (const UScriptStruct* Duplicated = Reinstancer->GetDuplicatedReinstantingStruct(Struct))
 						{
@@ -3437,11 +3437,11 @@ void FInstancedPropertyBag::AddStructReferencedObjects(FReferenceCollector& Coll
 			{
 				if (Desc.ValueType == EPropertyBagPropertyType::Struct)
 				{
-					MarkAsRequiresReinstanciation(Desc.ValueTypeObject);
+					MarkAsRequiresReinstantiation(Desc.ValueTypeObject);
 				}
 				if (Desc.KeyType == EPropertyBagPropertyType::Struct)
 				{
-					MarkAsRequiresReinstanciation(Desc.KeyTypeObject);
+					MarkAsRequiresReinstantiation(Desc.KeyTypeObject);
 				}
 			}
 		}
diff --git a/Engine/Source/Runtime/CoreUObject/Private/StructUtils/StructReinstancer.cpp b/Engine/Source/Runtime/CoreUObject/Private/StructUtils/StructReinstancer.cpp
--- a/Engine/Source/Runtime/CoreUObject/Private/StructUtils/StructReinstancer.cpp
+++ b/Engine/Source/Runtime/CoreUObject/Private/StructUtils/StructReinstancer.cpp
@@ -66,7 +66,7 @@ void FStructReinstancer::CollectObjects()
 			// Empty
 		}
 	};
-	struct FDataForReinstanciation
+	struct FDataForReinstantiation
 	{
 		TNotNull<UObject*> Object;
 		UObject* ClassGeneratedBy = nullptr;
@@ -78,11 +78,11 @@ void FStructReinstancer::CollectObjects()
 
 	FVisitorReferenceCollector Collector;
 
-	// bRequiresReinstanciation will be set (in AddStructReferencedObjects) when the struct is reached.
-	check(bSerializingForReinstanciation == false);
-	bRequiresReinstanciation = false;
+	// bRequiresReinstantiation will be set (in AddStructReferencedObjects) when the struct is reached.
+	check(bSerializingForReinstantiation == false);
+	bRequiresReinstantiation = false;
 
-	TArray<FDataForReinstanciation> ReinstanciationDatas;
+	TArray<FDataForReinstantiation> ReinstantiationDatas;
 
 	// Handle regular UObjects
 	for (TObjectIterator<UObject> ObjectIt; ObjectIt; ++ObjectIt)
@@ -92,10 +92,10 @@ void FStructReinstancer::CollectObjects()
 		Object->CallAddReferencedObjects(Collector);
 		Collector.AddPropertyReferencesWithStructARO(Object->GetClass(), Object);
 
-		if (bRequiresReinstanciation)
+		if (bRequiresReinstantiation)
 		{
-			ReinstanciationDatas.Emplace(Object, Object->GetClass()->ClassGeneratedBy.Get(), Object->GetClass(), (void*)Object, Object->GetClass(), (const void*)Object->GetClass()->GetDefaultObject());
-			bRequiresReinstanciation = false;
+			ReinstantiationDatas.Emplace(Object, Object->GetClass()->ClassGeneratedBy.Get(), Object->GetClass(), (void*)Object, Object->GetClass(), (const void*)Object->GetClass()->GetDefaultObject());
+			bRequiresReinstantiation = false;
 		}
 	}
 
@@ -110,13 +110,13 @@ void FStructReinstancer::CollectObjects()
 			UScriptStruct* SparseDataStruct = Class->GetSparseClassDataStruct();
 			Collector.AddPropertyReferencesWithStructARO(SparseDataStruct, SparseData);
 
-			if (bRequiresReinstanciation)
+			if (bRequiresReinstantiation)
 			{
 				UClass* SuperClass = Class->GetSuperClass();
 				const UScriptStruct* SuperSCDType = SuperClass ? SuperClass->GetSparseClassDataStruct() : nullptr;
 				const void* SuperSCD = SuperClass ? SuperClass->GetSparseClassData(EGetSparseClassDataMethod::ReturnIfNull) : nullptr;
-				ReinstanciationDatas.Emplace(Class, Class->ClassGeneratedBy.Get(), SparseDataStruct, SparseData, SuperSCDType, (const void*)SuperSCD);
-				bRequiresReinstanciation = false;
+				ReinstantiationDatas.Emplace(Class, Class->ClassGeneratedBy.Get(), SparseDataStruct, SparseData, SuperSCDType, (const void*)SuperSCD);
+				bRequiresReinstantiation = false;
 			}
 		}
 
@@ -125,45 +125,45 @@ void FStructReinstancer::CollectObjects()
 		{
 			Collector.AddPropertyReferencesWithStructARO(Class, CDO);
 
-			if (bRequiresReinstanciation)
+			if (bRequiresReinstantiation)
 			{
 				const UClass* SuperClass = Class->GetSuperClass();
 				const void* SuperData = SuperClass ? SuperClass->GetDefaultObject() : nullptr;
-				ReinstanciationDatas.Emplace(CDO, Class->ClassGeneratedBy.Get(), Class, (void*)CDO, SuperClass, (const void*)SuperData);
-				bRequiresReinstanciation = false;
+				ReinstantiationDatas.Emplace(CDO, Class->ClassGeneratedBy.Get(), Class, (void*)CDO, SuperClass, (const void*)SuperData);
+				bRequiresReinstantiation = false;
 			}
 		}
 	}
 
-	ensure(bRequiresReinstanciation == false);
+	ensure(bRequiresReinstantiation == false);
 
-	ObjectsToReinstanciate.Reset();
-	ObjectsToReinstanciate.Reserve(ReinstanciationDatas.Num());
-	for (FDataForReinstanciation& ReinstanciationData : ReinstanciationDatas)
+	ObjectsToReinstantiate.Reset();
+	ObjectsToReinstantiate.Reserve(ReinstantiationDatas.Num());
+	for (FDataForReinstantiation& ReinstantiationData : ReinstantiationDatas)
 	{
-		bSerializingForReinstanciation = true;
+		bSerializingForReinstantiation = true;
 
-		UE_LOGFMT(LogClass, Verbose, "Serialize object {Name}", ReinstanciationData.Object->GetFullName());
+		UE_LOGFMT(LogClass, Verbose, "Serialize object {Name}", ReinstantiationData.Object->GetFullName());
 
 		TArray<uint8> Buffer;
 		FObjectWriter Writer(Buffer);
-		ReinstanciationData.Struct->SerializeTaggedProperties(
+		ReinstantiationData.Struct->SerializeTaggedProperties(
 			Writer,
-			(uint8*)((void*)ReinstanciationData.Data),
-			ReinstanciationData.SuperStruct,
-			(uint8*)ReinstanciationData.SuperData,
-			ReinstanciationData.Object
+			(uint8*)((void*)ReinstantiationData.Data),
+			ReinstantiationData.SuperStruct,
+			(uint8*)ReinstantiationData.SuperData,
+			ReinstantiationData.Object
 		);
 
-		bSerializingForReinstanciation = false;
+		bSerializingForReinstantiation = false;
 
-		ObjectsToReinstanciate.Emplace(
-			FObjectKey(ReinstanciationData.Object),
-			FObjectKey(ReinstanciationData.ClassGeneratedBy),
-			TObjectKey<UStruct>(ReinstanciationData.Struct),
-			ReinstanciationData.Data,
-			TObjectKey<const UStruct>(ReinstanciationData.SuperStruct),
-			ReinstanciationData.SuperData,
+		ObjectsToReinstantiate.Emplace(
+			FObjectKey(ReinstantiationData.Object),
+			FObjectKey(ReinstantiationData.ClassGeneratedBy),
+			TObjectKey<UStruct>(ReinstantiationData.Struct),
+			ReinstantiationData.Data,
+			TObjectKey<const UStruct>(ReinstantiationData.SuperStruct),
+			ReinstantiationData.SuperData,
 			MoveTemp(Buffer)
 		);
 	}
@@ -180,10 +180,10 @@ void FStructReinstancer::ReinstanceObjects()
 	};
 
 	TArray<FSortedSavedData> SortedDatas;
-	SortedDatas.Reserve(ObjectsToReinstanciate.Num());
-	for (int32 Index = 0; Index < ObjectsToReinstanciate.Num(); ++Index)
+	SortedDatas.Reserve(ObjectsToReinstantiate.Num());
+	for (int32 Index = 0; Index < ObjectsToReinstantiate.Num(); ++Index)
 	{
-		if (UObject* SavedObject = ObjectsToReinstanciate[Index].Object.ResolveObjectPtr())
+		if (UObject* SavedObject = ObjectsToReinstantiate[Index].Object.ResolveObjectPtr())
 		{
 			SortedDatas.Emplace(SavedObject, SavedObject->GetArchetype(), Index);
 		}
@@ -195,7 +195,7 @@ void FStructReinstancer::ReinstanceObjects()
 		});
 	for (FSortedSavedData& SortedData : SortedDatas)
 	{
-		const FSavedDataForReinstanciation& SavedData = ObjectsToReinstanciate[SortedData.SavedDataIndex];
+		const FSavedDataForReinstantiation& SavedData = ObjectsToReinstantiate[SortedData.SavedDataIndex];
 		UObject* SavedObject = SortedData.Object;
 		UStruct* SavedStruct = SavedData.Struct.ResolveObjectPtr();
 		if (ensure(SavedStruct))
@@ -212,7 +212,7 @@ void FStructReinstancer::ReinstanceObjects()
 			const UStruct* SavedSuperStruct = SavedData.SuperStruct.ResolveObjectPtr();
 			const void* SavedSuperObjectData = SavedSuperStruct ? SavedData.SuperObjectData : nullptr;
 
-			bSerializingForReinstanciation = true;
+			bSerializingForReinstantiation = true;
 
 			UE_LOGFMT(LogClass, Verbose, "Reinstance object {Name}", SavedObject->GetFullName());
 
@@ -223,7 +223,7 @@ void FStructReinstancer::ReinstanceObjects()
 				SavedSuperStruct,
 				(const uint8*)SavedSuperObjectData);
 
-			bSerializingForReinstanciation = false;
+			bSerializingForReinstantiation = false;
 		}
 	}
 }
diff --git a/Engine/Source/Runtime/CoreUObject/Public/StructUtils/StructReinstancer.h b/Engine/Source/Runtime/CoreUObject/Public/StructUtils/StructReinstancer.h
--- a/Engine/Source/Runtime/CoreUObject/Public/StructUtils/StructReinstancer.h
+++ b/Engine/Source/Runtime/CoreUObject/Public/StructUtils/StructReinstancer.h
@@ -46,9 +46,9 @@ namespace UE::StructUtils
 		UE_API void ReinstanceObjects();
 
 		/** Serializing the UObject owners to gather the modified properties. */
-		bool IsSerializingForReinstanciation() const
+		bool IsSerializingForReinstantiation() const
 		{
-			return bSerializingForReinstanciation;
+			return bSerializingForReinstantiation;
 		}
 
 		/** Is the struct added to the FStructReinstancer. */
@@ -91,10 +91,10 @@ namespace UE::StructUtils
 			return Found ? Found->New.Get() : nullptr;
 		}
 
-		/** The current collected UObject contains at least one UScriptStruct that is being reinstanciated. */
-		void MarkAsRequiresReinstanciation()
+		/** The current collected UObject contains at least one UScriptStruct that is being reinstantiated. */
+		void MarkAsRequiresReinstantiation()
 		{
-			bRequiresReinstanciation = true;
+			bRequiresReinstantiation = true;
 		}
 	private:
 		struct FStruct
@@ -103,7 +103,7 @@ namespace UE::StructUtils
 			TObjectPtr<const UScriptStruct> Duplicated;
 			TObjectPtr<const UScriptStruct> New;
 		};
-		struct FSavedDataForReinstanciation
+		struct FSavedDataForReinstantiation
 		{
 			FObjectKey Object;
 			FObjectKey ClassGeneratedBy;
@@ -114,9 +114,9 @@ namespace UE::StructUtils
 			TArray<uint8> Buffer;
 		};
 		TArray<FStruct> StructuresToReinstantiate;
-		TArray<FSavedDataForReinstanciation> ObjectsToReinstanciate;
-		bool bRequiresReinstanciation = false;
-		bool bSerializingForReinstanciation = false;
+		TArray<FSavedDataForReinstantiation> ObjectsToReinstantiate;
+		bool bRequiresReinstantiation = false;
+		bool bSerializingForReinstantiation = false;
 	};
 
 	struct UE_INTERNAL FStructReinstancerScope final : FNoncopyable
