Allow FPackageTrailer to return it's offset in the archive that it was loaded from.

#rnx

- When a FPackageTrailer is serialized we store it's position in the archive. Up until now this has been purely for internal use but we now have use cases where it would be useful to be able to poll for this value.
- Renamed some variables to make their purpose clearer.

#Synced-CL 51332327

diff --git a/Engine/Source/Runtime/CoreUObject/Private/UObject/PackageTrailer.cpp b/Engine/Source/Runtime/CoreUObject/Private/UObject/PackageTrailer.cpp
--- a/Engine/Source/Runtime/CoreUObject/Private/UObject/PackageTrailer.cpp
+++ b/Engine/Source/Runtime/CoreUObject/Private/UObject/PackageTrailer.cpp
@@ -294,7 +294,7 @@ bool FPackageTrailerBuilder::BuildAndAppendTrailer(FLinkerSave* Linker, FArchive
 	return BuildAndAppendTrailer(Linker, DataArchive, CurrentOffset);
 }
 
-bool FPackageTrailerBuilder::BuildAndAppendTrailer(FLinkerSave* Linker, FArchive& DataArchive, int64& InOutPackageFileOffset)
+bool FPackageTrailerBuilder::BuildAndAppendTrailer(FLinkerSave* Linker, FArchive& DataArchive, int64& InOutArchiveOffset)
 {
 	// Note that we do not serialize containers directly as we want a file format that is 
 	// 100% under our control. This will allow people to create external scripts that can
@@ -355,7 +355,7 @@ bool FPackageTrailerBuilder::BuildAndAppendTrailer(FLinkerSave* Linker, FArchive
 
 	// Now that we have the complete trailer we can serialize it to the archive
 
-	Trailer.TrailerPositionInFile = InOutPackageFileOffset;
+	Trailer.PackageTrailerOffset = InOutArchiveOffset;
 
 	const int64 TrailerPositionInDataArchive = DataArchive.Tell();
 	DataArchive << Trailer.Header;
@@ -397,7 +397,7 @@ bool FPackageTrailerBuilder::BuildAndAppendTrailer(FLinkerSave* Linker, FArchive
 
 	// Minor sanity check that ::GetTrailerLength works
 	check(CalculateTrailerLength() == (DataArchive.Tell() - TrailerPositionInDataArchive) || DataArchive.IsError());
-	InOutPackageFileOffset += DataArchive.Tell() - TrailerPositionInDataArchive;
+	InOutArchiveOffset += DataArchive.Tell() - TrailerPositionInDataArchive;
 
 	return !DataArchive.IsError();
 }
@@ -547,7 +547,7 @@ bool FPackageTrailer::TryLoad(FArchive& Ar)
 {
 	check(Ar.IsLoading());
 
-	TrailerPositionInFile = Ar.Tell();
+	PackageTrailerOffset = Ar.Tell();
 
 	Header.Tag = 0; // Make sure we ignore anything previously loaded
 	Ar << Header.Tag;
@@ -563,7 +563,7 @@ bool FPackageTrailer::TryLoad(FArchive& Ar)
 	Ar << Header.PayloadsDataLength;
 
 	// The header and the payloads should not exceed the remaining data in the archive
-	if (TrailerPositionInFile + Header.HeaderLength + Header.PayloadsDataLength > (uint64)Ar.TotalSize())
+	if (PackageTrailerOffset + Header.HeaderLength + Header.PayloadsDataLength > (uint64)Ar.TotalSize())
 	{
 		Ar.SetError();
 		return false;
@@ -673,8 +673,8 @@ FCompressedBuffer FPackageTrailer::LoadLocalPayload(const FIoHash& Id, FArchive&
 		return FCompressedBuffer();
 	}
 
-	const int64 OffsetInFile = TrailerPositionInFile + Header.HeaderLength + Entry->OffsetInFile;
-	Ar.Seek(OffsetInFile);
+	const int64 PayloadOffsetInFile = PackageTrailerOffset + Header.HeaderLength + Entry->OffsetInFile;
+	Ar.Seek(PayloadOffsetInFile);
 
 	FCompressedBuffer Payload;
 	Ar << Payload;
@@ -741,14 +741,14 @@ int64 FPackageTrailer::FindPayloadOffsetInFile(const FIoHash& Id) const
 		const Private::FLookupTableEntry* Entry = FindEntryInHeader(Header, Id);
 
 		//TODO Better way to return an error?
-		check(TrailerPositionInFile != INDEX_NONE);
+		check(PackageTrailerOffset != INDEX_NONE);
 		check(Header.PayloadsDataLength != INDEX_NONE);
 		check(Entry != nullptr);
 
 		switch (Entry->AccessMode)
 		{
 			case EPayloadAccessMode::Local:
-				return TrailerPositionInFile + Header.HeaderLength + Entry->OffsetInFile;
+				return PackageTrailerOffset + Header.HeaderLength + Entry->OffsetInFile;
 
 			case EPayloadAccessMode::Referenced:
 				return Entry->OffsetInFile;
diff --git a/Engine/Source/Runtime/CoreUObject/Public/UObject/PackageTrailer.h b/Engine/Source/Runtime/CoreUObject/Public/UObject/PackageTrailer.h
--- a/Engine/Source/Runtime/CoreUObject/Public/UObject/PackageTrailer.h
+++ b/Engine/Source/Runtime/CoreUObject/Public/UObject/PackageTrailer.h
@@ -269,15 +269,15 @@ public:
 	[[nodiscard]] COREUOBJECT_API bool BuildAndAppendTrailer(FLinkerSave* Linker, FArchive& DataArchive);
 
 	/**
-	 * @param ExportsArchive	The linker associated with the package being written to disk.
-	 * @param DataArchive		The archive where the package data has been written to. This is where the FPackageTrailer will be written to
-	 * @param InOutPackageFileOffset The offset at which the trailer is written in the package file. In this function version,
-	 *        DataArchive might be an archive for the entire package file, or it might be a separate archive solely for the package trailer.
-	 *        The caller must provide the initial value with the accumulated offset, and the function will increment this value by how many
-	 *        bytes are written to DataArchive.
+	 * @param ExportsArchive		The linker associated with the package being written to disk.
+	 * @param DataArchive			The archive where the package data has been written to. This is where the FPackageTrailer will be written to
+	 * @param InOutArchiveOffset	The offset at which the trailer is written in the archive. In this function version, DataArchive might
+	 *								be an archive for the entire archive, or it might be a separate archive solely for the package trailer.
+	 *								The caller must provide the initial value with the accumulated offset, and the function will increment this
+	 *								value by how many bytes are written to DataArchive.
 	 * @return True if the builder was created and appended successfully and false if any error was encountered
 	 */
-	[[nodiscard]] COREUOBJECT_API bool BuildAndAppendTrailer(FLinkerSave* Linker, FArchive& DataArchive, int64& InOutPackageFileOffset);
+	[[nodiscard]] COREUOBJECT_API bool BuildAndAppendTrailer(FLinkerSave* Linker, FArchive& DataArchive, int64& InOutArchiveOffset);
 
 	/** Returns if the builder has any payload entries or not */
 	[[nodiscard]] COREUOBJECT_API bool IsEmpty() const;
@@ -490,6 +490,12 @@ public:
 	/** Returns the number of payloads that the trailer owns that match the given filter type. @See EPayloadFilter */
 	[[nodiscard]] COREUOBJECT_API int32 GetNumPayloads(EPayloadFilter Filter) const;
 
+	/** Returns the offset (in bytes) of the package trailer in the archive that it was loaded from or written too. */
+	[[nodiscard]] int64 GetPackageTrailerOffset() const
+	{
+		return PackageTrailerOffset;
+	}
+
 	struct FHeader
 	{
 		/** Unique value used to identify the header */
@@ -547,12 +553,13 @@ private:
 	/** Create a valid footer for the current trailer */
 	COREUOBJECT_API FFooter CreateFooter() const;
 
-	/** Where in the workspace domain package file the trailer is located */
-	int64 TrailerPositionInFile = INDEX_NONE;
+	/** Offset into the archive where the FPackageTrailer is located */
+	int64 PackageTrailerOffset = INDEX_NONE;
 
 	/** 
 	 * The header of the trailer. Since this contains the lookup table for payloads we keep this in memory once the trailer
-	 * has been loaded. There is no need to keep the footer in memory */
+	 * has been loaded. There is no need to keep the footer in memory
+	 */
 	FHeader Header;
 };
 
