Avoid crashing when loading a package with an incorrect FPackageFileSummary::PayloadTocOffset when possible.

- In crash reports we have seen that on very rare occasions that packages are being loaded with incorrect values set for FPackageFileSummary::PayloadTocOffset. In these scenarios it looks like the rest of the file is absolutely fine and the package trailer does exist at the end of the file, only the PayloadTocOffset value is incorrect.
- The incorrect offset seems to always be much larger than the package file in question, in some cases by multiple GiB but other than that there is no pattern between cases.
- At the moment we do not know what is causing this, only that it is limited to a handful of users every month.
- Since FPackageTrailer can be serialized in reverse and we know that it should be the last thing in the package file it is possible for us to use this slower loading path rather than just crashing on invalid file seeks if we detect that the PayloadTocOffset value is not reasonable.
- I have added an ensure to make sure that this problem is still recorded and we can track future attempts to prevent the corruption occuring in the first place.
- I have also added logging to the logfile/outputwindow/message log, informing the user of the problem, which package has the problem and that re-saving the package will msot likely fix it.
- The ensure can be supressed by setting the cvar 'Serialization.EnableSerializePackageTrailerEnsure' to false.

#Synced-CL 51332327

diff --git a/Engine/Source/Runtime/CoreUObject/Private/UObject/LinkerLoad.cpp b/Engine/Source/Runtime/CoreUObject/Private/UObject/LinkerLoad.cpp
--- a/Engine/Source/Runtime/CoreUObject/Private/UObject/LinkerLoad.cpp
+++ b/Engine/Source/Runtime/CoreUObject/Private/UObject/LinkerLoad.cpp
@@ -223,6 +223,12 @@ static TAutoConsoleVariable<bool> CVarApplyBulkDataFix(
 	true,
 	TEXT("When true, we will try to fix potentially bad bulkdata offsets"));
 
+bool GEnsureSerializePackageTrailer = true;
+static FAutoConsoleVariableRef CVarEnsureSerializePackageTrailer(
+	TEXT("Serialization.EnableSerializePackageTrailerEnsure"),
+	GEnsureSerializePackageTrailer,
+	TEXT("When true, an ensure will fire if we detect incorrect values when trying to serialize the package trailer"));
+
 /**
 * Test whether the given package index is a valid import or export in this package
 */
@@ -1871,16 +1877,29 @@ FLinkerLoad::ELinkerStatus FLinkerLoad::SerializePackageTrailer()
 		return LINKER_Loaded;
 	}
 
-	check(PackageTrailer == nullptr);
+	check(PackageTrailer == nullptr); // There should not be any pre-existing trailer if we are serializing one
 
-	if (Summary.PayloadTocOffset > 0)
+	if (Summary.PayloadTocOffset <= 0)
 	{
-		int64 CurPos = Tell();
-		Seek(Summary.PayloadTocOffset);
+		bHasSerializedPackageTrailer = true;
+		return LINKER_Loaded;
+	}
+
+	int64 CurPos = Tell();
 
-		PackageTrailer = MakeUnique<UE::FPackageTrailer>();
+	PackageTrailer = MakeUnique<UE::FPackageTrailer>();
+
+	bool bResult = false;
+
+	// We have seen a very low rate of package file where this value is corrupted and would cause us to seek outside of the
+	// package file. Instead of crashing we can attempt to load the trailer from the end of the file rather than using the
+	// offset. For now we will continue to track the problem with an ensure while we try to hunt down the source of the
+	// corruption.
+	if (Summary.PayloadTocOffset < TotalSize())
+	{
+		Seek(Summary.PayloadTocOffset);
+		bResult = PackageTrailer->TryLoad(*this);
 
-		bool bResult = PackageTrailer->TryLoad(*this);
 		if (!bResult && Summary.GetFileVersionUE().ToValue() == (int32)EUnrealEngineObjectUE5Version::DATA_RESOURCES)
 		{
 			// There was an issue that was causing incorrect values to be written to PayloadTocOffset for
@@ -1892,33 +1911,58 @@ FLinkerLoad::ELinkerStatus FLinkerLoad::SerializePackageTrailer()
 			Seek(TotalSize());
 			bResult = PackageTrailer->TryLoadBackwards(*this);
 		}
+	}
+	else
+	{
+		Seek(TotalSize());
+		bResult = PackageTrailer->TryLoadBackwards(*this);
 
-		if (!bResult)
+		// If we managed to load the trailer anyway then we should fix up the offset and warn the user that re-saving the package file will probably fix it.
+		// Otherwise we can 
+		if (bResult)
 		{
-			// If the archive has an error then we found a package trailer but it failed to serialize
-			// correctly and we most likely have a problem with the file.
-			// If the load failed but the archive is fine then the package is just of an older format
-			// and there never was a package trailer to load.
-			if (IsError())
-			{
-				UE_ASSET_LOG(LogLinker, Error, PackagePath, TEXT("Package has a corrupted package trailer"));
+			const int64 OldOffset = Summary.PayloadTocOffset;
+			Summary.PayloadTocOffset = PackageTrailer->GetPackageTrailerOffset();
 
-				Async(EAsyncExecution::TaskGraphMainThread,
-					[DebugName = GetDebugName()]()
-					{
-						FMessageLog("LoadErrors").SuppressLoggingToOutputLog(true)
-							.Error(FText::Format(NSLOCTEXT("Core", "LinkerLoad_CorruptTrailer", "Package {0} has a corrupted package trailer"),
-							FText::FromString(DebugName)));
-					}
-				);
+			// Although we can survive PayloadTocOffset being corrupted it still indicates a problem. The ensure will make sure that the
+			// problem is recorded so it can be tracked if it is still occurring.
+			ensureMsgf(!GEnsureSerializePackageTrailer, TEXT("FPackageFileSummary::PayloadTocOffset is incorrect for '%s'. Offset: %lld ArchiveSize: %lld CorrectedOffset: %lld"), *GetArchiveName(), OldOffset, TotalSize(), Summary.PayloadTocOffset);
 
-				return LINKER_Failed;
-			}
+			// Log a warning so that the user knows how to repair their package file (by re-saving it)
+			Async(EAsyncExecution::TaskGraphMainThread,
+				[DebugName = GetDebugName()]()
+				{
+					FMessageLog("LoadErrors") .Warning(
+						FText::Format(NSLOCTEXT("Core", "LinkerLoad_IncorrectSummary", "Package {0} has incorrect values in it's FPackageFileSummary, please re-save to fix this!"),
+						FText::FromString(DebugName)));
+				}
+			);
+		}
+	}
 
-			PackageTrailer.Reset();
+	Seek(CurPos); // Restore the original position
+
+	if (!bResult)
+	{
+		// If the archive has an error then we found a package trailer but it failed to serialize
+		// correctly and we most likely have a problem with the file.
+		// If the load failed but the archive is fine then the package is just of an older format
+		// and there never was a package trailer to load.
+		if (IsError())
+		{
+			Async(EAsyncExecution::TaskGraphMainThread,
+				[DebugName = GetDebugName()]()
+				{
+					FMessageLog("LoadErrors").Error(
+						FText::Format(NSLOCTEXT("Core", "LinkerLoad_CorruptTrailer", "Package {0} has a corrupted package trailer"),
+						FText::FromString(DebugName)));
+				}
+			);
+
+			return LINKER_Failed;
 		}
 
-		Seek(CurPos);
+		PackageTrailer.Reset();
 	}
 
 	bHasSerializedPackageTrailer = true;
