EditorBulkData will now only find the package trailer and/or package path from the FArchive it is serializing from and never from it's owners UPackage.

- Follow on from CL 50207788
- In additiona to the issue solved in 50207788 a new problem was recently discovered. If a UObject containing editor bulkdata was duplicated from one package to another AND the target package already contained an editor bulkdata instance with the exact same data BUT was already virtualized then the duplication would crash. This is because during the duplication serialization we would find the package trailer for the target package and assume that is where we are loading the data from rather than the duplication archive.
- Originally I did intend to just remove the cvar added in 50207788 but continue to check the package then the archive, but with the above in mind we should probably only ever search the archive for the FLinkerLoad/FPackageTrailer.
- Previous behavior can be restored by setting the cvar 'Serialization.LegacyLinkerDiscovery' to true. This cvar will be removed in a future release and is being provided only to make it easier to confirm/deny if any new serialization issues are being caused by this change or not.

#Synced-CL 51332327

diff --git a/Engine/Source/Runtime/CoreUObject/Private/Serialization/EditorBulkData.cpp b/Engine/Source/Runtime/CoreUObject/Private/Serialization/EditorBulkData.cpp
--- a/Engine/Source/Runtime/CoreUObject/Private/Serialization/EditorBulkData.cpp
+++ b/Engine/Source/Runtime/CoreUObject/Private/Serialization/EditorBulkData.cpp
@@ -188,22 +188,17 @@ static FPackagePath GetPackagePathFromOwner(UObject* Owner)
 	}
 }
 
-bool GCVarAllowLinkerFromArchive = true;
-static FAutoConsoleVariableRef CVar_AllowLinkerFromArchive(
-	TEXT("Serialization.AllowLinkerFromArchive"),
-	GCVarAllowLinkerFromArchive,
-	TEXT("When enabled we can attempt to find a FLinkerLoad from the archive as well as the UObject. Set to false for the old behavior")
+bool GCVarLegacyLinkerDiscovery = false;
+static FAutoConsoleVariableRef CVar_LegacyLinkerDiscovery(
+	TEXT("Serialization.LegacyLinkerDiscovery"),
+	GCVarLegacyLinkerDiscovery,
+	TEXT("When enabled we revert to the legacy behavior of trying find the FLinkerLoad from the owning UObject rather than the FArchive")
 );
 
 /** Utility for finding the FPackageTrailer associated from either a UObject or FArchive */
 static const FPackageTrailer* FindPackageTrailer(UObject* Owner, FArchive& Ar)
 {
-	const FLinkerLoad* Linker = GetLinkerLoadFromOwner(Owner);
-
-	if (Linker == nullptr && GCVarAllowLinkerFromArchive)
-	{
-		Linker = Cast<FLinkerLoad>(Ar.GetLinker());
-	}
+	const FLinkerLoad* Linker = GCVarLegacyLinkerDiscovery ? GetLinkerLoadFromOwner(Owner) : Cast<FLinkerLoad>(Ar.GetLinker());
 
 	if (Linker != nullptr)
 	{
@@ -218,12 +213,7 @@ static const FPackageTrailer* FindPackageTrailer(UObject* Owner, FArchive& Ar)
 /** Utility for finding the package path associated from either a UObject or FArchive */
 static FPackagePath FindPackagePath(UObject* Owner, FArchive& Ar)
 {
-	const FLinkerLoad* Linker = GetLinkerLoadFromOwner(Owner);
-
-	if (Linker == nullptr && GCVarAllowLinkerFromArchive)
-	{
-		Linker = Cast<FLinkerLoad>(Ar.GetLinker());
-	}
+	const FLinkerLoad* Linker = GCVarLegacyLinkerDiscovery ? GetLinkerLoadFromOwner(Owner) : Cast<FLinkerLoad>(Ar.GetLinker());
 
 	if (Linker != nullptr)
 	{
