Gather all subobjects of a borrowed object when resetting it.

#Synced-CL 51297925

diff --git a/Engine/Source/Runtime/CoreUObject/Private/UObject/RemoteObjectSerialization.cpp b/Engine/Source/Runtime/CoreUObject/Private/UObject/RemoteObjectSerialization.cpp
--- a/Engine/Source/Runtime/CoreUObject/Private/UObject/RemoteObjectSerialization.cpp
+++ b/Engine/Source/Runtime/CoreUObject/Private/UObject/RemoteObjectSerialization.cpp
@@ -740,7 +740,7 @@ public:
 		return ObjectsToSerialize;
 	}
 
-	FORCENOINLINE void SerializeRemoteObject(UObject* Object, const FRemoteObjectHeader& Header)
+	virtual void SerializeRemoteObject(UObject* Object, const FRemoteObjectHeader& Header)
 	{		
 		Object->Serialize(*this);
 	}
@@ -786,31 +786,56 @@ using FRemoteObjectReaderScope = TRemoteObjectArchiveScope<FMemoryReader>;
 /**
 * Helper archive that serializes the difference between archetypes and their instances
 * This is achieved using delta serialization but the data we serialize against is coming from instances of the archetypes
-* Effectively this is the opposite how delta serialization normally works which serializes instances of archetypes agains the archetypes
-* In other words this archive is used to serialize archetypes, not their instances.
+* Effectively this is the opposite how delta serialization normally works which serializes instances of archetypes against the archetypes
+* In other words this archive is used to serialize archetypes, not their instances (more specifically the delta between the archetype and its instance)
 */
 class FArchetypeDeltaWriter : public FArchiveRemoteObjectWriter
 {
-	TMap<const UObject*, UObject*>& ArchetypeToInstanceMap;
+	const TMap<UObject*, const UObject*>& InstanceToArchetypeMap;
+	const TMap<const UObject*, UObject*>& ArchetypeToInstanceMap;
+	UObject* InstanceScope = nullptr;
 
 public:
 
-	FArchetypeDeltaWriter(UObject* InRootObject, FRemoteObjectData& OutObjectData, TMap<const UObject*, UObject*>& InArchetypeToInstanceMap)
+	FArchetypeDeltaWriter(UObject* InRootObject, FRemoteObjectData& OutObjectData, const TMap<UObject*, const UObject*>& InInstanceToArchetypeMap, const TMap<const UObject*, UObject*>& InArchetypeToInstanceMap)
 		: FArchiveRemoteObjectWriter(InRootObject, OutObjectData, /*MigrationContext*/ nullptr, TEXT("RemoteArchetypeDeltaWriter"))
+		, InstanceToArchetypeMap(InInstanceToArchetypeMap)
 		, ArchetypeToInstanceMap(InArchetypeToInstanceMap)
 	{
 	}
 
-	virtual UObject* GetArchetypeFromLoader(const UObject* Obj) override
+	virtual UObject* GetArchetypeFromLoader(const UObject* InArchetype) override
 	{
+		checkf(!InArchetype || InArchetype->IsTemplate(), TEXT("FArchetypeDeltaWriter should only be serializing archetypes but is trying to find an instance for %s"), *GetPathNameSafe(InArchetype));
+		
 		// Since this archive serializes the archetype we want the archetype of that archetype to be its instance (effectively reversing the object -> archetype relationship)
 		// this way we will delta serialize the difference between the archetype and its instance
-		if (UObject** FoundArchetypeInstance = ArchetypeToInstanceMap.Find(Obj))
+		if (InstanceScope)
+		{
+			// Prefer Instance -> Archetype lookup because it's never ambiguous
+			const UObject* const* ScopeArchetype = InstanceToArchetypeMap.Find(InstanceScope);
+			if (ScopeArchetype && *ScopeArchetype == InArchetype)
+			{
+				return InstanceScope;
+			}
+		}
+
+		UE_LOG(LogRemoteSerialization, Warning, TEXT("FArchetypeDeltaWriter::ArchetypeToInstanceMap does not contain an archetype mapping for %s"), *GetPathNameSafe(InArchetype));
+		return FArchiveRemoteObjectWriter::GetArchetypeFromLoader(InArchetype);
+	}
+
+	virtual void SerializeRemoteObject(UObject* Object, const FRemoteObjectHeader& Header) override
+	{		
+		UObject* ObjectToSerialize = nullptr;
+		if (const UObject* const* Archetype = InstanceToArchetypeMap.Find(Object))
 		{
-			return *FoundArchetypeInstance;
+			ObjectToSerialize = const_cast<UObject*>(*Archetype);			
 		}
-		UE_LOG(LogRemoteSerialization, Warning, TEXT("FArchetypeDeltaWriter::ArchetypeToInstanceMap does not contain an archetype mapping for %s"), *GetPathNameSafe(Obj));
-		return FArchiveRemoteObjectWriter::GetArchetypeFromLoader(Obj);
+		checkf(ObjectToSerialize, TEXT("Serializing an instance without matching archetype: %s"), *GetPathNameSafe(Object));
+
+		InstanceScope = Object;
+		ObjectToSerialize->Serialize(*this);
+		InstanceScope = nullptr;
 	}
 
 protected:
@@ -821,38 +846,20 @@ protected:
 
 		FRemoteObjectReferenceInfo ReplacementInfo(RefInfo);
 
-		// If we're serializing a reference to an archetype try to replace it with a reference to its instance
-		// This way the produced delta between the archetype and its instance will be correctly pointing to the instances of default subobjects
-		// we can then deserialize over the archetype instance.
-		if (UObject** Instance = ArchetypeToInstanceMap.Find(RefInfo.Object))
+		if (InstanceScope)
 		{
-			ReplacementInfo.Object = *Instance;
-			ReplacementInfo.Id = FRemoteObjectId(*Instance);
+			// If we're serializing a reference to an archetype try to replace it with a reference to its instance
+			// This way the produced delta between the archetype and its instance will be correctly pointing to the instances of default subobjects
+			// we can then deserialize over the archetype instance.
+			if (UObject* const* FoundInstance = ArchetypeToInstanceMap.Find(RefInfo.Object))
+			{
+				ReplacementInfo.Object = *FoundInstance;
+				ReplacementInfo.Id = FRemoteObjectId(*FoundInstance);
+			}
 		}
 
 		FArchiveRemoteObjectWriter::WriteObjectReference(ReplacementInfo);
 	}
-
-public:
-
-	virtual bool PopulateObjectHeader(UObject* Object, FRemoteObjectHeader& OutHeader) override
-	{
-		// As per comment in WriteObjectReference - when serializing archetype data pretend we're actually serializing references to instances of the archetype 
-		// so substitute the archetype object with its instance when serializing the object header
-		if (UObject** Instance = ArchetypeToInstanceMap.Find(Object))
-		{
-			Object = *Instance;
-		}
-
-		// If for some reason we failed to substitute the archetype with its instance or for some other reason Object is an asset, skip it
-		bool bCanSaveObject = !FRemoteObjectId(Object).IsAsset();
-		if (bCanSaveObject)
-		{
-			bCanSaveObject = FArchiveRemoteObjectWriter::PopulateObjectHeader(Object, OutHeader);
-		}
-		
-		return bCanSaveObject;
-	}
 };
 
 
@@ -873,14 +880,25 @@ class FArchiveRemoteObjectReader : public TArchiveRemoteObjectBase<FMemoryReader
 {
 	const TArray<FName>& Names;
 	const TArray<FRemoteObjectId>& RemoteIds;
-	const TArray<UObject*>& ResolvedPathNameObjects;
+	TArray<UObject*> ResolvedPathNameObjects;
 	const ERemoteObjectSerializationFlags DeserializeFlags = ERemoteObjectSerializationFlags::None;
 
+	void ResolvePathNames(const FRemoteObjectData& InObjectData, TArray<UObject*>& OutResolvedObjects)
+	{
+		OutResolvedObjects.Reset(InObjectData.PathNames.Num());
+		for (const FPackedRemoteObjectPathName& PathName : InObjectData.PathNames)
+		{
+			UObject* Object = PathName.Resolve(InObjectData.Tables);
+			OutResolvedObjects.Add(Object);
+		}
+	}
+
 	ERemoteReferenceType ReadObjectReference(FObjectPtr& Value)
 	{
 		using namespace UE::RemoteObject::Private;
 		using namespace UE::RemoteObject::Handle;
 
+		FObjectPtr NewValue;
 		ERemoteReferenceType Type;		
 		*this << Type;
 
@@ -913,7 +931,7 @@ class FArchiveRemoteObjectReader : public TArchiveRemoteObjectBase<FMemoryReader
 					if (Obj)
 					{
 						MarkAsRemoteReference(Obj);
-						Value = FObjectPtr(Obj);
+						NewValue = FObjectPtr(Obj);
 						bNeedsResolvingWithId = false;
 					}
 				}
@@ -926,18 +944,24 @@ class FArchiveRemoteObjectReader : public TArchiveRemoteObjectBase<FMemoryReader
 				if (!EnumHasAnyFlags(DeserializeFlags, ERemoteObjectSerializationFlags::PreserveRemoteReferences) || Value.GetResidence() == EResidence::Local)
 				{
 					FObjectHandle Handle = FObjectHandle::FromIdNoResolve(ObjId);
-					Value = FObjectPtr(Handle);
+					NewValue = FObjectPtr(Handle);
 				}
 #endif // UE_WITH_REMOTE_OBJECT_HANDLE
 			}
 		}
-		else
-		{
-			Value = FObjectPtr();
-		}
+
+		SetObjectReference(Value, NewValue);
+
 		return Type;
 	}
 
+protected:
+
+	virtual void SetObjectReference(FObjectPtr& CurrentValue, FObjectPtr NewValue)
+	{
+		CurrentValue = NewValue;
+	}
+
 public:
 
 	/**
@@ -946,13 +970,13 @@ public:
 	 * @param bInAssignedOwnership:  We must take ownership of the objects we deserialize. We could already have ownership of those objects.
 	 * @param InDeserializeFlags:  The flags for how we should treat references during deserialization.
 	 */
-	FArchiveRemoteObjectReader(FRemoteObjectData& InObjectData, const TArray<UObject*>& InResolvedPathNames, const FUObjectMigrationContext* InMigrationContext, ERemoteObjectSerializationFlags InDeserializeFlags)
-		: TArchiveRemoteObjectBase<FMemoryReader>(InObjectData, InMigrationContext, TEXT("RemoteObjectReader"))
+	FArchiveRemoteObjectReader(FRemoteObjectData& InObjectData, const FUObjectMigrationContext* InMigrationContext, ERemoteObjectSerializationFlags InDeserializeFlags, const TCHAR* ArchiveName = nullptr)
+		: TArchiveRemoteObjectBase<FMemoryReader>(InObjectData, InMigrationContext, ArchiveName ? ArchiveName : TEXT("RemoteObjectReader"))
 		, Names(InObjectData.Tables.Names)
 		, RemoteIds(InObjectData.Tables.RemoteIds)
-		, ResolvedPathNameObjects(InResolvedPathNames)
 		, DeserializeFlags(InDeserializeFlags)
 	{
+		ResolvePathNames(InObjectData, ResolvedPathNameObjects);
 	}
 
 	virtual FArchive& operator<<(UObject*& Obj) override
@@ -1001,6 +1025,104 @@ public:
 
 		return *this;
 	}
+
+	virtual void SerializeRemoteObject(UObject* Object)
+	{
+		Object->Serialize(*this);
+	}
+};
+
+/**
+* Helper archive that deserializes the difference between archetypes and their instances
+* FRemoteObjectData passed to this archive is the difference between the archetypes and their instances
+* The data is deserialized on top of existing instances
+* When deserializing references to UObjects, this archive will preserve references to the known existing instances instead of replacing them with refs to archetypes stored in the provided data (see SetObjectReference override)
+*/
+class FArchetypeDeltaReader : public FArchiveRemoteObjectReader
+{
+	const TMap<UObject*, const UObject*>& InstanceToArchetypeMap;
+	const TMap<const UObject*, UObject*>& ArchetypeToInstanceMap;
+	UObject* InstanceScope = nullptr;	
+
+public:
+
+	FArchetypeDeltaReader(FRemoteObjectData& InObjectData, ERemoteObjectSerializationFlags InDeserializeFlags, const TMap<UObject*, const UObject*>& InInstanceToArchetypeMap, const TMap<const UObject*, UObject*>& InArchetypeToInstanceMap)
+		: FArchiveRemoteObjectReader(InObjectData, /*MigrationContext*/ nullptr, InDeserializeFlags, TEXT("RemoteArchetypeDeltaReader"))
+		, InstanceToArchetypeMap(InInstanceToArchetypeMap)
+		, ArchetypeToInstanceMap(InArchetypeToInstanceMap)
+	{
+	}
+
+	virtual void SerializeRemoteObject(UObject* Object)
+	{
+		InstanceScope = Object;
+		Object->Serialize(*this);
+		InstanceScope = nullptr;
+	}
+
+protected:
+
+	virtual void SetObjectReference(FObjectPtr& CurrentValue, FObjectPtr NewValue) override
+	{
+		using namespace UE::CoreUObject::Private;
+		
+		// Helper lambda that extracts a raw object pointer from an FObjectPtr with resolving it.
+		auto ExtractObjectNoResolve = [](FObjectPtr& Ptr)
+			{
+				UObject* Result = nullptr;
+			#if UE_WITH_REMOTE_OBJECT_HANDLE
+				FObjectHandle Handle = Ptr.GetHandle();
+				if (Handle.ExtractHandleType() == ERemoteObjectHandleType::ObjectPointer)
+				{
+					Result = const_cast<UObject*>(Handle.ExtractObjectPointer());
+				}
+			#else
+				Result = Ptr.Get();
+			#endif
+				return Result;
+			};
+
+		UObject* NewObjectValue = ExtractObjectNoResolve(NewValue);
+		if (InstanceScope && NewObjectValue && NewObjectValue->IsTemplate())
+		{
+			// The reason why we use CurrentValueRemoteId to determine what the current value is that
+			// we don't really care about the current value pointer because only care if it's null (!IsValid())
+			// or not in which case we want to preserve the current value because it already points to an instance of its archetype
+			// which will be reset in its own Serialize call.
+			FRemoteObjectId CurrentValueRemoteId = CurrentValue.GetRemoteId();
+
+			// We're deserializing the delta on top of an existing instance and we're assigning a reference to a template object to a property in that instance
+			// The template reference is just a default subobject of the root object's archetype.
+			// There's a few options what the current reference (CurrentValue) is pointing to:			
+			if (CurrentValueRemoteId.IsValid())
+			{
+				// We have a valid, non-template (sub)object so preserve it. It will be reset in its own Serialize() call
+#if DO_CHECK
+				if (UObject* CurrentObjectValue = ExtractObjectNoResolve(CurrentValue))
+				{
+					checkf(!CurrentObjectValue->IsTemplate(), TEXT("Current value %s is a template when deserializing archetype delta. It should be instanced."), *CurrentObjectValue->GetPathName());
+				}
+#endif				
+				NewValue = CurrentValue;
+			}
+			else if (UObject* const* FoundInstance = ArchetypeToInstanceMap.Find(NewObjectValue))
+			{
+				// The current value is null (the value may have been nulled after construction by gameplay code) 
+				// but we have a unique mapping to an existing instance of the archetype reference so assign it to the current value
+				NewValue = *FoundInstance;
+			}
+			else
+			{
+				// The current value is null and we can't map the archetype to any of the known instances. This may be because
+				// the archetype -> value mapping was not unique (e.g. items in an array)
+				FStaticConstructObjectParameters Params(NewObjectValue->GetClass());
+				Params.Outer = InstanceScope;
+				Params.Template = NewObjectValue;
+				NewValue = StaticConstructObject_Internal(Params);
+			}
+		}
+		FArchiveRemoteObjectReader::SetObjectReference(CurrentValue, NewValue);
+	}
 };
 
 /**
@@ -1053,41 +1175,6 @@ public:
 	}
 };
 
-
-class FArchiveSubObjectGatherer : public FArchiveUObject
-{
-	UObject* RootObject = nullptr;
-	TMap<const UObject*, UObject*>& ArchetypeToObjectMap;
-	TArray<UObject*>& ObjectsToSerialize;
-
-public:
-
-	FArchiveSubObjectGatherer(UObject* InRootObject, TMap<const UObject*, UObject*>& InArchetypeToObjectMap, TArray<UObject*>& InObjectsToSerialize)
-		: RootObject(InRootObject)
-		, ArchetypeToObjectMap(InArchetypeToObjectMap)
-		, ObjectsToSerialize(InObjectsToSerialize)
-	{
-		ArchetypeToObjectMap.Add(FindArchetype(RootObject), RootObject);
-		ArIsObjectReferenceCollector = true;
-#if WITH_EDITORONLY_DATA
-		SetIsSaving(true); // workaround for UE-225585 (UStruct::SerializeBinEx() asserts in structured archive code in the editor)
-#endif
-		SetIsPersistent(false);
-		SetPortFlags(PPF_AvoidRemoteObjectMigration);
-	}
-
-	virtual FArchive& operator<<(FObjectPtr& Value) override
-	{
-		FRemoteObjectReferenceInfo Info = FArchiveRemoteObjectWriter::GetReferenceInfo(Value, RootObject);
-		if (Info.bIsSubobject)
-		{
-			ArchetypeToObjectMap.Add(FindArchetype(Info.Object), Info.Object);
-			ObjectsToSerialize.Add(Info.Object);
-		}
-		return *this;
-	}
-};
-
 FRemoteObjectConstructionOverrides::FRemoteObjectConstructionOverrides(const FRemoteObjectData& ObjectData, const TArray<FRemoteObjectHeader>& InObjectHeaders)
 {
 	using namespace UE::RemoteObject::Handle;
@@ -1221,15 +1308,6 @@ UObject* ConstructRemoteObject(const FRemoteObjectData& ObjectData, const FRemot
 	return Object;
 }
 
-void ResolvePathNames(const FRemoteObjectData& InObjectData, TArray<UObject*>& OutResolvedObjects)
-{	
-	for (const FPackedRemoteObjectPathName& PathName : InObjectData.PathNames)
-	{
-		UObject* Object = PathName.Resolve(InObjectData.Tables);
-		OutResolvedObjects.Add(Object);
-	}
-}
-
 FORCENOINLINE void SortRemoteObjectHeadersByOuter(FArchiveRemoteObjectWriter& Ar, TArray<FRemoteObjectHeader>& ObjectHeaders)
 {
 	using namespace UE::CoreUObject::Private;
@@ -1413,6 +1491,13 @@ FRemoteObjectData SerializeObjectData(UObject* InObject, TSet<UObject*>& OutObje
 	return MoveTemp(ObjectData);
 }
 
+int32 DeserializeObjectDataInternal(FArchiveRemoteObjectReader& Ar,
+	FRemoteObjectData& ObjectData,
+	const FUObjectMigrationContext* MigrationContext,
+	TArray<FRemoteObjectId>& OutObjectRemoteIds,
+	TArray<UObject*>& OutObjects,
+	ERemoteObjectSerializationFlags DeserializeFlags);
+
 void ResetRemoteObject(UObject* InObject)
 {
 	UObject* Object = InObject;
@@ -1421,53 +1506,72 @@ void ResetRemoteObject(UObject* InObject)
 		Object = FindCanonicalRootObjectForSerialization(Object);
 	}
 
-	TMap<const UObject*, UObject*> ReverseArchetypeToObjectMap;
+	TArray<UObject*> ObjectHierarchy;
+	TMap<UObject*, const UObject*> ObjectToArchetypeMap; // Maps unique instances to archetypes
+	TMap<const UObject*, UObject*> ArchetypeToObjectMap; // Maps unique archetypes to their instances
 
-	// Serialize Object to gather its subobjects and map the subobject archetypes to their respective instances
-	// This map will also be used to replace archetype object pathnames in the serialized archetype data to their instances' pathnames
+	// Gather all subobjects of the object to reset and map them to their respective archetypes
+	// This map will be used to replace serialized objects with the mapped archetypes (see FArchetypeDeltaWriter::SerializeRemoteObject())
+	// and to preserve references to instances when deserializing the delta between the archetypes and their respective instances (see FArchetypeDeltaReader::SetObjectReference())
 	{
-		TArray<UObject*> ObjectsToSerialize;
-		TSet<UObject*> SerializedObjects;
-		FArchiveSubObjectGatherer SubobjectGather(Object, ReverseArchetypeToObjectMap, ObjectsToSerialize);
-		ObjectsToSerialize.Add(Object);
+		ObjectHierarchy.Add(Object);
+		GetObjectsWithOuter(Object, ObjectHierarchy, /*bIncludeNestedObjects*/ true);
 
-		for (int32 Index = 0; Index < ObjectsToSerialize.Num(); ++Index)
+		TSet<const UObject*> NonUniqueArchetypes;
+		for (UObject* ObjectToReset : ObjectHierarchy)
 		{
-			UObject* Obj = ObjectsToSerialize[Index];
-			if (!SerializedObjects.Contains(Obj))
+			const UObject* FoundArchetype = FindArchetype(ObjectToReset);
+			ObjectToArchetypeMap.Add(ObjectToReset, FoundArchetype);
+			if (ArchetypeToObjectMap.Find(FoundArchetype))
 			{
-				Obj->Serialize(SubobjectGather);
-				SerializedObjects.Add(Obj);
+				NonUniqueArchetypes.Add(FoundArchetype);
 			}
+			else
+			{
+				ArchetypeToObjectMap.Add(FoundArchetype, ObjectToReset);
+			}
+		}
+		// Remove all archetypes that don't have unique mapping to their respective instances
+		// FArchetypeDeltaReader will create new instances of archetypes with non-unique mapping
+		for (const UObject* NotUniqueArchetype : NonUniqueArchetypes)
+		{
+			ArchetypeToObjectMap.Remove(NotUniqueArchetype);
 		}
+		
 	}
 
 	FRemoteObjectData* ArchetypeDelta = new FRemoteObjectData();
 	{
-		// Serialize the Object Archetype against the Object (and its subobjects)
-		// This will produce a delta between the archetypes and their instances which will then be used to restore the instances' state to the archetypes'
-		FRemoteObjectId RequestedObjectId(Object);
-		UObject* Archetype = const_cast<UObject*>(FindArchetype(Object));
-		FArchetypeDeltaWriter Ar(Archetype, *ArchetypeDelta, ReverseArchetypeToObjectMap);
-		TSet<UObject*> ArchetypeSerializedObjects;
-		SerializeObjectDataInternal(Ar, Archetype, RequestedObjectId, ArchetypeSerializedObjects, ERemoteObjectSerializationFlags::Resetting);
+		FRemoteObjectId RequestedObjectId(Object);		
+		FArchetypeDeltaWriter Ar(Object, *ArchetypeDelta, ObjectToArchetypeMap, ArchetypeToObjectMap);
+		for (UObject* ObjectToSerialize : ObjectHierarchy)
+		{
+			Ar.GetObjectsToSerialize().Add(ObjectToSerialize);
+		}
+
+		// Even though we pass the Object to reset to SerializeObjectDataInternal, the FArchetypeDeltaWriter will actually serialize the differences between the archetypes' data and their respective instances
+		// This will produce a delta between the archetypes and their instances which will then be used to restore the instances' state to the archetypes' default state
+		TSet<UObject*> SerializedObjects;
+		SerializeObjectDataInternal(Ar, Object, RequestedObjectId, SerializedObjects, ERemoteObjectSerializationFlags::Resetting | ERemoteObjectSerializationFlags::SkipCanonicalRootSearch);
 	}
 
 	{
-		// Deserialize archetype data on top of the object (and its subobjects) to restore their state to the archetype values
+		// Deserialize archetype data on top of the object to reset (and its subobjects) to restore its state to the archetype values
 		TArray<UObject*> DeserializedObjects;
 		TArray<FRemoteObjectId> DeserializedIds;
 		// When deserializing archetype delta we want to:
 		// Preserve any references to remote objects that have not been migrated yet (this is because we can't generate archetype delta for them because they don't exist on this server and we don't have their data)
-		// Use any existing objects and deserialize archetype delta on top of them avoiding re-construction which may lead to undiserable side-effects
+		// Use any existing objects and deserialize archetype delta on top of them avoiding re-construction which may lead to undesirable side-effects
 		// Additionally we don't want to recursively re-enter this function so let the deserialization process know we're already resetting migrated object(s)
 		ERemoteObjectSerializationFlags DeserializationFlags = ERemoteObjectSerializationFlags::PreserveRemoteReferences | ERemoteObjectSerializationFlags::UseExistingObjects | ERemoteObjectSerializationFlags::Resetting;
-		DeserializeObjectData(*ArchetypeDelta, /*MigrationContext*/ nullptr, DeserializedIds, DeserializedObjects, DeserializationFlags);
+		FArchetypeDeltaReader Ar(*ArchetypeDelta, DeserializationFlags, ObjectToArchetypeMap, ArchetypeToObjectMap);
+		DeserializeObjectDataInternal(Ar, *ArchetypeDelta, /*MigrationContext*/ nullptr, DeserializedIds, DeserializedObjects, DeserializationFlags);
 	}
 	delete ArchetypeDelta;
 }
 
-int32 DeserializeObjectData(FRemoteObjectData& ObjectData,
+int32 DeserializeObjectDataInternal(FArchiveRemoteObjectReader& Ar,
+							FRemoteObjectData& ObjectData,
 							const FUObjectMigrationContext* MigrationContext,
 							TArray<FRemoteObjectId>& OutObjectRemoteIds,
 							TArray<UObject*>& OutObjects,
@@ -1485,10 +1589,8 @@ int32 DeserializeObjectData(FRemoteObjectData& ObjectData,
 	int32 RequestedObjectIndex = -1;
 	const bool bResetting = EnumHasAnyFlags(DeserializeFlags, ERemoteObjectSerializationFlags::Resetting);
 
-	ResolvePathNames(ObjectData, ResolvedPathNameObjects);
-
 	// If we are being assigned ownership, we *must* take ownership of the objects.  Note: We also may already have ownership of those objects. 
-	FArchiveRemoteObjectReader Ar(ObjectData, ResolvedPathNameObjects, MigrationContext, DeserializeFlags);
+
 	// If we're calling this function to reset an object to its archetype state then we don't want FArchiveRemoteObjectReader to be marked as migrating remote objects (Which it is by default)
 	Ar.SetMigratingRemoteObjects(!bResetting);
 
@@ -1623,7 +1725,7 @@ int32 DeserializeObjectData(FRemoteObjectData& ObjectData,
 				{
 					FRemoteObjectReaderScope Scope(Ar, Object);
 					Ar.Seek(ObjectHeaders[ObjectIndex].StartOffset);
-					Object->Serialize(Ar);
+					Ar.SerializeRemoteObject(Object);
 					if (RequestedObjectIndex == -1 && RequestedObjectId == FRemoteObjectId(Object))
 					{
 						RequestedObjectIndex = ObjectIndex;
@@ -1651,6 +1753,16 @@ int32 DeserializeObjectData(FRemoteObjectData& ObjectData,
 	return RequestedObjectIndex;
 }
 
+int32 DeserializeObjectData(FRemoteObjectData& ObjectData,
+	const FUObjectMigrationContext* MigrationContext,
+	TArray<FRemoteObjectId>& OutObjectRemoteIds,
+	TArray<UObject*>& OutObjects,
+	ERemoteObjectSerializationFlags DeserializeFlags)
+{
+	FArchiveRemoteObjectReader Ar(ObjectData, MigrationContext, DeserializeFlags);
+	return DeserializeObjectDataInternal(Ar, ObjectData, MigrationContext, OutObjectRemoteIds, OutObjects, DeserializeFlags);
+}
+
 UObject* FindCanonicalRootObjectForSerialization(UObject* Object)
 {
 	// find the outermost migration root
