Added TargetPlatform to FExportsValidationContext, to match it being available in FObjectSaveContextData

#rnx

#Synced-CL 51332317

diff --git a/Engine/Source/Runtime/CoreUObject/Private/UObject/SavePackage2.cpp b/Engine/Source/Runtime/CoreUObject/Private/UObject/SavePackage2.cpp
--- a/Engine/Source/Runtime/CoreUObject/Private/UObject/SavePackage2.cpp
+++ b/Engine/Source/Runtime/CoreUObject/Private/UObject/SavePackage2.cpp
@@ -773,7 +773,7 @@ ESavePackageResult ValidateExports(FSaveContext& SaveContext)
 		FOutputDevice* OutputDevice = SaveContext.IsGenerateSaveError() ? SaveContext.GetError() : nullptr;
 		for (const TFunction<FSavePackageSettings::ExternalExportValidationFunc>& ValidateExport : SaveContext.GetExternalExportValidations())
 		{
-			SaveContext.Result = ValidateExport({ SaveContext.GetPackage(), Exports, SaveContext.GetObjectSaveContext().SaveOverrides, Flags, OutputDevice });
+			SaveContext.Result = ValidateExport({ SaveContext.GetPackage(), Exports, SaveContext.GetObjectSaveContext().SaveOverrides, Flags, SaveContext.GetTargetPlatform(), OutputDevice });
 			if (SaveContext.Result != ESavePackageResult::Success)
 			{
 				return SaveContext.Result;
diff --git a/Engine/Source/Runtime/CoreUObject/Public/UObject/SavePackage.h b/Engine/Source/Runtime/CoreUObject/Public/UObject/SavePackage.h
--- a/Engine/Source/Runtime/CoreUObject/Public/UObject/SavePackage.h
+++ b/Engine/Source/Runtime/CoreUObject/Public/UObject/SavePackage.h
@@ -153,11 +153,12 @@ struct FExportsValidationContext
 	};
 
 	FExportsValidationContext(const UPackage* InPackage, const TSet<UObject*>& InExports, const TMap<UObject*, FObjectSaveOverride>& InSaveOverrides, 
-		EFlags InFlags, FOutputDevice* InOutputDevice)
+		EFlags InFlags, const ITargetPlatform* InTargetPlatform, FOutputDevice* InOutputDevice)
 		: Package(InPackage)
 		, Exports(InExports)
 		, SaveOverrides(InSaveOverrides)
 		, Flags(InFlags)
+		, TargetPlatform(InTargetPlatform)
 		, OutputDevice(InOutputDevice)
 	{}
 
@@ -165,6 +166,7 @@ struct FExportsValidationContext
 	const TSet<UObject*>& Exports;
 	const TMap<UObject*, FObjectSaveOverride>& SaveOverrides;
 	const EFlags Flags;
+	const ITargetPlatform* TargetPlatform; // Null if not cooking
 	FOutputDevice* OutputDevice;
 };
 ENUM_CLASS_FLAGS(FExportsValidationContext::EFlags)
