Moved TIsTSubclassOf into its own header.

diff --git a/Engine/Source/Runtime/Core/Public/Traits/IsTSubclassOf.h b/Engine/Source/Runtime/Core/Public/Traits/IsTSubclassOf.h
new file mode 100644
--- /dev/null
+++ b/Engine/Source/Runtime/Core/Public/Traits/IsTSubclassOf.h
@@ -0,0 +1,20 @@
+// Copyright Epic Games, Inc. All Rights Reserved.
+
+#pragma once
+
+template <typename T>
+class TSubclassOf;
+
+template <typename T>
+inline constexpr bool TIsTSubclassOf_V = false;
+
+template <typename T> inline constexpr bool TIsTSubclassOf_V<               TSubclassOf<T>> = true;
+template <typename T> inline constexpr bool TIsTSubclassOf_V<const          TSubclassOf<T>> = true;
+template <typename T> inline constexpr bool TIsTSubclassOf_V<      volatile TSubclassOf<T>> = true;
+template <typename T> inline constexpr bool TIsTSubclassOf_V<const volatile TSubclassOf<T>> = true;
+
+template <typename T>
+struct TIsTSubclassOf
+{
+	static constexpr bool Value = TIsTSubclassOf_V<T>;
+};
diff --git a/Engine/Source/Runtime/CoreUObject/Public/Templates/SubclassOf.h b/Engine/Source/Runtime/CoreUObject/Public/Templates/SubclassOf.h
--- a/Engine/Source/Runtime/CoreUObject/Public/Templates/SubclassOf.h
+++ b/Engine/Source/Runtime/CoreUObject/Public/Templates/SubclassOf.h
@@ -2,26 +2,13 @@
 
 #pragma once
 
+#include "Traits/IsTSubclassOf.h"
 #include "UObject/Class.h"
 
 #include <type_traits>
 
 class FStructuredArchiveSlot;
 
-template <typename T>
-class TSubclassOf;
-
-template <typename T>
-struct TIsTSubclassOf
-{
-	enum { Value = false };
-};
-
-template <typename T> struct TIsTSubclassOf<               TSubclassOf<T>> { enum { Value = true }; };
-template <typename T> struct TIsTSubclassOf<const          TSubclassOf<T>> { enum { Value = true }; };
-template <typename T> struct TIsTSubclassOf<      volatile TSubclassOf<T>> { enum { Value = true }; };
-template <typename T> struct TIsTSubclassOf<const volatile TSubclassOf<T>> { enum { Value = true }; };
-
 /**
  * Template to allow UClass types to be passed around with type safety
  */
