[CoreUObject]
* Added checkSlow of object in garbage collector loop just to make static analysis happy. object can never be null here

#Synced-CL 51501514

diff --git a/Engine/Source/Runtime/CoreUObject/Private/UObject/GarbageCollection.cpp b/Engine/Source/Runtime/CoreUObject/Private/UObject/GarbageCollection.cpp
--- a/Engine/Source/Runtime/CoreUObject/Private/UObject/GarbageCollection.cpp
+++ b/Engine/Source/Runtime/CoreUObject/Private/UObject/GarbageCollection.cpp
@@ -960,6 +960,8 @@ public:
 				if (Object)
 #endif // UE_WITH_CONSTINIT_UOBJECT
 				{
+					checkSlow(Object); // This is here to make static analysis happy. Object can never be null here
+
 					int32 OffsetToObject = -Object->InternalIndex; // OffsetToObject == -OffsetToAllocation
 					Object->InternalIndex = INDEX_NONE;
 
