Performance improvements for async package loading:
- Cache hash indices to avoid redundant binary searches when handling export collisions.
- Only hash RF_Public exports or exports with imported outers, skips internal exports while still handling external actors correctly.

#Synced-CL 51495466

diff --git a/Engine/Source/Runtime/CoreUObject/Private/Serialization/AsyncLoading2.cpp b/Engine/Source/Runtime/CoreUObject/Private/Serialization/AsyncLoading2.cpp
--- a/Engine/Source/Runtime/CoreUObject/Private/Serialization/AsyncLoading2.cpp
+++ b/Engine/Source/Runtime/CoreUObject/Private/Serialization/AsyncLoading2.cpp
@@ -1209,14 +1209,17 @@ private:
 			}
 		}
 
-		void Store(uint64 ExportHash, int32 ObjectIndex)
+		void Store(uint64 ExportHash, int32 ObjectIndex, int32 Index = INDEX_NONE)
 		{
 			checkf(ExportHash, TEXT("Invalid to store a hash of 0 in this map"));
 			checkf(ObjectIndex != InvalidValue, TEXT("Use Remove() instead"));
 
 			TArrayView<uint64> Keys = GetKeys();
 			TArrayView<int32> Values = GetValues();
-			int32 Index = Algo::LowerBound(Keys, ExportHash);
+			if (Index == INDEX_NONE)
+			{
+				Index = Algo::LowerBound(Keys, ExportHash);
+			}
 			if (Index < Capacity && Keys[Index] == ExportHash)
 			{
 				// Slot already exists so reuse it
@@ -1271,11 +1274,15 @@ private:
 			return false;
 		}
 
-		// Returns the value stored if any, returing InvalidValue for not present or removed values
-		int32 Find(uint64 ExportHash)
+		// Returns the value stored if any, returning InvalidValue for not present or removed values
+		int32 Find(uint64 ExportHash, int32* OutHashIndex = nullptr)
 		{
 			TArrayView<uint64> Keys = GetKeys();
 			int32 Index = Algo::LowerBound(Keys, ExportHash);
+			if (OutHashIndex)
+			{
+				*OutHashIndex = Index;
+			}
 			if (Index < Capacity && Keys[Index] == ExportHash)
 			{
 				TArrayView<int32> Values = GetValues();
@@ -1695,10 +1702,10 @@ public:
 		PublicExportMap.PreInsertPublicExports(Exports);
 	}
 
-	void StorePublicExport(uint64 ExportHash, UObject* Object)
+	void StorePublicExport(uint64 ExportHash, int32 HashIndex, UObject* Object)
 	{
 		check(Mutex.IsLocked());
-		PublicExportMap.Store(ExportHash, GUObjectArray.ObjectToIndex(Object));
+		PublicExportMap.Store(ExportHash, GUObjectArray.ObjectToIndex(Object), HashIndex);
 	}
 
 	void RemovePublicExport(uint64 ExportHash, FName ObjectName = NAME_None)
@@ -1714,10 +1721,10 @@ public:
 		}
 	}
 
-	UObject* GetPublicExport(uint64 ExportHash)
+	UObject* GetPublicExport(uint64 ExportHash, int32* OutHashIndex = nullptr)
 	{
 		check(Mutex.IsLocked() || Mutex.IsLockShared());
-		int32 ObjectIndex = PublicExportMap.Find(ExportHash);
+		int32 ObjectIndex = PublicExportMap.Find(ExportHash, OutHashIndex);
 		if (ObjectIndex >= 0)
 		{
 			return static_cast<UObject*>(GUObjectArray.IndexToObject(ObjectIndex)->GetObject());
@@ -2638,7 +2645,8 @@ public:
 		{
 			FWriteLockedPackageRef PackageRef = FindPackageRefCheckedForWrite(Key.GetPackageId());
 
-			UObject* ExistingObject = PackageRef->GetPublicExport(Key.GetExportHash());
+			int32 HashIndex;
+			UObject* ExistingObject = PackageRef->GetPublicExport(Key.GetExportHash(), &HashIndex);
 			if (ExistingObject && ExistingObject != Object)
 			{
 				int32 ExistingObjectIndex = GUObjectArray.ObjectToIndex(ExistingObject);
@@ -2667,7 +2675,7 @@ public:
 				}
 			) > 0;
 
-			PackageRef->StorePublicExport(ExportHash, Object);
+			PackageRef->StorePublicExport(ExportHash, HashIndex, Object);
 			ObjectIndexToPublicExport.Add(ObjectIndex, Key);
 		}
 
@@ -7744,7 +7752,7 @@ EEventLoadNodeExecutionResult FAsyncPackage2::ProcessLinkerLoadPackageSummary(FA
 			LinkerLoadState->bContainsClasses |= true;
 		}
 
-		//if ((ObjectExport.ObjectFlags & RF_Public) > 0) // We need hashes for all objects, external actors are breaking assumptions about public exports
+		if ((ObjectExport.ObjectFlags & RF_Public) || ObjectExport.OuterIndex.IsImport())
 		{
 			TArray<int32, TInlineAllocator<128>> FullPath;
 			int32 PathIndex = ExportIndex;
@@ -13334,19 +13342,34 @@ bool FPublicExportMapBasicTest::RunTest(const FString& Parameters)
 	typedef FLoadedPackageRef::FPublicExportMap FPublicExportMap;
 
 	FPublicExportMap Map;
+	int32 HashIndex = INDEX_NONE;
 	TestEqual(TEXT("Empty map has no count"), Map.Num(), 0);
 	TestEqual(TEXT("Empty map has no keys"), Map.GetKeys().Num(), 0);
 	TestEqual(TEXT("Empty map has no values"), Map.GetValues().Num(), 0);
 
 	Map.Store(1, 2);
-	TestEqual(TEXT("Find stored key"), Map.Find(1), 2);
+	TestEqual(TEXT("Find stored key"), Map.Find(1, &HashIndex), 2);
+	TestBetweenInclusive(TEXT("Find stored key hash index"), HashIndex, 0, Map.GetValues().Num());
 	TestEqual(TEXT("Store increased size"), Map.Num(), 1);
 
 	// May resize
 	Map.Store(2, 3);
-	TestEqual(TEXT("Find stored key"), Map.Find(2), 3);
+	HashIndex = INDEX_NONE;
+	TestEqual(TEXT("Find stored key"), Map.Find(2, &HashIndex), 3);
+	TestBetweenInclusive(TEXT("Find stored key hash index"), HashIndex, 0, Map.GetValues().Num());
 	TestEqual(TEXT("Store increased size"), Map.Num(), 2);
 
+	HashIndex = INDEX_NONE;
+	TestEqual(TEXT("Find removed key fails"), Map.Find(3, &HashIndex), FPublicExportMap::InvalidValue);
+	Map.Store(3,4,HashIndex);
+	int32 NewHashIndex = INDEX_NONE;
+	TestEqual(TEXT("Find stored key"), Map.Find(3, &NewHashIndex), 4);
+	TestEqual(TEXT("Same hash index"), NewHashIndex, HashIndex);
+	TestBetweenInclusive(TEXT("Find stored key hash index"), NewHashIndex, 0, Map.GetValues().Num());
+	TestEqual(TEXT("Store increased size"), Map.Num(), 3);
+	TestEqual(TEXT("Remove existing key"), Map.Remove(3), true);
+	TestEqual(TEXT("Removed decreased count"), Map.Num(), 2);
+
 	TestEqual(TEXT("Remove existing key"), Map.Remove(1), true);
 	TestEqual(TEXT("Removed decreased count"), Map.Num(), 1);
 	TestEqual(TEXT("Remove removed key fails"), Map.Remove(1), false);
@@ -13367,6 +13390,7 @@ bool FPublicExportMapIterationTest::RunTest(const FString& Parameters)
 	typedef FLoadedPackageRef::FPublicExportMap FPublicExportMap;
 
 	FPublicExportMap Map;
+	int32 HashIndex = INDEX_NONE;
 	TArray<uint64> Keys = { 1, 2, 3, 4, 5, 6, 7, 8 };
 	TArray<int32> Values = { 1, 2, 3, 4, 5, 6, 7, 8 };
 	check(Keys.Num() == Values.Num());
@@ -13402,7 +13426,8 @@ bool FPublicExportMapIterationTest::RunTest(const FString& Parameters)
 
 	for (uint64 Key : Keys)
 	{
-		TestEqual(TEXT("Lookup works"), Map.Find(Key), Values[(int32)(Key - 1)]);
+		TestEqual(TEXT("Lookup works"), Map.Find(Key, &HashIndex), Values[(int32)(Key - 1)]);
+		TestBetweenInclusive(TEXT("Find stored key hash index"), HashIndex, 0, Map.GetValues().Num());
 	}
 	return true;
 }
@@ -13415,6 +13440,7 @@ bool FPublicExportMapPreInsertPublicExportsTest::RunTest(const FString& Paramete
 	// Empty map
 	{
 		FPublicExportMap Map;
+		int32 HashIndex = INDEX_NONE;
 		TArray<FExportMapEntry> Exports = { {.PublicExportHash = 1 }, {.PublicExportHash = 2 }, {.PublicExportHash = 3 }, {.PublicExportHash = 4 } };
 		int32 OriginaMapCount = Map.Num();
 
@@ -13435,13 +13461,15 @@ bool FPublicExportMapPreInsertPublicExportsTest::RunTest(const FString& Paramete
 		// Overwrite preinserted value
 		TestEqual(TEXT("Find stored key with no value set"), Map.Find(Exports[0].PublicExportHash), FPublicExportMap::InvalidValue);
 		Map.Store(Exports[0].PublicExportHash, (int32)Exports[0].PublicExportHash);
-		TestEqual(TEXT("Find stored key"), Map.Find(Exports[0].PublicExportHash), (int32)Exports[0].PublicExportHash);
+		TestEqual(TEXT("Find stored key"), Map.Find(Exports[0].PublicExportHash, &HashIndex), (int32)Exports[0].PublicExportHash);
+		TestBetweenInclusive(TEXT("Find stored key hash index"), HashIndex, 0, Map.GetValues().Num());
 		TestEqual(TEXT("Store increased count"), Map.Num(), OriginaMapCount + 1);
 
 		// Insert new value
 		TestEqual(TEXT("Find missing key"), Map.Find(111), FPublicExportMap::InvalidValue);
 		Map.Store(111, 111);
-		TestEqual(TEXT("Find stored key"), Map.Find(111), 111);
+		TestEqual(TEXT("Find stored key"), Map.Find(111, &HashIndex), 111);
+		TestBetweenInclusive(TEXT("Find stored key hash index"), HashIndex, 0, Map.GetValues().Num());
 		TestEqual(TEXT("Store increased count"), Map.Num(), OriginaMapCount + 2);
 
 		TestEqual(TEXT("Remove existing key"), Map.Remove(111), true);
@@ -13454,6 +13482,7 @@ bool FPublicExportMapPreInsertPublicExportsTest::RunTest(const FString& Paramete
 	// Empty map, populate a single key
 	{
 		FPublicExportMap Map;
+		int32 HashIndex = INDEX_NONE;
 		TArray<FExportMapEntry> Exports = { {.PublicExportHash = 1 } };
 		int32 OriginaMapCount = Map.Num();
 
@@ -13474,13 +13503,15 @@ bool FPublicExportMapPreInsertPublicExportsTest::RunTest(const FString& Paramete
 		// Overwrite preinserted value
 		TestEqual(TEXT("Find stored key with no value set"), Map.Find(Exports[0].PublicExportHash), FPublicExportMap::InvalidValue);
 		Map.Store(Exports[0].PublicExportHash, (int32)Exports[0].PublicExportHash);
-		TestEqual(TEXT("Find stored key"), Map.Find(Exports[0].PublicExportHash), (int32)Exports[0].PublicExportHash);
+		TestEqual(TEXT("Find stored key"), Map.Find(Exports[0].PublicExportHash, &HashIndex), (int32)Exports[0].PublicExportHash);
+		TestBetweenInclusive(TEXT("Find stored key hash index"), HashIndex, 0, Map.GetValues().Num());
 		TestEqual(TEXT("Store increased count"), Map.Num(), OriginaMapCount + 1);
 
 		// Insert new value
 		TestEqual(TEXT("Find missing key"), Map.Find(111), FPublicExportMap::InvalidValue);
 		Map.Store(111, 111);
-		TestEqual(TEXT("Find stored key"), Map.Find(111), 111);
+		TestEqual(TEXT("Find stored key"), Map.Find(111, &HashIndex), 111);
+		TestBetweenInclusive(TEXT("Find stored key hash index"), HashIndex, 0, Map.GetValues().Num());
 		TestEqual(TEXT("Store increased count"), Map.Num(), OriginaMapCount + 2);
 
 		TestEqual(TEXT("Remove existing key"), Map.Remove(111), true);
@@ -13490,10 +13521,10 @@ bool FPublicExportMapPreInsertPublicExportsTest::RunTest(const FString& Paramete
 		TestEqual(TEXT("Find removed key returns no value set"), Map.Find(111), FPublicExportMap::InvalidValue);
 	}
 
-
 	// One existing entry (preinsert multiple keys, some overlapping)
 	{
 		FPublicExportMap Map;
+		int32 HashIndex = INDEX_NONE;
 		TArray<FExportMapEntry> Exports = { {.PublicExportHash = 1 }, {.PublicExportHash = 2 }, {.PublicExportHash = 3 }, {.PublicExportHash = 4 } };
 
 		uint64 PreExistingKey = 2;
@@ -13501,7 +13532,8 @@ bool FPublicExportMapPreInsertPublicExportsTest::RunTest(const FString& Paramete
 		Map.Store(PreExistingKey, PreExistingValue);
 		int32 OriginaMapCount = Map.Num();
 
-		TestEqual(TEXT("Find stored key"), Map.Find(PreExistingKey), PreExistingValue);
+		TestEqual(TEXT("Find stored key"), Map.Find(PreExistingKey, &HashIndex), PreExistingValue);
+		TestBetweenInclusive(TEXT("Find stored key hash index"), HashIndex, 0, Map.GetValues().Num());
 		TestEqual(TEXT("Store increased count"), OriginaMapCount, 1);
 		TestEqual(TEXT("Store increased keys"), Map.GetKeys().Num(), OriginaMapCount);
 		TestEqual(TEXT("Store increased values"), Map.GetValues().Num(), OriginaMapCount);
@@ -13517,7 +13549,7 @@ bool FPublicExportMapPreInsertPublicExportsTest::RunTest(const FString& Paramete
 			if (Key == PreExistingKey)
 			{
 				bFoundPreExistingKey = true;
-				TestEqual(TEXT("Pre-existing keys remain intact after PreInsertPublicExports"), Map.Find(Key), PreExistingValue);
+				TestEqual(TEXT("Pre-existing keys remain intact after PreInsertPublicExports"), Map.Find(Key, &HashIndex), PreExistingValue);
 			}
 			else
 			{
@@ -13530,13 +13562,15 @@ bool FPublicExportMapPreInsertPublicExportsTest::RunTest(const FString& Paramete
 		// Overwrite preinserted value
 		TestEqual(TEXT("Find stored key with no value set"), Map.Find(Exports[0].PublicExportHash), FPublicExportMap::InvalidValue);
 		Map.Store(Exports[0].PublicExportHash, (int32)Exports[0].PublicExportHash);
-		TestEqual(TEXT("Find stored key"), Map.Find(Exports[0].PublicExportHash), (int32)Exports[0].PublicExportHash);
+		TestEqual(TEXT("Find stored key"), Map.Find(Exports[0].PublicExportHash, &HashIndex), (int32)Exports[0].PublicExportHash);
+		TestBetweenInclusive(TEXT("Find stored key hash index"), HashIndex, 0, Map.GetValues().Num());
 		TestEqual(TEXT("Store increased count"), Map.Num(), OriginaMapCount + 1);
 
 		// Insert new value
 		TestEqual(TEXT("Find missing key"), Map.Find(111), FPublicExportMap::InvalidValue);
 		Map.Store(111, 111);
-		TestEqual(TEXT("Find stored key"), Map.Find(111), 111);
+		TestEqual(TEXT("Find stored key"), Map.Find(111, &HashIndex), 111);
+		TestBetweenInclusive(TEXT("Find stored key hash index"), HashIndex, 0, Map.GetValues().Num());
 		TestEqual(TEXT("Store increased count"), Map.Num(), OriginaMapCount + 2);
 
 		TestEqual(TEXT("Remove existing key"), Map.Remove(111), true);
@@ -13549,6 +13583,7 @@ bool FPublicExportMapPreInsertPublicExportsTest::RunTest(const FString& Paramete
 	// One existing entry (preinsert a single entry overlapping)
 	{
 		FPublicExportMap Map;
+		int32 HashIndex = INDEX_NONE;
 		TArray<FExportMapEntry> Exports = { {.PublicExportHash = 1 } };
 
 		uint64 PreExistingKey = 1;
@@ -13556,7 +13591,8 @@ bool FPublicExportMapPreInsertPublicExportsTest::RunTest(const FString& Paramete
 		Map.Store(PreExistingKey, PreExistingValue);
 		int32 OriginaMapCount = Map.Num();
 
-		TestEqual(TEXT("Find stored key"), Map.Find(PreExistingKey), PreExistingValue);
+		TestEqual(TEXT("Find stored key"), Map.Find(PreExistingKey, &HashIndex), PreExistingValue);
+		TestBetweenInclusive(TEXT("Find stored key hash index"), HashIndex, 0, Map.GetValues().Num());
 		TestEqual(TEXT("Store increased count"), OriginaMapCount, 1);
 		TestEqual(TEXT("Store increased keys"), Map.GetKeys().Num(), OriginaMapCount);
 		TestEqual(TEXT("Store increased values"), Map.GetValues().Num(), OriginaMapCount);
@@ -13572,7 +13608,8 @@ bool FPublicExportMapPreInsertPublicExportsTest::RunTest(const FString& Paramete
 			if (Key == PreExistingKey)
 			{
 				bFoundPreExistingKey = true;
-				TestEqual(TEXT("Pre-existing keys remain intact after PreInsertPublicExports"), Map.Find(Key), PreExistingValue);
+				TestEqual(TEXT("Pre-existing keys remain intact after PreInsertPublicExports"), Map.Find(Key, &HashIndex), PreExistingValue);
+				TestBetweenInclusive(TEXT("Find stored key hash index"), HashIndex, 0, Map.GetValues().Num());
 			}
 			else
 			{
@@ -13586,6 +13623,7 @@ bool FPublicExportMapPreInsertPublicExportsTest::RunTest(const FString& Paramete
 	// One existing entry (preinsert a single entry non-overlapping)
 	{
 		FPublicExportMap Map;
+		int32 HashIndex = INDEX_NONE;
 		TArray<FExportMapEntry> Exports = { {.PublicExportHash = 1 } };
 
 		uint64 PreExistingKey = 2;
@@ -13593,7 +13631,8 @@ bool FPublicExportMapPreInsertPublicExportsTest::RunTest(const FString& Paramete
 		Map.Store(PreExistingKey, PreExistingValue);
 		int32 OriginaMapCount = Map.Num();
 
-		TestEqual(TEXT("Find stored key"), Map.Find(PreExistingKey), PreExistingValue);
+		TestEqual(TEXT("Find stored key"), Map.Find(PreExistingKey, &HashIndex), PreExistingValue);
+		TestBetweenInclusive(TEXT("Find stored key hash index"), HashIndex, 0, Map.GetValues().Num());
 		TestEqual(TEXT("Store increased count"), OriginaMapCount, 1);
 		TestEqual(TEXT("Store increased keys"), Map.GetKeys().Num(), OriginaMapCount);
 		TestEqual(TEXT("Store increased values"), Map.GetValues().Num(), OriginaMapCount);
@@ -13609,7 +13648,8 @@ bool FPublicExportMapPreInsertPublicExportsTest::RunTest(const FString& Paramete
 			if (Key == PreExistingKey)
 			{
 				bFoundPreExistingKey = true;
-				TestEqual(TEXT("Pre-existing keys remain intact after PreInsertPublicExports"), Map.Find(Key), PreExistingValue);
+				TestEqual(TEXT("Pre-existing keys remain intact after PreInsertPublicExports"), Map.Find(Key, &HashIndex), PreExistingValue);
+				TestBetweenInclusive(TEXT("Find stored key hash index"), HashIndex, 0, Map.GetValues().Num());
 			}
 			else
 			{
@@ -13622,13 +13662,15 @@ bool FPublicExportMapPreInsertPublicExportsTest::RunTest(const FString& Paramete
 		// Overwrite preinserted value
 		TestEqual(TEXT("Find stored key with no value set"), Map.Find(Exports[0].PublicExportHash), FPublicExportMap::InvalidValue);
 		Map.Store(Exports[0].PublicExportHash, (int32)Exports[0].PublicExportHash);
-		TestEqual(TEXT("Find stored key"), Map.Find(Exports[0].PublicExportHash), (int32)Exports[0].PublicExportHash);
+		TestEqual(TEXT("Find stored key"), Map.Find(Exports[0].PublicExportHash, &HashIndex), (int32)Exports[0].PublicExportHash);
+		TestBetweenInclusive(TEXT("Find stored key hash index"), HashIndex, 0, Map.GetValues().Num());
 		TestEqual(TEXT("Store increased count"), Map.Num(), OriginaMapCount + 1);
 	}
 
 	// One existing entry (preinsert multiple keys, non-overlapping)
 	{
 		FPublicExportMap Map;
+		int32 HashIndex = INDEX_NONE;
 		TArray<FExportMapEntry> Exports = { {.PublicExportHash = 1 }, {.PublicExportHash = 2 }, {.PublicExportHash = 3 }, {.PublicExportHash = 4 } };
 
 		uint64 PreExistingKey = 7;
@@ -13636,7 +13678,8 @@ bool FPublicExportMapPreInsertPublicExportsTest::RunTest(const FString& Paramete
 		Map.Store(PreExistingKey, PreExistingValue);
 		int32 OriginaMapCount = Map.Num();
 
-		TestEqual(TEXT("Find stored key"), Map.Find(PreExistingKey), PreExistingValue);
+		TestEqual(TEXT("Find stored key"), Map.Find(PreExistingKey, &HashIndex), PreExistingValue);
+		TestBetweenInclusive(TEXT("Find stored key hash index"), HashIndex, 0, Map.GetValues().Num());
 		TestEqual(TEXT("Store increased count"), OriginaMapCount, 1);
 		TestEqual(TEXT("Store increased keys"), Map.GetKeys().Num(), OriginaMapCount);
 		TestEqual(TEXT("Store increased values"), Map.GetValues().Num(), OriginaMapCount);
@@ -13651,7 +13694,8 @@ bool FPublicExportMapPreInsertPublicExportsTest::RunTest(const FString& Paramete
 			if (Key == PreExistingKey)
 			{
 				bFoundPreExistingKey = true;
-				TestEqual(TEXT("Pre-existing keys remain intact after PreInsertPublicExports"), Map.Find(Key), PreExistingValue);
+				TestEqual(TEXT("Pre-existing keys remain intact after PreInsertPublicExports"), Map.Find(Key, &HashIndex), PreExistingValue);
+				TestBetweenInclusive(TEXT("Find stored key hash index"), HashIndex, 0, Map.GetValues().Num());
 			}
 			else
 			{
@@ -13664,13 +13708,15 @@ bool FPublicExportMapPreInsertPublicExportsTest::RunTest(const FString& Paramete
 		// Overwrite preinserted value
 		TestEqual(TEXT("Find stored key with no value set"), Map.Find(Exports[0].PublicExportHash), FPublicExportMap::InvalidValue);
 		Map.Store(Exports[0].PublicExportHash, (int32)Exports[0].PublicExportHash);
-		TestEqual(TEXT("Find stored key"), Map.Find(Exports[0].PublicExportHash), (int32)Exports[0].PublicExportHash);
+		TestEqual(TEXT("Find stored key"), Map.Find(Exports[0].PublicExportHash, &HashIndex), (int32)Exports[0].PublicExportHash);
+		TestBetweenInclusive(TEXT("Find stored key hash index"), HashIndex, 0, Map.GetValues().Num());
 		TestEqual(TEXT("Store increased count"), Map.Num(), OriginaMapCount + 1);
 
 		// Insert new value
 		TestEqual(TEXT("Find missing key"), Map.Find(111), FPublicExportMap::InvalidValue);
 		Map.Store(111, 111);
-		TestEqual(TEXT("Find stored key"), Map.Find(111), 111);
+		TestEqual(TEXT("Find stored key"), Map.Find(111, &HashIndex), 111);
+		TestBetweenInclusive(TEXT("Find stored key hash index"), HashIndex, 0, Map.GetValues().Num());
 		TestEqual(TEXT("Store increased count"), Map.Num(), OriginaMapCount + 2);
 
 		TestEqual(TEXT("Remove existing key"), Map.Remove(111), true);
@@ -13683,6 +13729,7 @@ bool FPublicExportMapPreInsertPublicExportsTest::RunTest(const FString& Paramete
 	// Many pre-existing keys
 	{
 		FPublicExportMap Map;
+		int32 HashIndex = INDEX_NONE;
 		// Note we are explicitly overlapping key 1 with the preinsert key 1
 		// Note we are mixing ranges of preinserted and preexisting keys so we are forced to sort the keys together
 		uint64 OverlappingKey = 1;
@@ -13719,11 +13766,13 @@ bool FPublicExportMapPreInsertPublicExportsTest::RunTest(const FString& Paramete
 			if (PreExistingKeys.Contains(Key))
 			{
 				FoundPreExistingKeys.Add(Key);
-				TestEqual(TEXT("Pre-existing keys remain intact after PreInsertPublicExports"), Map.Find(Key), PreExistingValues[PreExistingKeys.IndexOfByKey(Key)]);
+				TestEqual(TEXT("Pre-existing keys remain intact after PreInsertPublicExports"), Map.Find(Key, &HashIndex), PreExistingValues[PreExistingKeys.IndexOfByKey(Key)]);
+				TestBetweenInclusive(TEXT("Find stored key hash index"), HashIndex, 0, Map.GetValues().Num());
 			}
 			else if (Key == OverlappingKey)
 			{
-				TestEqual(TEXT("Ensure preinsert didn't overwrite the existing value"), Map.Find(Key), (int32)OverlappingKey);
+				TestEqual(TEXT("Ensure preinsert didn't overwrite the existing value"), Map.Find(Key, &HashIndex), (int32)OverlappingKey);
+				TestBetweenInclusive(TEXT("Find stored key hash index"), HashIndex, 0, Map.GetValues().Num());
 			}
 			else
 			{
@@ -13736,13 +13785,15 @@ bool FPublicExportMapPreInsertPublicExportsTest::RunTest(const FString& Paramete
 		// Overwrite preinserted value
 		TestEqual(TEXT("Find stored key with no value set"), Map.Find(Exports[0].PublicExportHash), FPublicExportMap::InvalidValue);
 		Map.Store(Exports[0].PublicExportHash, (int32)Exports[0].PublicExportHash);
-		TestEqual(TEXT("Find stored key"), Map.Find(Exports[0].PublicExportHash), (int32)Exports[0].PublicExportHash);
+		TestEqual(TEXT("Find stored key"), Map.Find(Exports[0].PublicExportHash, &HashIndex), (int32)Exports[0].PublicExportHash);
+		TestBetweenInclusive(TEXT("Find stored key hash index"), HashIndex, 0, Map.GetValues().Num());
 		TestEqual(TEXT("Store increased count"), Map.Num(), PreExistingValues.Num() + 1);
 
 		// Insert new value
 		TestEqual(TEXT("Find missing key"), Map.Find(111), FPublicExportMap::InvalidValue);
 		Map.Store(111, 111);
-		TestEqual(TEXT("Find stored key"), Map.Find(111), 111);
+		TestEqual(TEXT("Find stored key"), Map.Find(111, &HashIndex), 111);
+		TestBetweenInclusive(TEXT("Find stored key hash index"), HashIndex, 0, Map.GetValues().Num());
 		TestEqual(TEXT("Store increased count"), Map.Num(), PreExistingValues.Num() + 2);
 
 		TestEqual(TEXT("Remove existing key"), Map.Remove(111), true);
@@ -13754,4 +13805,4 @@ bool FPublicExportMapPreInsertPublicExportsTest::RunTest(const FString& Paramete
 	return true;
 }
 
-#endif // WITH_LOW_LEVEL_TESTS
+#endif // WITH_LOW_LEVEL_TESTS
\ No newline at end of file
