Allow the linker to proceed in UE_INTERNAL_UNIVERSAL_PLACEHOLDERS environs when a new custom version is encountered

#rnx
#Synced-CL 51516599

diff --git a/Engine/Source/Runtime/CoreUObject/Private/UObject/LinkerLoad.cpp b/Engine/Source/Runtime/CoreUObject/Private/UObject/LinkerLoad.cpp
--- a/Engine/Source/Runtime/CoreUObject/Private/UObject/LinkerLoad.cpp
+++ b/Engine/Source/Runtime/CoreUObject/Private/UObject/LinkerLoad.cpp
@@ -1778,7 +1778,12 @@ FLinkerLoad::ELinkerStatus FLinkerLoad::UpdateFromPackageFileSummary()
 			else if (Diff.Type == ECustomVersionDifference::Newer)
 			{
 				FCustomVersion LatestVersion = FCurrentCustomVersions::Get(Diff.Version->Key).GetValue();
-
+#if UE_INTERNAL_UNIVERSAL_PLACEHOLDERS
+				FString AppName = FApp::GetName();
+				UE_ASSET_LOG(LogLinker, Error, PackagePath, TEXT("Package was saved with a newer custom version than %s. Rebuild %s to avoid this error - Tag %s Name '%s' PackageVersion %d  MaxExpected %d"),
+					*AppName, *AppName,
+					*Diff.Version->Key.ToString(), *LatestVersion.GetFriendlyName().ToString(), Diff.Version->Version, LatestVersion.Version);
+#else
 				// Loading a package with a newer custom version than the current one.
 				UE_ASSET_LOG(LogLinker, Error, PackagePath, TEXT("Package was saved with a newer custom version than the current. Tag %s Name '%s' PackageVersion %d  MaxExpected %d"),
 					*Diff.Version->Key.ToString(), *LatestVersion.GetFriendlyName().ToString(), Diff.Version->Version, LatestVersion.Version);
@@ -1794,6 +1799,7 @@ FLinkerLoad::ELinkerStatus FLinkerLoad::UpdateFromPackageFileSummary()
 				);
 
 				return LINKER_Failed;
+#endif // UE_INTERNAL_UNIVERSAL_PLACEHOLDERS
 			}
 		}
 	}
