[verse] predicts: multi-client replication support

also:
- faster predicts-function lookup (O(1), used to be O(n))
- cache native-owned object net id's during replication to minimize # of calls to native handlers
- make FDynamicallyTypedValue return NullType when uninitialized to avoid GC hazard
- add more tests + split off some test groups out of Predicts.versetest into their own files

#Synced-CL 51488275

diff --git a/Engine/Source/Runtime/CoreUObject/Public/UObject/DynamicallyTypedValue.h b/Engine/Source/Runtime/CoreUObject/Public/UObject/DynamicallyTypedValue.h
--- a/Engine/Source/Runtime/CoreUObject/Public/UObject/DynamicallyTypedValue.h
+++ b/Engine/Source/Runtime/CoreUObject/Public/UObject/DynamicallyTypedValue.h
@@ -123,7 +123,10 @@ namespace UE
 		void* GetDataPointer() { return IsInline() ? &InlineData : HeapData; }
 
 		// Returns the value's type.
-		FDynamicallyTypedValueType& GetType() const { return *Type; }
+		FDynamicallyTypedValueType& GetType() const
+		{
+			return Type ? *Type : NullType();
+		}
 
 		// Sets the value to the null state.
 		void SetToNull()
