Fix cut&paste support for overridable maps/arrays inside structs; improve sub-object shadow handling in overridable containers

Extended the overridable serialization system to properly handle cut&paste (export/import text) for overridable arrays and maps that are nested inside structs. The previous code skipped struct properties during text import, causing the operation to silently fail.

Key changes:

- Property.cpp: Added struct-property handling in ImportSingleProperty ? when the overridable container element is a FStructProperty, recursively processes its sub-properties instead of falling through.
- UnrealExporter.cpp: Generalized the properties export function to operate on UStruct* instead of UClass*, and added a PropertyPrefix parameter to support dot-notation paths for struct-nested properties (e.g., "Struct.PropertyName"). Also adds ShadowIndices tracking for instanced sub-object arrays.
- OverridableManager.cpp: Fixed HasOverrides detection to look up nodes via the full FArchiveSerializedPropertyChain (from context path) rather than a bare property ID. Extended the check to cover Add operations in addition to Remove when determining if an overridable container has pending overrides.
- OverriddenPropertySet.cpp: Replaced hard checkf crashes with ensureMsgf + graceful return in the array/map add and remove implementations to prevent data loss when an index is unexpectedly out of range.
- OverriddenPropertySet.h: Promoted GetOverriddenPropertyOperation(PortFlags, Chain, Property, DataPtr, DefaultValue) to public API (COREUOBJECT_API).
- OverridableSerializationTestUtils.h: Added FOverridableSubStruct and extended FOverridableStruct with nested overridable arrays, maps, and instanced sub-object containers to enable struct-level test coverage.
- OverridableSerializationTest.cpp: Fixed failing GetOverridableState() assertions and restructured the SubObjectOverridableMap test to correctly validate cut&paste and base-class-change propagation.

[RN] Copy& paste support of overridable containers inside struct
#Synced-CL 51456614

diff --git a/Engine/Source/Runtime/CoreUObject/Private/UObject/OverridableManager.cpp b/Engine/Source/Runtime/CoreUObject/Private/UObject/OverridableManager.cpp
--- a/Engine/Source/Runtime/CoreUObject/Private/UObject/OverridableManager.cpp
+++ b/Engine/Source/Runtime/CoreUObject/Private/UObject/OverridableManager.cpp
@@ -1,6 +1,7 @@
 // Copyright Epic Games, Inc. All Rights Reserved.
 
 #include "UObject/OverridableManager.h"
+#include "Serialization/ArchiveSerializedPropertyChain.h"
 #include "UObject/InstanceDataObjectUtils.h"
 #include "UObject/UObjectGlobals.h"
 #include "UObject/PropertyNames.h"
@@ -273,13 +274,14 @@ EOverriddenState FOverridableManager::GetOverriddenState(TNotNull<UObject*> Obje
 			{
 				if (Operation == EOverriddenPropertyOperation::Modified)
 				{
-					// We will walk through all sub-properties below, except for removed elements, so we test for those here.
-					FOverriddenPropertyNodeID NodeID(Property);
-					if (const FOverriddenPropertyNode* Node = OverriddenProperties->GetOverriddenPropertyNode(NodeID))
+					// We will walk through all sub-properties below, except for removed and added elements, so we test for those here.
+					FArchiveSerializedPropertyChain CurrentPropertyChain = Context.Path.ToSerializedPropertyChain();
+					if (const FOverriddenPropertyNode* Node = OverriddenProperties->GetOverriddenPropertyNode(&CurrentPropertyChain))
 					{
 						bHasOverrides |= Node->GetSubPropertyNodes().ContainsByPredicate([](const FOverriddenPropertyNode& SubNode)
 							{
-								return SubNode.GetOperation() == EOverriddenPropertyOperation::Remove;
+								return SubNode.GetOperation() == EOverriddenPropertyOperation::Remove || 
+									SubNode.GetOperation() == EOverriddenPropertyOperation::Add;
 							});
 					}
 				}
