[Core] Fix CI precision loss warning

#rnx
#Synced-CL 51535387

diff --git a/Engine/Source/Runtime/CoreUObject/Tests/PluginManagerTest.cpp b/Engine/Source/Runtime/CoreUObject/Tests/PluginManagerTest.cpp
--- a/Engine/Source/Runtime/CoreUObject/Tests/PluginManagerTest.cpp
+++ b/Engine/Source/Runtime/CoreUObject/Tests/PluginManagerTest.cpp
@@ -188,7 +188,7 @@ namespace UE::CoreUObject::Private::Tests
 			});
 		FText Reason;
 		CHECK(PluginManager.UnmountExplicitlyLoadedPlugin(TEXT("MyTestPlugin"), &Reason));
-		FTSTicker::GetCoreTicker().Tick(0.01);
+		FTSTicker::GetCoreTicker().Tick(0.01f);
 		CHECK(!FPackageName::MountPointExists(TEXT("/MyTestPlugin")));
 		CHECK(PluginPackageName != PluginPackage->GetFName()); //package should have been renamed.
 		CHECK(PluginPackage->HasAnyInternalFlags(EInternalObjectFlags::Garbage));
@@ -237,7 +237,7 @@ namespace UE::CoreUObject::Private::Tests
 
 		FText Reason;
 		CHECK(PluginManager.UnmountExplicitlyLoadedPlugin(TEXT("MyTestPlugin"), &Reason));
-		FTSTicker::GetCoreTicker().Tick(0.01);
+		FTSTicker::GetCoreTicker().Tick(0.01f);
 		CHECK(!FPackageName::MountPointExists(TEXT("/MyTestPlugin")));
 		CHECK(!FindObject<UObjectPtrTestClass>(AssetPath));
 	}
