Disabling idos for objects without unknown properties

#Synced-CL 51561634

diff --git a/Engine/Source/Runtime/CoreUObject/Internal/UObject/InstanceDataObjectUtils.h b/Engine/Source/Runtime/CoreUObject/Internal/UObject/InstanceDataObjectUtils.h
--- a/Engine/Source/Runtime/CoreUObject/Internal/UObject/InstanceDataObjectUtils.h
+++ b/Engine/Source/Runtime/CoreUObject/Internal/UObject/InstanceDataObjectUtils.h
@@ -98,11 +98,6 @@ UE_API UObject* CreateInstanceDataObject(UObject* Owner, FArchive& Ar, int64 Sta
 /** Returns the IDO for the object if it has one, otherwise returns the object. */
 UE_API UObject* ResolveInstanceDataObject(UObject* Object);
 
-namespace Private
-{
-	UE_API bool TryCreateIDOForConstructedObject(UObject* Object, bool bCreateIDOsForSubObjects = true);
-} // Private
-
 } // UE
 
 #undef UE_API
diff --git a/Engine/Source/Runtime/CoreUObject/Private/UObject/InstanceDataObjectUtils.cpp b/Engine/Source/Runtime/CoreUObject/Private/UObject/InstanceDataObjectUtils.cpp
--- a/Engine/Source/Runtime/CoreUObject/Private/UObject/InstanceDataObjectUtils.cpp
+++ b/Engine/Source/Runtime/CoreUObject/Private/UObject/InstanceDataObjectUtils.cpp
@@ -1850,92 +1850,6 @@ UObject* ResolveInstanceDataObject(UObject* Object)
 	return InstanceDataObject ? InstanceDataObject : Object;
 }
 
-namespace Private
-{
-
-static bool CanCreateIDOForConstructedObject(UObject* Object)
-{
-	bool bCanCreateIDO = false;
-
-	const UClass* ObjectClass = Object->GetClass();
-
-	// Ignore objects created very early in the initialization of CoreUObject.
-	if (ObjectClass && ObjectClass->GetClass())
-	{
-		// Do not create IDOs for the following types of objects:
-		// - Object that will be loaded: we handle those via FScopedObjectSerializeContext.
-		// - CDOs: created when the IDO class is created.
-		// - Archetypes: created on-demand.
-		// - Transient objects.
-		// - Objects with an obsolete type.
-		// - Types that do not support IDOs.
-		// - IDOs: we don't want to create IDOs for IDOs.
-		if (!Object->HasAnyFlags(RF_NeedLoad | RF_ClassDefaultObject | RF_ArchetypeObject | RF_Transient)
-			&& !Object->HasAnyInternalFlags(EInternalObjectFlags_AsyncLoading)
-			&& !ObjectClass->HasAnyClassFlags(CLASS_NewerVersionExists)
-			&& IsInstanceDataObjectSupportEnabledForClass(ObjectClass)
-			&& !Object->IsInPackage(GetTransientPackage()))
-		{
-			// Assume this object is getting loaded if an object is being serialized, so skip IDO creation.
-			const FUObjectSerializeContext* SerializeContext = FUObjectThreadContext::Get().GetSerializeContext();
-			if (!SerializeContext || !SerializeContext->SerializedObject)
-			{
-				bCanCreateIDO = !IsInstanceDataObject(Object) && IsInstanceDataObjectSupportEnabled(Object);
-			}
-		}
-	}
-
-	return bCanCreateIDO;
-}
-
-static bool TryCreateIDOForConstructedObject_Recursive(UObject* Object)
-{
-	bool bDidCreateIDO = false;
-
-	if (CanCreateIDOForConstructedObject(Object))
-	{
-		UObject* IDO = CreateInstanceDataObject(Object);
-		bDidCreateIDO = (IDO != nullptr);
-	}
-	TArray<UObject*> SubObjects;
-	constexpr bool bIncludeNestedObjects = false;
-	GetObjectsWithOuter(Object, SubObjects, bIncludeNestedObjects);
-	for (UObject* SubObject : SubObjects)
-	{
-		if (TryCreateIDOForConstructedObject_Recursive(SubObject))
-		{
-			bDidCreateIDO = true;
-		}
-	}
-
-	return bDidCreateIDO;
-}
-
-bool TryCreateIDOForConstructedObject(UObject* Object, bool bCreateIDOsForSubObjects)
-{
-	bool bDidCreateIDO = false;
-
-	if (bEnableIDOOnSpawnedObjects)
-	{
-		UE_AUTORTFM_ONCOMMIT(Object, bDidCreateIDO, bCreateIDOsForSubObjects)
-		{
-			if (bCreateIDOsForSubObjects)
-			{
-				bDidCreateIDO = TryCreateIDOForConstructedObject_Recursive(Object);
-			}
-			else if (CanCreateIDOForConstructedObject(Object))
-			{
-				UObject* IDO = CreateInstanceDataObject(Object);
-				bDidCreateIDO = (IDO != nullptr);
-			}
-		};
-	}
-
-	return bDidCreateIDO;
-}
-
-} // Private
-
 } // UE
 
 #endif // WITH_EDITORONLY_DATA
diff --git a/Engine/Source/Runtime/CoreUObject/Private/UObject/PropertyBagRepository.cpp b/Engine/Source/Runtime/CoreUObject/Private/UObject/PropertyBagRepository.cpp
--- a/Engine/Source/Runtime/CoreUObject/Private/UObject/PropertyBagRepository.cpp
+++ b/Engine/Source/Runtime/CoreUObject/Private/UObject/PropertyBagRepository.cpp
@@ -903,6 +903,11 @@ UObject* FPropertyBagRepository::CreateInstanceDataObject(UObject* Owner, FArchi
 	TSharedPtr<FPropertyPathNameTree> PropertyTree = FUnknownPropertyTree(Owner).Find();
 	const FUnknownEnumNames EnumNames(Owner);
 	const FUnknownEnumNames* EnumNamesPtr = !EnumNames.IsEmpty() ? &EnumNames : nullptr;
+	// Don't create Instance Data Objects if they're not needed to store unknown properties
+	if (!PropertyTree || PropertyTree->IsEmpty())
+	{
+		return nullptr;
+	}
 	const UClass* InstanceDataObjectClass = CreateInstanceDataObjectClass(PropertyTree.Get(), EnumNamesPtr, Owner->GetClass(), GetTransientPackage());
 
 	// Find the template for the IDO.
