Fixed USE_STABLE_LOCALIZATION_KEYS macro. When disabled, the solution would not compile anymore.

#Synced-CL 51545236

diff --git a/Engine/Plugins/Developer/PluginUtils/Source/PluginUtils/Private/PluginUtils.cpp b/Engine/Plugins/Developer/PluginUtils/Source/PluginUtils/Private/PluginUtils.cpp
--- a/Engine/Plugins/Developer/PluginUtils/Source/PluginUtils/Private/PluginUtils.cpp
+++ b/Engine/Plugins/Developer/PluginUtils/Source/PluginUtils/Private/PluginUtils.cpp
@@ -400,6 +400,7 @@ namespace PluginUtils
 								AssetRenameData.Add(RenameData);
 							}
 
+#if USE_STABLE_LOCALIZATION_KEYS
 							// ANY assets created/duplicated from a plugin template should have its package localization ID reset to avoid conflicts.
 							if (UObject* Asset = AssetData.GetAsset())
 							{
@@ -407,6 +408,7 @@ namespace PluginUtils
 								TextNamespaceUtil::ClearPackageNamespace(Asset);
 								TextNamespaceUtil::EnsurePackageNamespace(Asset);
 							}
+#endif // USE_STABLE_LOCALIZATION_KEYS
 						}
 					}
 
diff --git a/Engine/Plugins/Editor/ContentBrowser/ContentBrowserAssetDataSource/Source/ContentBrowserAssetDataSource/Private/AssetFileContextMenu.cpp b/Engine/Plugins/Editor/ContentBrowser/ContentBrowserAssetDataSource/Source/ContentBrowserAssetDataSource/Private/AssetFileContextMenu.cpp
--- a/Engine/Plugins/Editor/ContentBrowser/ContentBrowserAssetDataSource/Source/ContentBrowserAssetDataSource/Private/AssetFileContextMenu.cpp
+++ b/Engine/Plugins/Editor/ContentBrowser/ContentBrowserAssetDataSource/Source/ContentBrowserAssetDataSource/Private/AssetFileContextMenu.cpp
@@ -1999,13 +1999,15 @@ void FAssetFileContextMenu::ExecuteShowLocalizationCache(const FAssetData InAsse
 
 		if (UObject* Asset = InAsset.FastGetAsset())
 		{
-			const FString CachedLocalizationId = TextNamespaceUtil::GetPackageNamespace(Asset);
 			TArray<FGatherableTextData> GatherableTextDataArray;
 
 			EPropertyLocalizationGathererResultFlags GatherableTextResultFlags = EPropertyLocalizationGathererResultFlags::Empty;
 			FPropertyLocalizationDataGatherer _(GatherableTextDataArray, Asset->GetPackage(), GatherableTextResultFlags);
 
+#if USE_STABLE_LOCALIZATION_KEYS
+			const FString CachedLocalizationId = TextNamespaceUtil::GetPackageNamespace(Asset);
 			AppendLocalizationCache(CachedLocalizationId, GatherableTextDataArray);
+#endif // USE_STABLE_LOCALIZATION_KEYS
 		}
 		else
 		{
diff --git a/Engine/Source/Editor/BlueprintGraph/Private/EdGraphSchema_K2.cpp b/Engine/Source/Editor/BlueprintGraph/Private/EdGraphSchema_K2.cpp
--- a/Engine/Source/Editor/BlueprintGraph/Private/EdGraphSchema_K2.cpp
+++ b/Engine/Source/Editor/BlueprintGraph/Private/EdGraphSchema_K2.cpp
@@ -5309,6 +5309,7 @@ static void ConformAutogeneratedDefaultValuePackage(
 		// Attempt to find the correct package namespace to use for text within this pin
 		// Favor using the node if we have it, as that will be most up-to-date
 		FString PackageNamespace;
+#if USE_STABLE_LOCALIZATION_KEYS
 		if (OwningNode)
 		{
 			PackageNamespace = TextNamespaceUtil::GetPackageNamespace(OwningNode);
@@ -5317,6 +5318,7 @@ static void ConformAutogeneratedDefaultValuePackage(
 		{
 			PackageNamespace = TextNamespaceUtil::ExtractPackageNamespace(FTextInspector::GetNamespace(DefaultTextValue).Get(FString()));
 		}
+#endif // USE_STABLE_LOCALIZATION_KEYS
 
 		if (!PackageNamespace.IsEmpty())
 		{
@@ -8155,7 +8157,7 @@ float UEdGraphSchema_K2::GetActionFilteredWeight(const FEdGraphSchemaAction& InC
 	const FString* EachTerm = nullptr;
 	const FString* EachTermSanitized = nullptr;
 
-	// For every filter item the user has typed in (the text in the search bar, seperated by spaces)
+	// For every filter item the user has typed in (the text in the search bar, separated by spaces)
 	for (int32 FilterIndex = 0; FilterIndex < InFilterTerms.Num(); ++FilterIndex)
 	{
 		EachTerm = &InFilterTerms[FilterIndex];
diff --git a/Engine/Source/Editor/EditorWidgets/Private/STextPropertyEditableTextBox.cpp b/Engine/Source/Editor/EditorWidgets/Private/STextPropertyEditableTextBox.cpp
--- a/Engine/Source/Editor/EditorWidgets/Private/STextPropertyEditableTextBox.cpp
+++ b/Engine/Source/Editor/EditorWidgets/Private/STextPropertyEditableTextBox.cpp
@@ -575,6 +575,7 @@ FReply STextPropertyEditableStringTableReference::OnUnlinkClicked()
 			// This prevents problems with properties that disallow empty text (e.g. enum display name)
 			FString NewNamespace;
 			FString NewKey;
+#if USE_STABLE_LOCALIZATION_KEYS
 			EditableTextProperty->GetStableTextId(
 				TextIndex,
 				IEditableTextProperty::ETextPropertyEditAction::EditedKey,
@@ -584,6 +585,9 @@ FReply STextPropertyEditableStringTableReference::OnUnlinkClicked()
 				NewNamespace,
 				NewKey
 			);
+#else	// USE_STABLE_LOCALIZATION_KEYS
+			NewKey = TextNamespaceUtil::GenerateRandomTextKey();
+#endif	// USE_STABLE_LOCALIZATION_KEYS
 			
 			EditableTextProperty->SetText(TextIndex, FText::ChangeKey(NewNamespace, NewKey, CurrentText));
 		}
@@ -1519,6 +1523,7 @@ void STextPropertyEditableTextBox::HandleLocalizableCheckStateChanged(ECheckBoxS
 				// Get the stable namespace and key that we should use for this property
 				FString NewNamespace;
 				FString NewKey;
+#if USE_STABLE_LOCALIZATION_KEYS
 				EditableTextProperty->GetStableTextId(
 					TextIndex,
 					IEditableTextProperty::ETextPropertyEditAction::EditedKey,
@@ -1528,6 +1533,9 @@ void STextPropertyEditableTextBox::HandleLocalizableCheckStateChanged(ECheckBoxS
 					NewNamespace,
 					NewKey
 					);
+#else // USE_STABLE_LOCALIZATION_KEYS
+				NewKey = TextNamespaceUtil::GenerateRandomTextKey();
+#endif // USE_STABLE_LOCALIZATION_KEYS
 
 				EditableTextProperty->SetText(TextIndex, FText::AsLocalizable_Advanced(NewNamespace, NewKey, PropertyValue.ToString()));
 			}
diff --git a/Engine/Source/Editor/UnrealEd/Private/Commandlets/GatherTextFromAssetsCommandlet.cpp b/Engine/Source/Editor/UnrealEd/Private/Commandlets/GatherTextFromAssetsCommandlet.cpp
--- a/Engine/Source/Editor/UnrealEd/Private/Commandlets/GatherTextFromAssetsCommandlet.cpp
+++ b/Engine/Source/Editor/UnrealEd/Private/Commandlets/GatherTextFromAssetsCommandlet.cpp
@@ -1606,6 +1606,7 @@ bool UGatherTextFromAssetsCommandlet::LoadAndProcessUncachedPackages(TArray<FNam
 		// and the package is in the list of packages with duplicate localization IDs,
 		// reset the package's localization ID
 		bool bPackageLocIdWasReset = false;
+#if USE_STABLE_LOCALIZATION_KEYS
 		if (bFixPackageLocalizationIdConflict)
 		{
 			if (PackagesWithDuplicateLocalizationIds.Contains(PackagePendingGather.PackageName))
@@ -1629,6 +1630,7 @@ bool UGatherTextFromAssetsCommandlet::LoadAndProcessUncachedPackages(TArray<FNam
 				}
 			}
 		}
+#endif
 
 		// Tick background tasks
 		UE::Private::GatherTextFromAssetsCommandlet::TickBackgroundTasks();
diff --git a/Engine/Source/Editor/UnrealEd/Private/ObjectTools.cpp b/Engine/Source/Editor/UnrealEd/Private/ObjectTools.cpp
--- a/Engine/Source/Editor/UnrealEd/Private/ObjectTools.cpp
+++ b/Engine/Source/Editor/UnrealEd/Private/ObjectTools.cpp
@@ -4504,9 +4504,11 @@ namespace ObjectTools
 
 			if (OldPackage)
 			{
+#if USE_STABLE_LOCALIZATION_KEYS
 				// Migrate the localization ID to the new package
 				TextNamespaceUtil::ForcePackageNamespace(NewPackage, TextNamespaceUtil::GetPackageNamespace(OldPackage));
 				TextNamespaceUtil::ClearPackageNamespace(OldPackage);
+#endif
 
 				// Remove any metadata from old package pointing to moved objects
 				OldPackage->GetMetaData().RemoveMetaDataOutsidePackage(OldPackage);
diff --git a/Engine/Source/Runtime/Core/Public/Internationalization/TextNamespaceFwd.h b/Engine/Source/Runtime/Core/Public/Internationalization/TextNamespaceFwd.h
--- a/Engine/Source/Runtime/Core/Public/Internationalization/TextNamespaceFwd.h
+++ b/Engine/Source/Runtime/Core/Public/Internationalization/TextNamespaceFwd.h
@@ -6,6 +6,9 @@
 
 // USE_STABLE_LOCALIZATION_KEYS controls whether we attempt to keep localization keys stable when editing text properties.
 // This requires WITH_EDITORONLY_DATA, but can also be optionally disabled in editor builds by changing its value to 0.
+// Note: Prior to this flag being introduced, localization package namespace (loc package ids) did not exist. Therefore, any
+// code related to the use of localization package namespace should be wrapped around this macro. The whole goal of a loc
+// package namespace is all about stabilizing text localization keys within a package.
 #if WITH_EDITORONLY_DATA
 	#define USE_STABLE_LOCALIZATION_KEYS (1)
 #else	// WITH_EDITORONLY_DATA
diff --git a/Engine/Source/Runtime/CoreUObject/Private/Internationalization/TextPackageNamespaceUtil.cpp b/Engine/Source/Runtime/CoreUObject/Private/Internationalization/TextPackageNamespaceUtil.cpp
--- a/Engine/Source/Runtime/CoreUObject/Private/Internationalization/TextPackageNamespaceUtil.cpp
+++ b/Engine/Source/Runtime/CoreUObject/Private/Internationalization/TextPackageNamespaceUtil.cpp
@@ -92,8 +92,31 @@ void ForcePackageNamespace(UPackage* InPackage, const FString& InNamespace)
 	}
 }
 
+FString GenerateDeterministicPackageNamespace(const FString& InPackagePath, const FGuid& InNonUniqueGuid)
+{
+	// Make a deterministic new guid that will vary based on the package
+	FBlake3 Builder;
+	Builder.Update(&InNonUniqueGuid, sizeof(FGuid));
+	Builder.Update(*InPackagePath, InPackagePath.Len() * sizeof(**InPackagePath));
+	FBlake3Hash Hash = Builder.Finalize();
+	// We use the first 16 bytes of the FIoHash to create the guid, there is
+	// no specific reason why these were chosen, we could take any pattern or combination
+	// of bytes.
+	uint32* HashBytes = (uint32*)Hash.GetBytes();
+	FGuid NewGuid(HashBytes[0], HashBytes[1], HashBytes[2], HashBytes[3]);
+
+	return NewGuid.ToString();
+}
+
 } // TextNamespaceUtilImpl
 
+FString TextNamespaceUtil::GenerateDeterministicPackageNamespace(const UPackage* InPackage)
+{
+	FString PackagePath = InPackage->GetName();
+	FGuid NonUniqueGuid = InPackage->GetPersistentGuid(); // Can be the same for duplicated packages
+	return TextNamespaceUtilImpl::GenerateDeterministicPackageNamespace(PackagePath, NonUniqueGuid);
+}
+
 FString TextNamespaceUtil::GetPackageNamespace(const UPackage* InPackage)
 {
 	return TextNamespaceUtilImpl::FindOrAddPackageNamespace(const_cast<UPackage*>(InPackage), /*bCanAdd*/false);
@@ -158,26 +181,6 @@ FText TextNamespaceUtil::CopyTextToPackage(const FText& InText, UObject* InObjec
 #endif	// USE_STABLE_LOCALIZATION_KEYS
 }
 
-#if WITH_EDITORONLY_DATA
-FString TextNamespaceUtil::GenerateDeterministicPackageNamespace(const UPackage* InPackage)
-{
-	// Make a deterministic new guid that will vary based on the package
-	FBlake3 Builder;
-	FString PackagePath = InPackage->GetName();
-	FGuid NonUniqueGuid = InPackage->GetPersistentGuid(); // Can be the same for duplicated packages
-	Builder.Update(&NonUniqueGuid, sizeof(FGuid));
-	Builder.Update(*PackagePath, PackagePath.Len() * sizeof(**PackagePath));
-	FBlake3Hash Hash = Builder.Finalize();
-	// We use the first 16 bytes of the FIoHash to create the guid, there is
-	// no specific reason why these were chosen, we could take any pattern or combination
-	// of bytes.
-	uint32* HashBytes = (uint32*)Hash.GetBytes();
-	FGuid NewGuid(HashBytes[0], HashBytes[1], HashBytes[2], HashBytes[3]);
-
-	return NewGuid.ToString();
-}
-#endif // WITH_EDITORONLY_DATA
-
 FString TextNamespaceUtil::GenerateRandomTextKey()
 {
 	return FGuid::NewGuid().ToString();
diff --git a/Engine/Source/Runtime/CoreUObject/Public/Internationalization/TextPackageNamespaceUtil.h b/Engine/Source/Runtime/CoreUObject/Public/Internationalization/TextPackageNamespaceUtil.h
--- a/Engine/Source/Runtime/CoreUObject/Public/Internationalization/TextPackageNamespaceUtil.h
+++ b/Engine/Source/Runtime/CoreUObject/Public/Internationalization/TextPackageNamespaceUtil.h
@@ -17,6 +17,14 @@ namespace TextNamespaceUtil
 
 #if USE_STABLE_LOCALIZATION_KEYS
 
+/**
+ * Generate a deterministic package namespace based on the given package info.
+ * @note This key will be formatted like a GUID, but the value will actually be based on deterministic hashes.
+ *
+ * @param InPackage						The package to generate the namespace for.
+ */
+COREUOBJECT_API FString GenerateDeterministicPackageNamespace(const UPackage* InPackage);
+
 /**
  * Given a package, try and get the namespace it should use for localization.
  *
@@ -100,16 +108,6 @@ COREUOBJECT_API void ForcePackageNamespace(UObject* InObject, const FString& InN
 COREUOBJECT_API FText CopyTextToPackage(const FText& InText, UPackage* InPackage, const ETextCopyMethod InCopyMethod = ETextCopyMethod::NewKey, const bool bAlwaysApplyPackageNamespace = false);
 COREUOBJECT_API FText CopyTextToPackage(const FText& InText, UObject* InObject, const ETextCopyMethod InCopyMethod = ETextCopyMethod::NewKey, const bool bAlwaysApplyPackageNamespace = false);
 
-#if WITH_EDITORONLY_DATA
-/**
- * Generate a deterministic package namespace based on the given package info.
- * @note This key will be formatted like a GUID, but the value will actually be based on deterministic hashes.
- *
- * @param InPackage						The package to generate the namespace for.
- */
-COREUOBJECT_API FString GenerateDeterministicPackageNamespace(const UPackage* InPackage);
-#endif // WITH_EDITORONLY_DATA
-
 /**
  * Generate a random text key.
  * @note This key will be a GUID.
