CoreUObject: Deprecate and remove the FField virtual functions Bind(), PostLoad() and BeginDestroy() that are implemented but never called.

#tests BuildCookRun Lyra

#Synced-CL 51536473

diff --git a/Engine/Source/Runtime/CoreUObject/Private/UObject/Field.cpp b/Engine/Source/Runtime/CoreUObject/Private/UObject/Field.cpp
--- a/Engine/Source/Runtime/CoreUObject/Private/UObject/Field.cpp
+++ b/Engine/Source/Runtime/CoreUObject/Private/UObject/Field.cpp
@@ -414,15 +414,6 @@ UPackage* FField::GetOutermost() const
 	return OwnerUObject->GetOutermost();
 }
 
-void FField::Bind()
-{
-}
-
-void FField::PostLoad()
-{
-	Bind();
-}
-
 void FField::Serialize(FArchive& Ar)
 {
 	Ar << NamePrivate;
@@ -468,10 +459,6 @@ void FField::GetPreloadDependencies(TArray<UObject*>& OutDeps)
 
 }
 
-void FField::BeginDestroy()
-{
-}
-
 void FField::AddReferencedObjects(FReferenceCollector& Collector)
 {
 	TObjectPtr<UObject> OwnerUObject = ObjectPtrWrap(Owner.ToUObject());
diff --git a/Engine/Source/Runtime/CoreUObject/Private/UObject/PropertyBaseObject.cpp b/Engine/Source/Runtime/CoreUObject/Private/UObject/PropertyBaseObject.cpp
--- a/Engine/Source/Runtime/CoreUObject/Private/UObject/PropertyBaseObject.cpp
+++ b/Engine/Source/Runtime/CoreUObject/Private/UObject/PropertyBaseObject.cpp
@@ -70,18 +70,6 @@ FObjectPropertyBase::FObjectPropertyBase(UField* InField)
 }
 #endif // WITH_EDITORONLY_DATA
 
-void FObjectPropertyBase::BeginDestroy()
-{
-#if USE_CIRCULAR_DEPENDENCY_LOAD_DEFERRING
-	if (ULinkerPlaceholderClass* PlaceholderClass = Cast<ULinkerPlaceholderClass>(PropertyClass))
-	{
-		PlaceholderClass->RemoveReferencingProperty(this);
-	}
-#endif // USE_CIRCULAR_DEPENDENCY_LOAD_DEFERRING
-
-	Super::BeginDestroy();
-}
-
 void FObjectPropertyBase::InstanceSubobjects(void* Data, void const* DefaultData, TNotNull<UObject*> InOwner, FObjectInstancingGraph* InstanceGraph )
 {
 	for ( int32 ArrayIndex = 0; ArrayIndex < ArrayDim; ArrayIndex++ )
diff --git a/Engine/Source/Runtime/CoreUObject/Private/UObject/PropertyClass.cpp b/Engine/Source/Runtime/CoreUObject/Private/UObject/PropertyClass.cpp
--- a/Engine/Source/Runtime/CoreUObject/Private/UObject/PropertyClass.cpp
+++ b/Engine/Source/Runtime/CoreUObject/Private/UObject/PropertyClass.cpp
@@ -40,18 +40,6 @@ FClassProperty::FClassProperty(UField* InField)
 }
 #endif // WITH_EDITORONLY_DATA
 
-void FClassProperty::BeginDestroy()
-{
-#if USE_CIRCULAR_DEPENDENCY_LOAD_DEFERRING
-	if (ULinkerPlaceholderClass* PlaceholderClass = Cast<ULinkerPlaceholderClass>(MetaClass))
-	{
-		PlaceholderClass->RemoveReferencingProperty(this);
-	}
-#endif // USE_CIRCULAR_DEPENDENCY_LOAD_DEFERRING
-
-	Super::BeginDestroy();
-}
-
 void FClassProperty::PostDuplicate(const FField& InField)
 {
 	const FClassProperty& Source = static_cast<const FClassProperty&>(InField);
diff --git a/Engine/Source/Runtime/CoreUObject/Private/UObject/PropertyDelegate.cpp b/Engine/Source/Runtime/CoreUObject/Private/UObject/PropertyDelegate.cpp
--- a/Engine/Source/Runtime/CoreUObject/Private/UObject/PropertyDelegate.cpp
+++ b/Engine/Source/Runtime/CoreUObject/Private/UObject/PropertyDelegate.cpp
@@ -221,15 +221,3 @@ void FDelegateProperty::AppendSchemaHash(FBlake3& Builder, bool bSkipEditorOnly)
 }
 #endif
 
-
-void FDelegateProperty::BeginDestroy()
-{
-#if USE_CIRCULAR_DEPENDENCY_LOAD_DEFERRING
-	if (auto PlaceholderFunc = Cast<ULinkerPlaceholderFunction>(SignatureFunction))
-	{
-		PlaceholderFunc->RemoveReferencingProperty(this);
-	}
-#endif // USE_CIRCULAR_DEPENDENCY_LOAD_DEFERRING
-
-	Super::BeginDestroy();
-}
diff --git a/Engine/Source/Runtime/CoreUObject/Private/UObject/PropertyInterface.cpp b/Engine/Source/Runtime/CoreUObject/Private/UObject/PropertyInterface.cpp
--- a/Engine/Source/Runtime/CoreUObject/Private/UObject/PropertyInterface.cpp
+++ b/Engine/Source/Runtime/CoreUObject/Private/UObject/PropertyInterface.cpp
@@ -40,18 +40,6 @@ FInterfaceProperty::FInterfaceProperty(UField* InField)
 }
 #endif // WITH_EDITORONLY_DATA
 
-void FInterfaceProperty::BeginDestroy()
-{
-#if USE_CIRCULAR_DEPENDENCY_LOAD_DEFERRING
-	if (ULinkerPlaceholderClass* PlaceholderClass = Cast<ULinkerPlaceholderClass>(InterfaceClass))
-	{
-		PlaceholderClass->RemoveReferencingProperty(this);
-	}
-#endif // USE_CIRCULAR_DEPENDENCY_LOAD_DEFERRING
-
-	Super::BeginDestroy();
-}
-
 void FInterfaceProperty::PostDuplicate(const FField& InField)
 {
 	const FInterfaceProperty& Source = static_cast<const FInterfaceProperty&>(InField);
diff --git a/Engine/Source/Runtime/CoreUObject/Private/UObject/PropertyMulticastDelegate.cpp b/Engine/Source/Runtime/CoreUObject/Private/UObject/PropertyMulticastDelegate.cpp
--- a/Engine/Source/Runtime/CoreUObject/Private/UObject/PropertyMulticastDelegate.cpp
+++ b/Engine/Source/Runtime/CoreUObject/Private/UObject/PropertyMulticastDelegate.cpp
@@ -387,18 +387,6 @@ void FMulticastDelegateProperty::Serialize( FArchive& Ar )
 #endif // USE_CIRCULAR_DEPENDENCY_LOAD_DEFERRING
 }
 
-void FMulticastDelegateProperty::BeginDestroy()
-{
-#if USE_CIRCULAR_DEPENDENCY_LOAD_DEFERRING
-	if (auto PlaceholderFunc = Cast<ULinkerPlaceholderFunction>(SignatureFunction))
-	{
-		PlaceholderFunc->RemoveReferencingProperty(this);
-	}
-#endif // USE_CIRCULAR_DEPENDENCY_LOAD_DEFERRING
-
-	Super::BeginDestroy();
-}
-
 bool FMulticastDelegateProperty::SameType(const FProperty* Other) const
 {
 	return Super::SameType(Other) && (SignatureFunction == ((FMulticastDelegateProperty*)Other)->SignatureFunction);
diff --git a/Engine/Source/Runtime/CoreUObject/Private/UObject/PropertySoftClassPtr.cpp b/Engine/Source/Runtime/CoreUObject/Private/UObject/PropertySoftClassPtr.cpp
--- a/Engine/Source/Runtime/CoreUObject/Private/UObject/PropertySoftClassPtr.cpp
+++ b/Engine/Source/Runtime/CoreUObject/Private/UObject/PropertySoftClassPtr.cpp
@@ -33,18 +33,6 @@ FSoftClassProperty::FSoftClassProperty(UField* InField)
 }
 #endif // WITH_EDITORONLY_DATA
 
-void FSoftClassProperty::BeginDestroy()
-{
-#if USE_CIRCULAR_DEPENDENCY_LOAD_DEFERRING
-	if (ULinkerPlaceholderClass* PlaceholderClass = Cast<ULinkerPlaceholderClass>(MetaClass))
-	{
-		PlaceholderClass->RemoveReferencingProperty(this);
-	}
-#endif // USE_CIRCULAR_DEPENDENCY_LOAD_DEFERRING
-
-	Super::BeginDestroy();
-}
-
 void FSoftClassProperty::PostDuplicate(const FField& InField)
 {
 	const FSoftClassProperty& Source = static_cast<const FSoftClassProperty&>(InField);
diff --git a/Engine/Source/Runtime/CoreUObject/Public/UObject/Field.h b/Engine/Source/Runtime/CoreUObject/Public/UObject/Field.h
--- a/Engine/Source/Runtime/CoreUObject/Public/UObject/Field.h
+++ b/Engine/Source/Runtime/CoreUObject/Public/UObject/Field.h
@@ -683,9 +683,13 @@ public:
 
 	// Begin UObject interface: the following functions mimic UObject interface for easier transition from UProperties to FProperties
 	COREUOBJECT_API virtual void Serialize(FArchive& Ar);
-	COREUOBJECT_API virtual void PostLoad();
+#if !UE_STRIP_DEPRECATED_PROPERTIES
+	UE_DEPRECATED(5.8, "FField::PostLoad is deprecated. It is never called.")
+	virtual void PostLoad() final {};
+	UE_DEPRECATED(5.8, "FField::BeginDestroy is deprecated. It is never called.")
+	virtual void BeginDestroy() final {};
+#endif
 	COREUOBJECT_API virtual void GetPreloadDependencies(TArray<UObject*>& OutDeps);
-	COREUOBJECT_API virtual void BeginDestroy();	
 	COREUOBJECT_API virtual void AddReferencedObjects(FReferenceCollector& Collector);
 	COREUOBJECT_API bool IsRooted() const;
 	COREUOBJECT_API bool IsNative() const;
@@ -700,7 +704,10 @@ public:
 
 	// Begin UField interface.
 	COREUOBJECT_API virtual void AddCppProperty(FProperty* Property);
-	COREUOBJECT_API virtual void Bind();
+#if !UE_STRIP_DEPRECATED_PROPERTIES
+	UE_DEPRECATED(5.8, "FField::Bind is deprecated. It is never called.")
+	virtual void Bind() final {};
+#endif
 	// End UField interface
 
 	/** Constructs a new field given its class */
diff --git a/Engine/Source/Runtime/CoreUObject/Public/UObject/UnrealType.h b/Engine/Source/Runtime/CoreUObject/Public/UObject/UnrealType.h
--- a/Engine/Source/Runtime/CoreUObject/Public/UObject/UnrealType.h
+++ b/Engine/Source/Runtime/CoreUObject/Public/UObject/UnrealType.h
@@ -2794,7 +2794,6 @@ public:
 	// UObject interface
 	UE_API virtual void Serialize( FArchive& Ar ) override;
 	UE_API virtual void AddReferencedObjects(FReferenceCollector& Collector) override;
-	UE_API virtual void BeginDestroy() override;
 #if UE_WITH_CONSTINIT_UOBJECT && !IS_MONOLITHIC
 	UE_API virtual void LinkCompiledInPointerFields(UE::CoreUObject::Private::FCompiledInObjectLinker& Linker);
 #endif // UE_WITH_CONSTINIT_UOBJECT && !IS_MONOLITHIC
@@ -3487,7 +3486,6 @@ public:
 	// UObject interface
 	UE_API virtual void Serialize( FArchive& Ar ) override;
 	UE_API virtual void AddReferencedObjects(FReferenceCollector& Collector) override;
-	UE_API virtual void BeginDestroy() override;
 #if UE_WITH_CONSTINIT_UOBJECT && !IS_MONOLITHIC
 	UE_API virtual void LinkCompiledInPointerFields(UE::CoreUObject::Private::FCompiledInObjectLinker& Linker);
 #endif // UE_WITH_CONSTINIT_UOBJECT && !IS_MONOLITHIC
@@ -3581,7 +3579,6 @@ public:
 	// UObject interface
 	UE_API virtual void Serialize( FArchive& Ar ) override;
 	UE_API virtual void AddReferencedObjects(FReferenceCollector& Collector) override;
-	UE_API virtual void BeginDestroy() override;
 #if UE_WITH_CONSTINIT_UOBJECT && !IS_MONOLITHIC
 	UE_API virtual void LinkCompiledInPointerFields(UE::CoreUObject::Private::FCompiledInObjectLinker& Linker);
 #endif // UE_WITH_CONSTINIT_UOBJECT && !IS_MONOLITHIC
@@ -3679,7 +3676,6 @@ public:
 	// UObject interface
 	UE_API virtual void Serialize( FArchive& Ar ) override;
 	UE_API virtual void EmitReferenceInfo(UE::GC::FSchemaBuilder& Schema, int32 BaseOffset, TArray<const FStructProperty*>& EncounteredStructProps, UE::GC::FPropertyStack& DebugPath) override;
-	UE_API virtual void BeginDestroy() override;
 	UE_API virtual void AddReferencedObjects(FReferenceCollector& Collector) override;
 #if UE_WITH_CONSTINIT_UOBJECT && !IS_MONOLITHIC
 	UE_API virtual void LinkCompiledInPointerFields(UE::CoreUObject::Private::FCompiledInObjectLinker& Linker);
@@ -6511,7 +6507,6 @@ public:
 	// UObject interface
 	UE_API virtual void Serialize( FArchive& Ar ) override;
 	UE_API virtual void AddReferencedObjects(FReferenceCollector& Collector) override;
-	UE_API virtual void BeginDestroy() override;
 #if UE_WITH_CONSTINIT_UOBJECT && !IS_MONOLITHIC
 	UE_API virtual void LinkCompiledInPointerFields(UE::CoreUObject::Private::FCompiledInObjectLinker& Linker);
 #endif // UE_WITH_CONSTINIT_UOBJECT && !IS_MONOLITHIC
@@ -6580,7 +6575,6 @@ public:
 
 	// UObject interface
 	UE_API virtual void Serialize( FArchive& Ar ) override;
-	UE_API virtual void BeginDestroy() override;
 	UE_API virtual void AddReferencedObjects(FReferenceCollector& Collector) override;
 #if UE_WITH_CONSTINIT_UOBJECT && !IS_MONOLITHIC
 	UE_API virtual void LinkCompiledInPointerFields(UE::CoreUObject::Private::FCompiledInObjectLinker& Linker);
