[Insights] Added initial support for UObject minimal trace (snapshots).

#Synced-CL 51542463

diff --git a/Engine/Source/Runtime/CoreUObject/Private/UObject/ObjectTrace.cpp b/Engine/Source/Runtime/CoreUObject/Private/UObject/ObjectTrace.cpp
new file mode 100644
--- /dev/null
+++ b/Engine/Source/Runtime/CoreUObject/Private/UObject/ObjectTrace.cpp
@@ -0,0 +1,348 @@
+// Copyright Epic Games, Inc. All Rights Reserved.
+
+#include "UObject/ObjectTrace.h"
+
+#if UE_UOBJECT_TRACE_ENABLED
+
+#include "Containers/VersePath.h"
+#include "HAL/IConsoleManager.h"
+#include "HAL/PlatformTime.h"
+#include "Logging/LogCategory.h"
+#include "Misc/AssertionMacros.h"
+#include "ProfilingDebugging/StringsTrace.h"
+#include "ProfilingDebugging/TraceAuxiliary.h"
+#include "Trace/Config.h"
+#include "Trace/Trace.inl"
+#include "UObject/Class.h"
+#include "UObject/NameTypes.h"
+#include "UObject/Object.h"
+#include "UObject/Package.h"
+#include "UObject/UObjectArray.h"
+#include "UObject/UObjectBaseUtility.h"
+
+DEFINE_LOG_CATEGORY_STATIC(LogUObjectTrace, Log, All);
+
+UE_TRACE_MINIMAL_CHANNEL_DEFINE(CoreUObjectChannel, "UObject Trace", false)
+
+namespace UE::Trace
+{
+
+////////////////////////////////////////////////////////////////////////////////////////////////////
+
+UE_TRACE_MINIMAL_EVENT_BEGIN(CoreUObject, BeginSnapshot, NoSync)
+	UE_TRACE_MINIMAL_EVENT_FIELD(uint64, Cycle)
+	UE_TRACE_MINIMAL_EVENT_FIELD(uint32, SizeOfUObject)
+	UE_TRACE_MINIMAL_EVENT_FIELD(uint32, SizeOfUField)
+	UE_TRACE_MINIMAL_EVENT_FIELD(uint32, SizeOfUStruct)
+	UE_TRACE_MINIMAL_EVENT_FIELD(uint32, SizeOfUClass)
+	UE_TRACE_MINIMAL_EVENT_FIELD(uint32, SizeOfUPackage)
+UE_TRACE_MINIMAL_EVENT_END()
+
+UE_TRACE_MINIMAL_EVENT_BEGIN(CoreUObject, EndSnapshot, NoSync)
+	UE_TRACE_MINIMAL_EVENT_FIELD(uint64, Cycle)
+UE_TRACE_MINIMAL_EVENT_END()
+
+UE_TRACE_MINIMAL_EVENT_BEGIN(CoreUObject, ObjectSpec, NoSync)
+	UE_TRACE_MINIMAL_EVENT_FIELD(uint32, Id)
+	UE_TRACE_MINIMAL_EVENT_FIELD(uint32, ClassId)
+	UE_TRACE_MINIMAL_EVENT_FIELD(uint32, OuterId)
+	UE_TRACE_MINIMAL_EVENT_FIELD(uint32, Flags)
+	UE_TRACE_MINIMAL_EVENT_FIELD(UE::Trace::AnsiString, NameAnsi)
+	UE_TRACE_MINIMAL_EVENT_FIELD(UE::Trace::WideString, NameWide)
+UE_TRACE_MINIMAL_EVENT_END()
+
+UE_TRACE_MINIMAL_EVENT_BEGIN(CoreUObject, VerseExtraSpec, NoSync)
+	UE_TRACE_MINIMAL_EVENT_FIELD(uint32, Id)
+	UE_TRACE_MINIMAL_EVENT_FIELD(UE::Trace::WideString, Path)
+UE_TRACE_MINIMAL_EVENT_END()
+
+UE_TRACE_MINIMAL_EVENT_BEGIN(CoreUObject, ResourceSizeExtraSpec, NoSync)
+	UE_TRACE_MINIMAL_EVENT_FIELD(uint32, Id)
+	UE_TRACE_MINIMAL_EVENT_FIELD(uint64, SystemMemoryBytes) // exclusive estimated system memory
+	UE_TRACE_MINIMAL_EVENT_FIELD(uint64, VideoMemoryBytes) // exclusive estimated video memory
+	UE_TRACE_MINIMAL_EVENT_FIELD(uint64, TotalSystemMemoryBytes) // inclusive estimated system memory
+	UE_TRACE_MINIMAL_EVENT_FIELD(uint64, TotalVideoMemoryBytes) // inclusive estimated video memory
+UE_TRACE_MINIMAL_EVENT_END()
+
+UE_TRACE_MINIMAL_EVENT_BEGIN(CoreUObject, FieldExtraSpec, NoSync)
+	UE_TRACE_MINIMAL_EVENT_FIELD(uint32, Id)
+UE_TRACE_MINIMAL_EVENT_END()
+
+UE_TRACE_MINIMAL_EVENT_BEGIN(CoreUObject, StructExtraSpec, NoSync)
+	UE_TRACE_MINIMAL_EVENT_FIELD(uint32, Id)
+	UE_TRACE_MINIMAL_EVENT_FIELD(uint32, SuperId)
+	UE_TRACE_MINIMAL_EVENT_FIELD(uint32, InheritanceSuperId)
+	UE_TRACE_MINIMAL_EVENT_FIELD(int32, StructureSize)
+UE_TRACE_MINIMAL_EVENT_END()
+
+UE_TRACE_MINIMAL_EVENT_BEGIN(CoreUObject, ClassExtraSpec, NoSync)
+	UE_TRACE_MINIMAL_EVENT_FIELD(uint32, Id)
+UE_TRACE_MINIMAL_EVENT_END()
+
+UE_TRACE_MINIMAL_EVENT_BEGIN(CoreUObject, PackageExtraSpec, NoSync)
+	UE_TRACE_MINIMAL_EVENT_FIELD(uint32, Id)
+	UE_TRACE_MINIMAL_EVENT_FIELD(uint64, PackageId)
+	UE_TRACE_MINIMAL_EVENT_FIELD(UE::Trace::WideString, LocalFullPath)
+UE_TRACE_MINIMAL_EVENT_END()
+
+UE_TRACE_MINIMAL_EVENT_BEGIN(CoreUObject, ObjectRef, NoSync)
+	UE_TRACE_MINIMAL_EVENT_FIELD(uint32, ReferencerId)
+	UE_TRACE_MINIMAL_EVENT_FIELD(uint32, ObjectId)
+UE_TRACE_MINIMAL_EVENT_END()
+
+////////////////////////////////////////////////////////////////////////////////////////////////////
+
+static FAutoConsoleCommand CVarStartCapturingUObjectTraceSnapshot(
+	TEXT("obj.TraceSnapshot"),
+	TEXT("Capture a UObject Trace Snapshot (minimal tracing)."),
+	FConsoleCommandWithArgsDelegate::CreateLambda([](const TArray<FString>& Args)
+		{
+			FUObjectTrace::StartCapturingSnapshot();
+		})
+);
+
+////////////////////////////////////////////////////////////////////////////////////////////////////
+// FUObjectTrace
+////////////////////////////////////////////////////////////////////////////////////////////////////
+
+std::atomic<bool> FUObjectTrace::bIsCapturingSnapshot = false;
+
+void FUObjectTrace::StartCapturingSnapshot()
+{
+	check(IsInGameThread());
+
+	if (bIsCapturingSnapshot)
+	{
+		UE_LOG(LogUObjectTrace, Warning, TEXT("Already capturing a UObject Trace Snapshot."));
+		return;
+	}
+
+	const bool bShouldToggleChannel = !CoreUObjectChannel.IsEnabled();
+	if (bShouldToggleChannel)
+	{
+		const TCHAR* FailReason = nullptr;
+		if (!CoreUObjectChannel.Toggle(true, &FailReason))
+		{
+			UE_LOG(LogUObjectTrace, Error, TEXT("Failed to enable the CoreUObject trace channel."));
+			return;
+		}
+	}
+
+	const uint64 StartTime = FPlatformTime::Cycles64();
+	bIsCapturingSnapshot = true;
+	OutputBeginSnapshot();
+
+	uint32 ObjectsTraced = 0;
+	int32 ObjectArrayNum = GUObjectArray.GetObjectArrayNum();
+	for (int32 ObjectIndex = 0; ObjectIndex < ObjectArrayNum; ++ObjectIndex)
+	{
+		UObjectBase* Object = FIndexToObject::IndexToObject(ObjectIndex, false);
+		if (OutputObjectSpec(Object))
+		{
+			++ObjectsTraced;
+		}
+	}
+
+	OutputEndSnapshot();
+	bIsCapturingSnapshot = false;
+	const uint64 EndTime = FPlatformTime::Cycles64();
+
+	UE_LOG(LogUObjectTrace, Log, TEXT("UObject Trace Snapshot captured in %.3f seconds (%u objects)."),
+		(double)(EndTime - StartTime) * FPlatformTime::GetSecondsPerCycle64(), ObjectsTraced);
+
+	if (bShouldToggleChannel)
+	{
+		const TCHAR* FailReason = nullptr;
+		if (CoreUObjectChannel.Toggle(false, &FailReason))
+		{
+			UE_LOG(LogUObjectTrace, Error, TEXT("Failed to disable the CoreUObject trace channel."));
+		}
+	}
+}
+
+void FUObjectTrace::OutputBeginSnapshot()
+{
+	UE_TRACE_MINIMAL_LOG(CoreUObject, BeginSnapshot, CoreUObjectChannel)
+		<< BeginSnapshot.Cycle(FPlatformTime::Cycles64())
+		<< BeginSnapshot.SizeOfUObject(static_cast<uint32>(sizeof(UObject)))
+		<< BeginSnapshot.SizeOfUField(static_cast<uint32>(sizeof(UField)))
+		<< BeginSnapshot.SizeOfUStruct(static_cast<uint32>(sizeof(UStruct)))
+		<< BeginSnapshot.SizeOfUClass(static_cast<uint32>(sizeof(UClass)))
+		<< BeginSnapshot.SizeOfUPackage(static_cast<uint32>(sizeof(UPackage)))
+	;
+}
+
+void FUObjectTrace::OutputEndSnapshot()
+{
+	UE_TRACE_MINIMAL_LOG(CoreUObject, EndSnapshot, CoreUObjectChannel)
+		<< EndSnapshot.Cycle(FPlatformTime::Cycles64());
+}
+
+bool FUObjectTrace::OutputObjectSpec(UObjectBase* Object)
+{
+	if (!bool(CoreUObjectChannel))
+	{
+		return false;
+	}
+
+	if (!Object)
+	{
+		return false;
+	}
+
+	constexpr uint32 InvalidObjectId = uint32(-1);
+
+	const uint32 Id = Object->GetUniqueID();
+	if (Id == InvalidObjectId)
+	{
+		return false;
+	}
+
+	UObject* Obj = Cast<UObject>(Object);
+	check(Obj);
+
+	UClass* ObjectClass = Object->GetClass();
+	UObject* OuterObject = Object->GetOuter();
+	const EObjectFlags Flags = Object->GetFlags();
+
+	FString ObjectName;
+	FWideStringView ObjectNameWide;
+	TAnsiStringBuilder<1024> ObjectNameAnsiBuilder;
+	if (!Object->GetFName().TryAppendAnsiString(ObjectNameAnsiBuilder))
+	{
+		Object->GetFName().ToString(ObjectName);
+		ObjectNameWide = ObjectName;
+	}
+	FAnsiStringView ObjectNameAnsi = ObjectNameAnsiBuilder.ToView();
+
+	UE_TRACE_MINIMAL_LOG(CoreUObject, ObjectSpec, CoreUObjectChannel)
+		<< ObjectSpec.Id(Id)
+		<< ObjectSpec.ClassId(ObjectClass ? ObjectClass->GetUniqueID() : InvalidObjectId)
+		<< ObjectSpec.OuterId(OuterObject ? OuterObject->GetUniqueID() : InvalidObjectId)
+		<< ObjectSpec.Flags(uint32(Flags))
+		<< ObjectSpec.NameAnsi(ObjectNameAnsi.GetData(), ObjectNameAnsi.Len())
+		<< ObjectSpec.NameWide(ObjectNameWide.GetData(), ObjectNameWide.Len());
+
+	constexpr bool bTraceVerseExtraSpec = true;
+	if (bTraceVerseExtraSpec)
+	{
+		UE::Core::FVersePath VersePath = Obj->GetVersePath();
+		if (VersePath.IsValid())
+		{
+			FStringView VersePathView = VersePath.AsStringView();
+			UE_TRACE_MINIMAL_LOG(CoreUObject, VerseExtraSpec, CoreUObjectChannel)
+				<< VerseExtraSpec.Id(Id)
+				<< VerseExtraSpec.Path(VersePathView.GetData(), VersePathView.Len());
+		}
+	}
+
+	constexpr bool bTraceResourceSize = true;
+	if (bTraceResourceSize)
+	{
+		FResourceSizeEx ResourceSize(EResourceSizeMode::Exclusive);
+		FResourceSizeEx TotalResourceSize(EResourceSizeMode::EstimatedTotal);
+		Obj->GetResourceSizeEx(ResourceSize);
+		Obj->GetResourceSizeEx(TotalResourceSize);
+		SIZE_T SystemMemoryBytes = ResourceSize.GetDedicatedSystemMemoryBytes(); // exclusive
+		SIZE_T VideoMemoryBytes = ResourceSize.GetDedicatedVideoMemoryBytes(); // exclusive
+		SIZE_T TotalSystemMemoryBytes = TotalResourceSize.GetDedicatedSystemMemoryBytes(); // inclusive
+		SIZE_T TotalVideoMemoryBytes = TotalResourceSize.GetDedicatedVideoMemoryBytes(); // inclusive
+		if (SystemMemoryBytes || VideoMemoryBytes || TotalSystemMemoryBytes || TotalVideoMemoryBytes)
+		{
+			UE_TRACE_MINIMAL_LOG(CoreUObject, ResourceSizeExtraSpec, CoreUObjectChannel)
+				<< ResourceSizeExtraSpec.Id(Id)
+				<< ResourceSizeExtraSpec.SystemMemoryBytes(SystemMemoryBytes)
+				<< ResourceSizeExtraSpec.VideoMemoryBytes(VideoMemoryBytes)
+				<< ResourceSizeExtraSpec.TotalSystemMemoryBytes(TotalSystemMemoryBytes)
+				<< ResourceSizeExtraSpec.TotalVideoMemoryBytes(TotalVideoMemoryBytes);
+		}
+	}
+
+	constexpr bool bTraceExtraSpecs = true;
+	if (bTraceExtraSpecs)
+	{
+		// UField -> UObject
+		if (UField* ObjAsField = Cast<UField>(Object))
+		{
+			UE_TRACE_MINIMAL_LOG(CoreUObject, FieldExtraSpec, CoreUObjectChannel)
+				<< FieldExtraSpec.Id(Id);
+		}
+
+		// UStruct -> UField -> UObject
+		if (UStruct* ObjAsStruct = Cast<UStruct>(Object))
+		{
+			UStruct* ObjSuper = ObjAsStruct->GetSuperStruct();
+			UStruct* ObjInheritanceSuper = ObjAsStruct->GetInheritanceSuper();
+			const int32 StructureSize = ObjAsStruct->GetStructureSize();
+			UE_TRACE_MINIMAL_LOG(CoreUObject, StructExtraSpec, CoreUObjectChannel)
+				<< StructExtraSpec.Id(Id)
+				<< StructExtraSpec.SuperId(ObjSuper ? ObjSuper->GetUniqueID() : InvalidObjectId)
+				<< StructExtraSpec.InheritanceSuperId(ObjInheritanceSuper ? ObjInheritanceSuper->GetUniqueID() : InvalidObjectId)
+				<< StructExtraSpec.StructureSize(StructureSize);
+		}
+
+		// UClass -> UStruct -> UField -> UObject
+		if (UClass* ObjAsClass = Cast<UClass>(Object))
+		{
+			UE_TRACE_MINIMAL_LOG(CoreUObject, ClassExtraSpec, CoreUObjectChannel)
+				<< ClassExtraSpec.Id(Id);
+		}
+
+		// UPackage -> UObject
+		if (UPackage* ObjAsPackage = Cast<UPackage>(Object))
+		{
+			const uint64 PackageId = ObjAsPackage->GetPackageId().Value();
+			const FPackagePath& PackagePath = ObjAsPackage->GetLoadedPath();
+			FString LocalFullPath = PackagePath.GetLocalFullPath();
+			FWideStringView LocalFullPathView = LocalFullPath;
+			UE_TRACE_MINIMAL_LOG(CoreUObject, PackageExtraSpec, CoreUObjectChannel)
+				<< PackageExtraSpec.Id(Id)
+				<< PackageExtraSpec.PackageId(PackageId)
+				<< PackageExtraSpec.LocalFullPath(LocalFullPathView.GetData(), LocalFullPathView.Len());
+		}
+
+		// UWorld -> UObject
+		// ULevel -> UObject
+		// UTexture -> UStreamableRenderAsset -> UObject
+		// UMaterial -> UObject
+		// UStaticMesh -> UStreamableRenderAsset -> UObject
+		// USkeletalMesh -> USkinnedAsset -> UStreamableRenderAsset -> UObject
+		// AActor -> UObject
+		//TODO: Object->TraceExtraSpec();
+	}
+
+	return true;
+}
+
+void FUObjectTrace::OutputObjectRef(UObjectBase* Referencer, UObjectBase* Object)
+{
+	if (!bool(CoreUObjectChannel))
+	{
+		return;
+	}
+
+	check(Referencer);
+	check(Object);
+
+	constexpr uint32 InvalidObjectId = uint32(-1);
+
+	uint32 ReferencerId = Referencer->GetUniqueID();
+	if (ReferencerId == InvalidObjectId)
+	{
+		return;
+	}
+
+	uint32 ObjectId = Object->GetUniqueID();
+	if (ObjectId == InvalidObjectId)
+	{
+		return;
+	}
+
+	UE_TRACE_MINIMAL_LOG(CoreUObject, ObjectRef, CoreUObjectChannel)
+		<< ObjectRef.ReferencerId(ReferencerId)
+		<< ObjectRef.ObjectId(ObjectId);
+}
+
+} // namespace UE::Trace
+
+#endif // UE_UOBJECT_TRACE_ENABLED
diff --git a/Engine/Source/Runtime/CoreUObject/Public/UObject/ObjectTrace.h b/Engine/Source/Runtime/CoreUObject/Public/UObject/ObjectTrace.h
new file mode 100644
--- /dev/null
+++ b/Engine/Source/Runtime/CoreUObject/Public/UObject/ObjectTrace.h
@@ -0,0 +1,42 @@
+// Copyright Epic Games, Inc. All Rights Reserved.
+
+#pragma once
+
+#include "CoreTypes.h"
+#include "Trace/Config.h"
+
+#ifndef UE_UOBJECT_TRACE_ENABLED
+#define UE_UOBJECT_TRACE_ENABLED UE_TRACE_MINIMAL_ENABLED
+#endif
+
+#if UE_UOBJECT_TRACE_ENABLED
+
+#include "Containers/UnrealString.h"
+#include "Trace/Trace.h"
+
+#include <atomic>
+
+UE_TRACE_CHANNEL_EXTERN(CoreUObjectChannel, COREUOBJECT_API)
+
+class UObject;
+
+namespace UE::Trace
+{
+
+class FUObjectTrace
+{
+public:
+	static bool IsCapturingSnapshot() { return bIsCapturingSnapshot; }
+	COREUOBJECT_API static void StartCapturingSnapshot();
+	COREUOBJECT_API static void OutputBeginSnapshot();
+	COREUOBJECT_API static void OutputEndSnapshot();
+	COREUOBJECT_API static bool OutputObjectSpec(UObjectBase* Object);
+	COREUOBJECT_API static void OutputObjectRef(UObjectBase* Referencer, UObjectBase* Object);
+
+private:
+	static std::atomic<bool> bIsCapturingSnapshot;
+};
+
+} // namespace UE::Trace
+
+#endif // UE_UOBJECT_TRACE_ENABLED
