[UObject/ALT]
- Relax the rules to allow UObject reconstruction to happen from the loading thread

#Synced-CL 51539428

diff --git a/Engine/Source/Runtime/CoreUObject/Private/UObject/UObjectArray.cpp b/Engine/Source/Runtime/CoreUObject/Private/UObject/UObjectArray.cpp
--- a/Engine/Source/Runtime/CoreUObject/Private/UObject/UObjectArray.cpp
+++ b/Engine/Source/Runtime/CoreUObject/Private/UObject/UObjectArray.cpp
@@ -378,9 +378,6 @@ void FUObjectArray::FreeUObjectIndex(UObjectBase* Object)
 	LLM_SCOPE(ELLMTag::UObject);
 	LLM_SCOPE_BYTAG(UObject_UObjectArray);
 
-	// This should only be happening on the game thread (GC runs only on game thread when it's freeing objects)
-	check(IsInGameThread());
-
 	// No need to call LockInternalArray(); here as it should already be locked by GC
 
 #if UE_WITH_OBJECT_HANDLE_LATE_RESOLVE
@@ -392,6 +389,12 @@ void FUObjectArray::FreeUObjectIndex(UObjectBase* Object)
 	FUObjectItem* ObjectItem = IndexToObject(Index);
 	UE_CLOG(ObjectItem->GetObject() != Object, LogUObjectArray, Fatal, TEXT("Removing object (0x%016llx) at index %d but the index points to a different object (0x%016llx)!"), (int64)(PTRINT)Object, Index, (int64)(PTRINT)ObjectItem->GetObject());
 
+	// This should only be happening on the game thread (GC runs only on game thread when it's freeing objects)
+	// We loosen the restriction a little bit to allow overwrite of UObjects that are still in the construction loading phase,
+	// this only happens for very narrow cases and should be fine as long as the UObject's in question have thread-safe
+	// destructor and destruction virtuals.
+	check(IsInGameThread() || (IsInAsyncLoadingThread() && ObjectItem->HasAnyFlags(EInternalObjectFlags::AsyncLoadingPhase1)));
+
 	// Can't destroy a refcounted object
 	check(ObjectItem->GetRefCount() == 0 || GExitPurge);
 
diff --git a/Engine/Source/Runtime/CoreUObject/Private/UObject/UObjectGlobals.cpp b/Engine/Source/Runtime/CoreUObject/Private/UObject/UObjectGlobals.cpp
--- a/Engine/Source/Runtime/CoreUObject/Private/UObject/UObjectGlobals.cpp
+++ b/Engine/Source/Runtime/CoreUObject/Private/UObject/UObjectGlobals.cpp
@@ -3542,7 +3542,7 @@ UObject* StaticAllocateObject
 #endif
 
 			// It's not safe to access the object from other threads while it's being destroyed and reconstucted so simulate GC behavior and lock
-			if (ensureMsgf(IsInGameThread(), TEXT("GC lock can only be acquired on the game thread. If you hit this ensure then an object is being reconstructed on a worker thread which is not thread-safe")) && GCGuard)
+			if (IsInGameThread() && GCGuard)
 			{
 				// Take the GC lock if the object is still being async loaded to avoid any race condition, otherwise we don't care about the lock
 				if (Obj->HasAnyInternalFlags(EInternalObjectFlags_AsyncLoading))
@@ -3550,6 +3550,17 @@ UObject* StaticAllocateObject
 					GCGuard->Lock();
 				}
 			}
+			else if (IsInAsyncLoadingThread() && Obj->HasAnyInternalFlags(EInternalObjectFlags::AsyncLoadingPhase1))
+			{
+				// We loosen the restriction a little bit to allow overwrite of UObjects that are still in the construction loading phase,
+				// this only happens for very narrow cases and should be fine as long as the UObject's in question have thread-safe
+				// destructor and destruction virtuals.
+				// No need for the GCGuard here since async loading always has one when its active.
+			}
+			else
+			{
+				ensureMsgf(false, TEXT("GC lock can only be acquired on the game thread. If you hit this ensure then an object is being reconstructed on a worker thread which is not thread-safe"));
+			}
 
 			// Destroy the object.
 			SCOPE_CYCLE_COUNTER(STAT_DestroyObject);
