[UObjectArray]
- Fix missing critical section around object creation listener

#Synced-CL 51263423

diff --git a/Engine/Source/Runtime/CoreUObject/Private/UObject/UObjectArray.cpp b/Engine/Source/Runtime/CoreUObject/Private/UObject/UObjectArray.cpp
--- a/Engine/Source/Runtime/CoreUObject/Private/UObject/UObjectArray.cpp
+++ b/Engine/Source/Runtime/CoreUObject/Private/UObject/UObjectArray.cpp
@@ -326,7 +326,10 @@ void FUObjectArray::AllocateUObjectIndex(UObjectBase* Object, EInternalObjectFla
 
 	UnlockInternalArray();
 
-	//  @todo: threading: lock UObjectCreateListeners
+#if THREADSAFE_UOBJECTS
+	UE::TScopeLock UObjectCreateListenersLock(UObjectCreateListenersCritical);
+#endif
+
 	for (int32 ListenerIndex = 0; ListenerIndex < UObjectCreateListeners.Num(); ListenerIndex++)
 	{
 		UObjectCreateListeners[ListenerIndex]->NotifyUObjectCreated(Object,Index);
@@ -436,6 +439,10 @@ void FUObjectArray::FreeUObjectIndex(UObjectBase* Object)
  */
 void FUObjectArray::AddUObjectCreateListener(FUObjectCreateListener* Listener)
 {
+#if THREADSAFE_UOBJECTS
+	UE::TScopeLock UObjectCreateListenersLock(UObjectCreateListenersCritical);
+#endif
+
 	check(!UObjectCreateListeners.Contains(Listener));
 	UObjectCreateListeners.Add(Listener);
 }
@@ -447,6 +454,9 @@ void FUObjectArray::AddUObjectCreateListener(FUObjectCreateListener* Listener)
  */
 void FUObjectArray::RemoveUObjectCreateListener(FUObjectCreateListener* Listener)
 {
+#if THREADSAFE_UOBJECTS
+	UE::TScopeLock UObjectCreateListenersLock(UObjectCreateListenersCritical);
+#endif
 	int32 NumRemoved = UObjectCreateListeners.RemoveSingleSwap(Listener);
 	check(NumRemoved==1);
 }
@@ -563,6 +573,9 @@ void FUObjectArray::ShutdownUObjectArray()
 		UE_CLOG(UObjectDeleteListeners.Num(), LogUObjectArray, Fatal, TEXT("All UObject delete listeners should be unregistered when shutting down the UObject array"));
 	}
 	{
+#if THREADSAFE_UOBJECTS
+		UE::TScopeLock UObjectCreateListenersLock(UObjectCreateListenersCritical);
+#endif
 		for (int32 Index = UObjectCreateListeners.Num() - 1; Index >= 0; --Index)
 		{
 			FUObjectCreateListener* Listener = UObjectCreateListeners[Index];
diff --git a/Engine/Source/Runtime/CoreUObject/Public/UObject/UObjectArray.h b/Engine/Source/Runtime/CoreUObject/Public/UObject/UObjectArray.h
--- a/Engine/Source/Runtime/CoreUObject/Public/UObject/UObjectArray.h
+++ b/Engine/Source/Runtime/CoreUObject/Public/UObject/UObjectArray.h
@@ -1509,6 +1509,11 @@ private:
 	 * Array of things to notify when a UObjectBase is created
 	 */
 	TArray<FUObjectCreateListener* > UObjectCreateListeners;
+#if THREADSAFE_UOBJECTS
+	// Needs to support recursivicity because NotifyUObjectCreated callback might call RemoveUObjectCreateListener.
+	mutable FTransactionallySafeCriticalSection UObjectCreateListenersCritical;
+#endif
+
 	/**
 	 * Array of things to notify when a UObjectBase is destroyed
 	 */
@@ -1540,7 +1545,8 @@ public:
     {
 		UE::TScopeLock ObjListLock(ObjObjectsCritical);
 #if THREADSAFE_UOBJECTS
-		UE::TScopeLock ListenersLock(UObjectDeleteListenersCritical);
+		UE::TScopeLock CreateListenersLock(UObjectCreateListenersCritical);
+		UE::TScopeLock DeleteListenersLock(UObjectDeleteListenersCritical);
 #endif
         return ObjObjects.GetAllocatedSize() + ObjAvailableList.GetAllocatedSize() + UObjectCreateListeners.GetAllocatedSize() + UObjectDeleteListeners.GetAllocatedSize();
     }
