[IsSavingPackage]
- Finish proper deprecation of GIsSavingPackage

#Synced-CL 51539428

diff --git a/Engine/Source/Editor/UnrealEd/Private/Cooker/DiffPackageWriter.cpp b/Engine/Source/Editor/UnrealEd/Private/Cooker/DiffPackageWriter.cpp
--- a/Engine/Source/Editor/UnrealEd/Private/Cooker/DiffPackageWriter.cpp
+++ b/Engine/Source/Editor/UnrealEd/Private/Cooker/DiffPackageWriter.cpp
@@ -23,6 +23,7 @@
 #include "Serialization/LargeMemoryWriter.h"
 #include "Serialization/JsonWriter.h"
 #include "Serialization/CompactBinarySerialization.h"
+#include "Templates/GuardValueAccessors.h"
 #include "Templates/UniquePtr.h"
 #include "Templates/UnrealTemplate.h"
 #include "UObject/LinkerSave.h"
@@ -624,7 +625,7 @@ void FDiffPackageWriter::WritePackageData(const FPackageInfo& Info, FLargeMemory
 	{
 		// Avoid an assert when calling StaticFindObject during save, which we do to list the "exports" from a package.
 		// We are not writing the discovered objects into the saved package, so the call to StaticFindObject is legal.
-		TGuardValue<bool> GIsSavingPackageGuard(GIsSavingPackage, false);
+		TGuardValueAccessors<bool> IsSavingPackageGuard([]() { return UE::IsSavingPackage(); }, UE::SetIsSavingPackage, false);
 
 		Accumulator.OnSecondSaveComplete(LocalInfo.HeaderSize);
 
diff --git a/Engine/Source/Editor/UnrealEd/Private/Cooker/DiffWriterArchive.cpp b/Engine/Source/Editor/UnrealEd/Private/Cooker/DiffWriterArchive.cpp
--- a/Engine/Source/Editor/UnrealEd/Private/Cooker/DiffWriterArchive.cpp
+++ b/Engine/Source/Editor/UnrealEd/Private/Cooker/DiffWriterArchive.cpp
@@ -19,6 +19,7 @@
 #include "Serialization/AsyncLoading2.h"
 #include "Serialization/MemoryReader.h"
 #include "Serialization/StaticMemoryReader.h"
+#include "Templates/GuardValueAccessors.h"
 #include "Templates/Function.h"
 #include "UObject/LinkerLoad.h"
 #include "UObject/ObjectMacros.h"
@@ -2279,7 +2280,7 @@ void FPackageHeaderData::Initialize()
 			*AssetPathName);
 		check(FPackageName::IsValidLongPackageName(AssetPackageName, true /* bIncludeReadOnlyRoots */));
 
-		TGuardValue<bool> GuardIsSavingPackage(GIsSavingPackage, false);
+		TGuardValueAccessors<bool> IsSavingPackageGuard([]() { return UE::IsSavingPackage(); }, UE::SetIsSavingPackage, false);
 		TGuardValue<int32> GuardAllowUnversionedContentInEditor(GAllowUnversionedContentInEditor, 1);
 		TGuardValue<int32> GuardAllowCookedDataInEditorBuilds(GAllowCookedDataInEditorBuilds, 1);
 
diff --git a/Engine/Source/Runtime/CoreUObject/Private/Serialization/UnversionedPropertySerializationTest.cpp b/Engine/Source/Runtime/CoreUObject/Private/Serialization/UnversionedPropertySerializationTest.cpp
--- a/Engine/Source/Runtime/CoreUObject/Private/Serialization/UnversionedPropertySerializationTest.cpp
+++ b/Engine/Source/Runtime/CoreUObject/Private/Serialization/UnversionedPropertySerializationTest.cpp
@@ -4,6 +4,7 @@
 #include "ProfilingDebugging/CookStats.h"
 #include "Serialization/MemoryReader.h"
 #include "Serialization/MemoryWriter.h"
+#include "Templates/GuardValueAccessors.h"
 #include "UObject/UnrealType.h"
 
 #if UE_ENABLE_UNVERSIONED_PROPERTY_TEST
@@ -146,7 +147,7 @@ struct FUnversionedPropertyTest : public FUnversionedPropertyTestInput
 		FStructuredArchive StructuredArchive(Formatter);
 		FStructuredArchive::FSlot Slot = StructuredArchive.Open();
 
-		TGuardValue<bool> Guard(GIsSavingPackage, false);
+		TGuardValueAccessors<bool> IsSavingPackageGuard([]() { return UE::IsSavingPackage(); }, UE::SetIsSavingPackage, false);
 		FTestInstance Result(Struct);
 		// Call UStruct::SerializeTaggedProperties() directly to bypass
 		// UUserDefinedStruct::SerializeTaggedProperties() for test loading,
diff --git a/Engine/Source/Runtime/CoreUObject/Private/UObject/SavePackage/SavePackageUtilities.cpp b/Engine/Source/Runtime/CoreUObject/Private/UObject/SavePackage/SavePackageUtilities.cpp
--- a/Engine/Source/Runtime/CoreUObject/Private/UObject/SavePackage/SavePackageUtilities.cpp
+++ b/Engine/Source/Runtime/CoreUObject/Private/UObject/SavePackage/SavePackageUtilities.cpp
@@ -176,7 +176,9 @@ static FThreadSafeCounter OutstandingAsyncWrites;
 
 
 // Initialization for GIsSavingPackage
+PRAGMA_DISABLE_DEPRECATION_WARNINGS
 bool GIsSavingPackage = false;
+PRAGMA_ENABLE_DEPRECATION_WARNINGS
 
 // GIsSavingPackage is still accessible globally but is meant only for the game thread
 // so provide another one for any other thread than GT.
@@ -184,26 +186,36 @@ thread_local bool GIsSavingPackageTLS = false;
 
 namespace UE
 {
+	bool& GetIsSavingPackageRef()
+	{
+		if (IsInGameThread())
+		{
+PRAGMA_DISABLE_DEPRECATION_WARNINGS
+			return GIsSavingPackage;
+PRAGMA_ENABLE_DEPRECATION_WARNINGS
+		}
+
+		return GIsSavingPackageTLS;
+	}
+	
+	bool IsSavingPackage()
+	{
+		return GetIsSavingPackageRef();
+	}
+
 	bool IsSavingPackage(UObject* InOuter)
 	{
 		if (InOuter == nullptr)
 		{
-			if (IsInGameThread())
-			{
-				// We need to do a volatile read otherwise Clang/LLVM can decide to do
-				// some crazy optim where both GIsSavingPackage and IsInGameThread()
-				// are read into registry and evaluated together instead of going
-				// one after the other like we want to to avoid touching the GIsSavingPackage
-				// value from other threads and causing TSAN warnings.
-				return *(volatile bool*)&GIsSavingPackage;
-			}
-			else
-			{
-				return GIsSavingPackageTLS;
-			}
+			return IsSavingPackage();
 		}
 		return InOuter->GetPackage()->HasAnyPackageFlags(PKG_IsSaving);
 	}
+
+	void SetIsSavingPackage(bool bValue)
+	{
+		GetIsSavingPackageRef() = bValue;
+	}
 }
 
 FSavePackageSettings& FSavePackageSettings::GetDefaultSettings()
@@ -1181,7 +1193,9 @@ FScopedSavingFlag::FScopedSavingFlag(bool InSavingConcurrent, UPackage* InSavedP
 	{
 		// Don't need to take the GC Lock here since we've already validated we're not in the middle of GC
 		// with the check above. GIsSavingPackage provides the proper protection for avoiding object find during save (i.e. StaticFindObject)
+PRAGMA_DISABLE_DEPRECATION_WARNINGS
 		GIsSavingPackage = true;
+PRAGMA_ENABLE_DEPRECATION_WARNINGS
 	}
 	else
 	{
@@ -1201,7 +1215,9 @@ FScopedSavingFlag::~FScopedSavingFlag()
 {
 	if (IsInGameThread())
 	{
+PRAGMA_DISABLE_DEPRECATION_WARNINGS
 		GIsSavingPackage = false;
+PRAGMA_ENABLE_DEPRECATION_WARNINGS
 	}
 	else
 	{
diff --git a/Engine/Source/Runtime/CoreUObject/Public/UObject/UObjectGlobals.h b/Engine/Source/Runtime/CoreUObject/Public/UObject/UObjectGlobals.h
--- a/Engine/Source/Runtime/CoreUObject/Public/UObject/UObjectGlobals.h
+++ b/Engine/Source/Runtime/CoreUObject/Public/UObject/UObjectGlobals.h
@@ -118,7 +118,7 @@ DECLARE_DWORD_COUNTER_STAT_EXTERN(TEXT("FindObjectFast"),STAT_FindObjectFast,STA
  * Deprecated, use `IsSavingPackage` instead
  * @see IsSavingPackage()
  */
-//UE_DEPRECATED(5.0, "Use the IsSavingPackage() function instead.")
+UE_DEPRECATED(5.8, "Use the IsSavingPackage() function instead.")
 extern COREUOBJECT_API bool					GIsSavingPackage;
 /** This allows loading unversioned cooked content in the editor */
 extern COREUOBJECT_API int32				GAllowUnversionedContentInEditor;
@@ -238,13 +238,28 @@ struct FObjectDuplicationParameters
 // @note: should we start putting the code in the UE namespace?
 namespace UE
 {
+	/**
+	 * Return if the current thread is in a saving scope.
+	 * @returns true the current thread is inside a saving scope.
+	 */
+	COREUOBJECT_API bool IsSavingPackage();
+
 	/**
 	 * Return if the passed in package is currently saving,
-	 * or if any packages are saving if no parameter is passed.
+	 * or if if the current thread is in a saving scope if InOuter == nullptr.
 	 * @param InOuter The object which we want to check if its package is saving
-	 * @returns true if the package is saving, or if any packages are saving if no parameter is passed in.
+	 * @returns true if the package is saving, or if the current thread is in a saving scope if InOuter == nullptr.
+	 */
+	COREUOBJECT_API bool IsSavingPackage(UObject* InOuter);
+
+	/**
+	 * Used to toggle the value which is needed is some rare circumstance.
+	 *
+	 * @param bIsSavingPackage The value we want to set for the current thread.
+	 * 
+	 * @note Should be used with TGuardValueAccessors for RAII style restore of the old flag.
 	 */
-	COREUOBJECT_API bool IsSavingPackage(UObject* InOuter = nullptr);
+	COREUOBJECT_API void SetIsSavingPackage(bool bIsSavingPackage);
 
 } // namespace UE
 
@@ -324,7 +339,7 @@ UE_NODEBUG UE_FORCEINLINE_HINT UObject* StaticFindObjectFast(UClass* Class, UObj
 
 /**
  * Fast and safe version of StaticFindObject that relies on the passed in FName being the object name without any group/package qualifiers.
- * It will not assert on GIsSavingPackage or IsGarbageCollectingAndLockingUObjectHashTables(). If called from within package saving code or GC, will return nullptr
+ * It will not assert on UE::IsSavingPackage() or IsGarbageCollectingAndLockingUObjectHashTables(). If called from within package saving code or GC, will return nullptr
  * This will only find top level packages or subobjects nested directly within a passed in outer.
  *
  * @param	Class			The to be found object's class
@@ -388,7 +403,7 @@ UE_NODEBUG UE_FORCEINLINE_HINT UObject* StaticFindObjectChecked( UClass* Class,
 	return StaticFindObjectChecked(Class, InOuter, Name, bExactClass ? EFindObjectFlags::ExactClass : EFindObjectFlags::None);
 }
 
-/** Internal version of StaticFindObject that will not assert on GIsSavingPackage or IsGarbageCollectingAndLockingUObjectHashTables() */
+/** Internal version of StaticFindObject that will not assert on UE::IsSavingPackage() or IsGarbageCollectingAndLockingUObjectHashTables() */
 COREUOBJECT_API UObject* StaticFindObjectSafe( UClass* Class, UObject* InOuter, FStringView Name, EFindObjectFlags Flags = EFindObjectFlags::None );
 template <UE::CSameAs<bool> ExactClassType>
 UE_EXACTCLASS_BOOL_DEPRECATED("StaticFindObjectSafe")
@@ -399,7 +414,7 @@ UE_NODEBUG UE_FORCEINLINE_HINT UObject* StaticFindObjectSafe( UClass* Class, UOb
 
 /**
  * Tries to find an object in memory. This version uses FTopLevelAssetPath to find the object.
- * Version of StaticFindObject that will not assert on GIsSavingPackage or IsGarbageCollectingAndLockingUObjectHashTables()
+ * Version of StaticFindObject that will not assert on UE::IsSavingPackage() or IsGarbageCollectingAndLockingUObjectHashTables()
  *
  * @param	Class			The to be found object's class
  * @param	ObjectPath		FName pair representing the outer package object and the inner top level object (asset)
@@ -439,7 +454,7 @@ UE_NODEBUG UE_FORCEINLINE_HINT bool StaticFindAllObjectsFast(TArray<UObject*>& O
 /**
  * Fast version of StaticFindAllObjects that relies on the passed in FName being the object name without any group/package qualifiers.
  * This will find all objects matching the specified name and class.
- * This version of StaticFindAllObjectsFast will not assert on GIsSavingPackage or IsGarbageCollecting()
+ * This version of StaticFindAllObjectsFast will not assert on UE::IsSavingPackage() or IsGarbageCollecting()
  * 
  * @param	OutFoundObjects	Array of objects matching the search parameters
  * @param	ObjectClass		The to be found object's class
@@ -478,7 +493,7 @@ UE_NODEBUG UE_FORCEINLINE_HINT bool StaticFindAllObjects(TArray<UObject*>& OutFo
 
 /**
  * Tries to find all objects matching the search parameters in memory. This will handle fully qualified paths of the form /path/packagename.object:subobject and resolve references for you.
- * This version of StaticFindAllObjects will not assert on GIsSavingPackage or IsGarbageCollecting()
+ * This version of StaticFindAllObjects will not assert on UE::IsSavingPackage() or IsGarbageCollecting()
  * 
  * @param	OutFoundObjects	Array of objects matching the search parameters
  * @param	Class			The to be found object's class
@@ -522,7 +537,7 @@ COREUOBJECT_API UObject* StaticFindFirstObject(UClass* Class, FStringView Name,
 
 /**
  * Tries to find the first objects matching the search parameters in memory. This will handle fully qualified paths of the form /path/packagename.object:subobject and resolve references for you.
- * This version of StaticFindFirstObject will not assert on GIsSavingPackage or IsGarbageCollecting()
+ * This version of StaticFindFirstObject will not assert on UE::IsSavingPackage() or IsGarbageCollecting()
  * If multiple objects share the same name the returned object is random and not based on its time of creation unless otherwise specified in Options (see EFindFirstObjectOptions::NativeFirst)
  * This function is slow and should not be used in performance critical situations.
  * 
@@ -2100,7 +2115,7 @@ UE_NODEBUG UE_FORCEINLINE_HINT T* FindObjectChecked( UObject* Outer, FStringView
 }
 
 /**
- * Find an object without asserting on GIsSavingPackage or IsGarbageCollectingAndLockingUObjectHashTables()
+ * Find an object without asserting on UE::IsSavingPackage() or IsGarbageCollectingAndLockingUObjectHashTables()
  * @see StaticFindObjectSafe()
  */
 template< class T > 
@@ -2142,7 +2157,7 @@ inline T* FindFirstObject(FStringView Name, EFindFirstObjectOptions Options = EF
 }
 
 /**
- * Find an optional object with proper handling of potential ambiguity without asserting on GIsSavingPackage or IsGarbageCollecting()
+ * Find an optional object with proper handling of potential ambiguity without asserting on UE::IsSavingPackage() or IsGarbageCollecting()
  * @see StaticFindFirstObject()
  */
 template< class T >
diff --git a/Engine/Source/Runtime/Engine/Private/Materials/HLSLMaterialTranslator.cpp b/Engine/Source/Runtime/Engine/Private/Materials/HLSLMaterialTranslator.cpp
--- a/Engine/Source/Runtime/Engine/Private/Materials/HLSLMaterialTranslator.cpp
+++ b/Engine/Source/Runtime/Engine/Private/Materials/HLSLMaterialTranslator.cpp
@@ -62,6 +62,7 @@
 #include "Serialization/BufferArchive.h"
 #include "Serialization/MemoryReader.h"
 #include "Serialization/ObjectAndNameAsStringProxyArchive.h"
+#include "Templates/GuardValueAccessors.h"
 #include "DerivedDataRequestOwner.h"
 #include "DerivedDataCache.h"
 #include "MaterialCachedData.h"
@@ -1194,11 +1195,7 @@ EHLSLMaterialTranslatorResult FHLSLMaterialTranslator::Translate(bool bForceDisa
 	// but this can happen during save. FindObject is illegal during save because if the discovered objects
 	// are serialized into the package it will cause a crash on package load. But we are not storing the results
 	// of FindObject into the package so it is okay to remove the restriction.
-	TOptional<TGuardValue<bool>> IsSavingPackageGuard;
-	if (IsInGameThread())
-	{
-		IsSavingPackageGuard.Emplace(GIsSavingPackage, false);
-	}
+	TGuardValueAccessors<bool> IsSavingPackageGuard([]() { return UE::IsSavingPackage(); }, UE::SetIsSavingPackage, false);
 
 	static const bool bNoMaterialTranslationDDC = FParse::Param(FCommandLine::Get(), TEXT("nomaterialtranslationddc"));
 
diff --git a/Engine/Source/Runtime/Engine/Private/WorldPartition/WorldPartitionActorDescUtils.cpp b/Engine/Source/Runtime/Engine/Private/WorldPartition/WorldPartitionActorDescUtils.cpp
--- a/Engine/Source/Runtime/Engine/Private/WorldPartition/WorldPartitionActorDescUtils.cpp
+++ b/Engine/Source/Runtime/Engine/Private/WorldPartition/WorldPartitionActorDescUtils.cpp
@@ -10,6 +10,7 @@
 #include "WorldPartition/WorldPartitionLog.h"
 #include "Engine/Level.h"
 #include "GameFramework/Actor.h"
+#include "Templates/GuardValueAccessors.h"
 #include "UObject/AssetRegistryTagsContext.h"
 #include "UObject/CoreRedirects.h"
 #include "Misc/Base64.h"
@@ -44,7 +45,7 @@ FWorldPartitionActorDescUtils::FActorDescInitParams::FActorDescInitParams(const
 {
 	// Avoid an assert when calling StaticFindObject during save to retrieve the actor's class.
 	// Since we are only looking for a native class, the call to StaticFindObject is legit.
-	TGuardValue<bool> GIsSavingPackageGuard(GIsSavingPackage, false);
+	TGuardValueAccessors<bool> IsSavingPackageGuard([]() { return UE::IsSavingPackage(); }, UE::SetIsSavingPackage, false);
 
 	TUniquePtr<FWorldPartitionActorDesc> ActorDesc(InActor->CreateActorDesc());
 
@@ -100,7 +101,7 @@ UClass* FWorldPartitionActorDescUtils::GetActorNativeClassFromString(const FStri
 {
 	// Avoid an assert when calling StaticFindObject during save to retrieve the actor's class.
 	// Since we are only looking for a native class, the call to StaticFindObject is legit.
-	TGuardValue<bool> GIsSavingPackageGuard(GIsSavingPackage, false);
+	TGuardValueAccessors<bool> IsSavingPackageGuard([]() { return UE::IsSavingPackage(); }, UE::SetIsSavingPackage, false);
 
 	// Look for a class redirectors
 	const FString ActorNativeClassName = FWorldPartitionActorDescUtilsPrivate::ResolveClassRedirector(InClassPath);
@@ -171,7 +172,7 @@ void FWorldPartitionActorDescUtils::AppendAssetDataTagsFromActor(const AActor* I
 {
 	// Avoid an assert when calling StaticFindObject during save to retrieve the actor's class.
 	// Since we are only looking for a native class, the call to StaticFindObject is legit.
-	TGuardValue<bool> GIsSavingPackageGuard(GIsSavingPackage, false);
+	TGuardValueAccessors<bool> IsSavingPackageGuard([]() { return UE::IsSavingPackage(); }, UE::SetIsSavingPackage, false);
 
 	TUniquePtr<FWorldPartitionActorDesc> ActorDesc(InActor->CreateActorDesc());
 
@@ -275,7 +276,7 @@ UClass* FWorldPartitionActorDescUtils::GetNativeClassFromClassName(FName InClass
 {
 	// Avoid an assert when calling StaticFindObject during save to retrieve the class.
 	// Since we are only looking for a native class, the call to StaticFindObject is legit.
-	TGuardValue<bool> GIsSavingPackageGuard(GIsSavingPackage, false);
+	TGuardValueAccessors<bool> IsSavingPackageGuard([]() { return UE::IsSavingPackage(); }, UE::SetIsSavingPackage, false);
 
 	// Look for a class redirectors
 	const FString NativeClassName = FWorldPartitionActorDescUtilsPrivate::ResolveClassRedirector(InClassName.ToString());
