Skip over invalid import package names that are the result of orphaned entries from redirected packages.  Cleans up warnings of the form: 'FPackageName: Unable to identify a valid mount point associated with skipped package None. The package root is unknown.'.
#tests Launch citysample (editor) with -game and make sure that the warnings are gone
#Synced-CL 51571507

#Synced-CL 51571507

diff --git a/Engine/Source/Runtime/CoreUObject/Private/Serialization/AsyncLoading2.cpp b/Engine/Source/Runtime/CoreUObject/Private/Serialization/AsyncLoading2.cpp
--- a/Engine/Source/Runtime/CoreUObject/Private/Serialization/AsyncLoading2.cpp
+++ b/Engine/Source/Runtime/CoreUObject/Private/Serialization/AsyncLoading2.cpp
@@ -7680,6 +7680,13 @@ EEventLoadNodeExecutionResult FAsyncPackage2::ProcessLinkerLoadPackageSummary(FA
 		{
 			ImportPackageName = PackageImport.ObjectName;
 		}
+
+		// if the ImportPackageName is unset (NAME_None) then we skip processing it, it's likely an orphaned entry left over from a package redirector
+		if (ImportPackageName.IsNone())
+		{
+			continue;
+		}
+
 		TCHAR NameStr[FName::StringBufferSize];
 		uint32 NameLen = ImportPackageName.ToString(NameStr);
 		bool bIsScriptImport = FPackageName::IsScriptPackage(FStringView(NameStr, NameLen));
