[verse] FDynamicallyTypedValue: skip Deinit body if Type is null to avoid GC hazard for zero-initialized instances

#rnx
#Synced-CL 51488275

diff --git a/Engine/Source/Runtime/CoreUObject/Public/UObject/DynamicallyTypedValue.h b/Engine/Source/Runtime/CoreUObject/Public/UObject/DynamicallyTypedValue.h
--- a/Engine/Source/Runtime/CoreUObject/Public/UObject/DynamicallyTypedValue.h
+++ b/Engine/Source/Runtime/CoreUObject/Public/UObject/DynamicallyTypedValue.h
@@ -209,6 +209,10 @@ namespace UE
 		// Deinitializes this value back to the primordial state.
 		void Deinit()
 		{
+			if (!Type)
+			{
+				return;
+			}
 			Type->DestroyValue(GetDataPointer());
 			FreeData();
 			Type = nullptr;
