Changed FPropertyArrayExternalObjectHandler::FPropertyKey to construct properly if the FProperty is located in a struct (previously only worked in a class)
Fixes the crash in newly added test OverridableSerialization.IDO_Disabled.Cut&Paste.Struct.OverridableArray.SubObject.Add_SubObj_And_Remove

#Synced-CL 51605399

diff --git a/Engine/Source/Runtime/CoreUObject/Private/UObject/OverriddenObjectsExternalHandler.cpp b/Engine/Source/Runtime/CoreUObject/Private/UObject/OverriddenObjectsExternalHandler.cpp
--- a/Engine/Source/Runtime/CoreUObject/Private/UObject/OverriddenObjectsExternalHandler.cpp
+++ b/Engine/Source/Runtime/CoreUObject/Private/UObject/OverriddenObjectsExternalHandler.cpp
@@ -22,7 +22,7 @@ FPropertyArrayExternalObjectHandler::FPropertyKey::FPropertyKey(FName InClassNam
 FPropertyArrayExternalObjectHandler::FPropertyKey::FPropertyKey(UObject* Object, const FProperty* Property)
 {
 	PropertyName = Property->GetFName();
-	UClass* Class = Cast<UClass>(Property->Owner.ToUObject());
+	UStruct* Struct = Cast<UStruct>(Property->Owner.ToUObject());
 
 	if (UE::IsClassOfInstanceDataObjectClass(Object->GetClass()))
 	{
@@ -30,11 +30,11 @@ FPropertyArrayExternalObjectHandler::FPropertyKey::FPropertyKey(UObject* Object,
 		static const FName NAME_TemplateStructName(ANSITEXTVIEW("TemplateStructName"));
 
 		// use the ClassName the IDO class was created from 
-		ClassName = FName(Class->GetMetaData(NAME_TemplateStructName));
+		ClassName = FName(Struct->GetMetaData(NAME_TemplateStructName));
 	}
 	else
 	{
-		ClassName = Class->GetFName();
+		ClassName = Struct->GetFName();
 	}
 }
 
