Added missing FWeakObjectPtr::operator* and FWeakObjectPtr::operator->.

diff --git a/Engine/Source/Runtime/CoreUObject/Public/UObject/WeakObjectPtr.h b/Engine/Source/Runtime/CoreUObject/Public/UObject/WeakObjectPtr.h
--- a/Engine/Source/Runtime/CoreUObject/Public/UObject/WeakObjectPtr.h
+++ b/Engine/Source/Runtime/CoreUObject/Public/UObject/WeakObjectPtr.h
@@ -217,6 +217,16 @@ public:
 	 */
 	COREUOBJECT_API class UObject* Get(/*bool bEvenIfGarbage = false*/) const;
 
+	/** Operators for direct access to Get() */
+	UE_REWRITE UObject& operator*() const
+	{
+		return *Get();
+	}
+	UE_REWRITE UObject* operator->() const
+	{
+		return Get();
+	}
+
 	/** Dereference the weak pointer even if it is marked as Garbage or Unreachable */
 	COREUOBJECT_API class UObject* GetEvenIfUnreachable() const;
 
