[SavePackage]
- Add support for async package saving

#Synced-CL 51597565

diff --git a/Engine/Source/Runtime/CoreUObject/Private/UObject/SavePackage2.cpp b/Engine/Source/Runtime/CoreUObject/Private/UObject/SavePackage2.cpp
--- a/Engine/Source/Runtime/CoreUObject/Private/UObject/SavePackage2.cpp
+++ b/Engine/Source/Runtime/CoreUObject/Private/UObject/SavePackage2.cpp
@@ -22,6 +22,7 @@
 #include "Misc/MessageDialog.h"
 #include "Misc/Optional.h"
 #include "Misc/PackageAccessTracking.h"
+#include "Misc/PackageAccessTrackingOps.h"
 #include "Misc/PackageName.h"
 #include "Misc/PathViews.h"
 #include "Misc/RedirectCollector.h"
@@ -3878,7 +3879,7 @@ AUTORTFM_DISABLE ESavePackageResult SaveHarvestedRealms(FSaveContext& SaveContex
 /**
  * InnerSave is the portion of Save that can be safely run concurrently
  */
-AUTORTFM_DISABLE ESavePackageResult InnerSave(FSaveContext& SaveContext)
+AUTORTFM_DISABLE ESavePackageResult InnerSaveInternal(FSaveContext& SaveContext)
 {
 	// Create slow task dialog if needed
 	const int32 TotalSaveSteps = 3;
@@ -3984,36 +3985,26 @@ FText GetSlowTaskStatusMessage(const FSaveContext& SaveContext)
 	return FText::Format(NSLOCTEXT("Core", "SavingFile", "Saving file: {CleanFilename}..."), Args);
 }
 
-} // end namespace
-
-FSavePackageResultStruct UPackage::Save(UPackage* InPackage, UObject* InAsset, const TCHAR* InFilename,
-	const FSavePackageArgs& SaveArgs)
+void BeginSave(FSaveContext& SaveContext, FScopedSlowTask* SlowTask = nullptr)
 {
-	AutoRTFM::UnreachableIfClosed("TODO: Mark this function as AUTORTFM_DISABLED and fixup callers");
-
-	TRACE_CPUPROFILER_EVENT_SCOPE(UPackage::Save);
-	UE_LOG_CONTEXT("Package", FSoftObjectPath(InPackage));
+	TRACE_CPUPROFILER_EVENT_SCOPE(UPackage::BeginSave);
 
-	COOK_STAT(FScopedDurationTimer FuncSaveTimer(FSavePackageStats::SavePackageTimeSec));
-	COOK_STAT(FSavePackageStats::NumPackagesSaved++);
-	UE_SCOPED_COOK_STAT(InPackage->GetFName(), EPackageEventStatType::SavePackage);
-
-	FSaveContext SaveContext(InPackage, InAsset, InFilename, SaveArgs);
-
-	// Create the slow task dialog if needed
-	const int32 TotalSaveSteps = 8;
-	FScopedSlowTask SlowTask(TotalSaveSteps, GetSlowTaskStatusMessage(SaveContext), SaveContext.IsUsingSlowTask());
-	SlowTask.MakeDialogDelayed(3.0f, SaveContext.IsFromAutoSave());
+	if (SlowTask)
+	{
+		SlowTask->EnterProgressFrame();
+	}
 
-	SlowTask.EnterProgressFrame();
 	SaveContext.Result = ValidatePackage(SaveContext);
 	if (SaveContext.Result != ESavePackageResult::Success)
 	{
-		return SaveContext.Result;
+		return;
 	}
 
 	// Ensures
-	SlowTask.EnterProgressFrame();
+	if (SlowTask)
+	{
+		SlowTask->EnterProgressFrame();
+	}
 	{
 		if (!SaveContext.IsConcurrent())
 		{
@@ -4031,7 +4022,10 @@ FSavePackageResultStruct UPackage::Save(UPackage* InPackage, UObject* InAsset, c
 	}
 
 	// PreSave Asset
-	SlowTask.EnterProgressFrame();
+	if (SlowTask)
+	{
+		SlowTask->EnterProgressFrame();
+	}
 	PreSavePackage(SaveContext);
 	if (SaveContext.GetAsset() && !SaveContext.IsConcurrent())
 	{
@@ -4041,7 +4035,10 @@ FSavePackageResultStruct UPackage::Save(UPackage* InPackage, UObject* InAsset, c
 		SaveContext.SetPreSaveCleanup(ObjectSaveContext.bCleanupRequired);
 	}
 
-	SlowTask.EnterProgressFrame();
+	if (SlowTask)
+	{
+		SlowTask->EnterProgressFrame();
+	}
 	if (!SaveContext.IsConcurrent())
 	{
 		// Route Presave only if not calling concurrently or if the PackageWriter claims already completed, in those case they should be handled separately already
@@ -4052,7 +4049,7 @@ FSavePackageResultStruct UPackage::Save(UPackage* InPackage, UObject* InAsset, c
 			SaveContext.Result = RoutePresave(SaveContext);
 			if (SaveContext.Result != ESavePackageResult::Success)
 			{
-				return SaveContext.Result;
+				return;
 			}
 		}
 
@@ -4063,25 +4060,34 @@ FSavePackageResultStruct UPackage::Save(UPackage* InPackage, UObject* InAsset, c
 		SaveContext.Result = BeginCachePlatformCookedData(SaveContext);
 		if (SaveContext.Result != ESavePackageResult::Success)
 		{
-			return SaveContext.Result;
+			return;
 		}
 
 		// If we're writing to the existing file call ResetLoaders on the Package so that we drop the handle to the file on disk and can write to it
 		// This might end flushing async loading for this package
 		ResetLoadersForSave(SaveContext.GetPackage(), SaveContext.GetFilename());
 	}
+}
 
-	SlowTask.EnterProgressFrame();
-	{
-		// @todo: Once UE::IsSavingPackage() is reworked we should reinstore the saving flag here for the gc lock
-		//FScopedSavingFlag IsSavingFlag(SaveContext.IsConcurrent());
-		SaveContext.Result = InnerSave(SaveContext);
+void InnerSave(FSaveContext& SaveContext, FScopedSlowTask* SlowTask = nullptr)
+{
+	TRACE_CPUPROFILER_EVENT_SCOPE(UPackage::InnerSave);
+
+	SaveContext.Result = InnerSaveInternal(SaveContext);
 
-		// in case of failure or cancellation, do not exit here, still run cleanup (e.g. PostSaveRoot)
+	// in case of failure or cancellation, do not exit here, still run cleanup (e.g. PostSaveRoot)
+}
+
+void EndSave(FSaveContext& SaveContext, FScopedSlowTask* SlowTask = nullptr)
+{
+	TRACE_CPUPROFILER_EVENT_SCOPE(UPackage::EndSave);
+
+	// PostSave Package - edit in memory package and send events if save was successful
+	if (SlowTask)
+	{
+		SlowTask->EnterProgressFrame();
 	}
 
-	// PostSave Asset
-	SlowTask.EnterProgressFrame();
 	if (SaveContext.GetPostSaveRootRequired() && SaveContext.GetAsset())
 	{
 		SaveContext.GetObjectSaveContext().ObjectSaveContextPhase = EObjectSaveContextPhase::PostSave;
@@ -4090,11 +4096,124 @@ FSavePackageResultStruct UPackage::Save(UPackage* InPackage, UObject* InAsset, c
 	}
 
 	// PostSave Package - edit in memory package and send events if save was successful
-	SlowTask.EnterProgressFrame();
+	if (SlowTask)
+	{
+		SlowTask->EnterProgressFrame();
+	}
 	PostSavePackage(SaveContext);
+}
+
+} // end namespace
+
+#define COOK_SAVEPACKAGE_SCOPE(Package) \
+	UE_LOG_CONTEXT("Package", FSoftObjectPath(Package)); \
+	COOK_STAT(FScopedDurationTimer FuncSaveTimer(FSavePackageStats::SavePackageTimeSec)); \
+	UE_SCOPED_COOK_STAT(Package->GetFName(), EPackageEventStatType::SavePackage); \
+	UE_TRACK_REFERENCING_PACKAGE_SCOPED(Package, PackageAccessTrackingOps::NAME_Save);
+
+UE::Tasks::TTask<FSavePackageResultStruct> UPackage::SaveAsync(UPackage* InPackage, UObject* InAsset, const TCHAR* InFilename,
+	const FSavePackageArgs& SaveArgs)
+{
+	AutoRTFM::UnreachableIfClosed("TODO: Mark this function as AUTORTFM_DISABLED and fixup callers");
+
+	COOK_STAT(FSavePackageStats::NumPackagesSaved++);
+
+	auto SaveContext = MakeShared<FSaveContext>(InPackage, InAsset, InFilename, SaveArgs);
+
+	// It's important to know if the request was aborted after BeginSave. We cannot
+	// just look for errors in the context because if BeginSave didn't error out,
+	// we need to execute both InnerSave and EndSave even if InnerSave fails.
+	auto AbortFlag   = MakeShared<bool>(false);
+	auto BeginTask =
+		UE::Tasks::Launch(
+			TEXT("BeginSave"),
+			[SaveContext, AbortFlag]()
+			{
+				AutoRTFM::UnreachableIfClosed("This should never get called under autortfm");
+
+				COOK_SAVEPACKAGE_SCOPE(SaveContext->GetPackage());
+
+				BeginSave(*SaveContext);
+
+				*AbortFlag = (SaveContext->Result != ESavePackageResult::Success);
+			},
+			LowLevelTasks::ETaskPriority::Normal,
+			UE::Tasks::EExtendedTaskPriority::GameThreadHiPri
+		);
+
+	// Can run on any thread
+	auto InnerTask =
+		UE::Tasks::Launch(
+			TEXT("InnerSave"),
+			[SaveContext, AbortFlag]()
+			{
+				if (!*AbortFlag)
+				{
+					AutoRTFM::UnreachableIfClosed("This should never get called under autortfm");
+
+					COOK_SAVEPACKAGE_SCOPE(SaveContext->GetPackage());
+
+					InnerSave(*SaveContext);
+				}
+			},
+			BeginTask
+		);
+	
+	auto EndTask =
+		UE::Tasks::Launch(
+			TEXT("EndSave"),
+			[SaveContext, AbortFlag]()
+			{
+				if (!*AbortFlag)
+				{
+					AutoRTFM::UnreachableIfClosed("This should never get called under autortfm");
+
+					COOK_SAVEPACKAGE_SCOPE(SaveContext->GetPackage());
+
+					EndSave(*SaveContext);
+				}
+
+				return SaveContext->GetFinalResult();
+			},
+			InnerTask,
+			LowLevelTasks::ETaskPriority::Normal,
+			UE::Tasks::EExtendedTaskPriority::GameThreadHiPri
+		);
+
+	return EndTask;
+}
+
+FSavePackageResultStruct UPackage::Save(UPackage* InPackage, UObject* InAsset, const TCHAR* InFilename,
+	const FSavePackageArgs& SaveArgs)
+{
+	AutoRTFM::UnreachableIfClosed("TODO: Mark this function as AUTORTFM_DISABLED and fixup callers");
+
+	TRACE_CPUPROFILER_EVENT_SCOPE(UPackage::Save);
+
+	COOK_SAVEPACKAGE_SCOPE(InPackage);
+	COOK_STAT(FSavePackageStats::NumPackagesSaved++);
+
+	FSaveContext SaveContext(InPackage, InAsset, InFilename, SaveArgs);
+
+	// Create the slow task dialog if needed
+	const int32 TotalSaveSteps = 8;
+	FScopedSlowTask SlowTask(TotalSaveSteps, GetSlowTaskStatusMessage(SaveContext), SaveContext.IsUsingSlowTask());
+	SlowTask.MakeDialogDelayed(3.0f, SaveContext.IsFromAutoSave());
+
+	BeginSave(SaveContext, &SlowTask);
+	if (SaveContext.Result != ESavePackageResult::Success)
+	{
+		return SaveContext.Result;
+	}
+
+	SlowTask.EnterProgressFrame();
+	InnerSave(SaveContext, &SlowTask);
+	EndSave(SaveContext, &SlowTask);
+
 	return SaveContext.GetFinalResult();
 }
 
+#undef COOK_SAVEPACKAGE_SCOPE
 
 ESavePackageResult UPackage::SaveConcurrent(TArrayView<FPackageSaveInfo> InPackages, const FSavePackageArgs& SaveArgs, TArray<FSavePackageResultStruct>& OutResults)
 {
@@ -4160,7 +4279,7 @@ ESavePackageResult UPackage::SaveConcurrent(TArrayView<FPackageSaveInfo> InPacka
 		// Concurrent Part
 		ParallelFor(PackageSaveContexts.Num(), [&PackageSaveContexts] AUTORTFM_DISABLE (int32 PackageIdx)
 			{
-				InnerSave(PackageSaveContexts[PackageIdx]);
+				InnerSaveInternal(PackageSaveContexts[PackageIdx]);
 			});
 	}
 
diff --git a/Engine/Source/Runtime/CoreUObject/Public/UObject/Package.h b/Engine/Source/Runtime/CoreUObject/Public/UObject/Package.h
--- a/Engine/Source/Runtime/CoreUObject/Public/UObject/Package.h
+++ b/Engine/Source/Runtime/CoreUObject/Public/UObject/Package.h
@@ -20,6 +20,7 @@
 #include "Misc/SecureHash.h"
 #include "Misc/WorldCompositionUtility.h"
 #include "Serialization/CustomVersion.h"
+#include "Tasks/Task.h"
 #include "Templates/PimplPtr.h"
 #include "Templates/UniquePtr.h"
 #include "Templates/UnrealTemplate.h"
@@ -1148,6 +1149,23 @@ public:
 	COREUOBJECT_API FMetaData& GetMetaData();
 #endif
 
+	/**
+	* Save one specific object asynchronously (along with any objects it references contained within the same Outer) into an Unreal package.
+	*
+	* @param	InPackage		the package to save the asset into
+	* @param	InAsset			the asset that should be saved into the package
+	* @param	Filename		the name to use for the new package file
+	* @param	SaveArgs		Extended arguments to control the save
+	* @see		FSavePackageContext
+	*
+	* @return	A task containing a FSavePackageResultStruct enum value with the result of saving a package as well as extra data
+	*
+	* @note         Some tasks might need to execute on the gamethread. If the engine isn't ticking, it might be required
+        *               to pump the gamethread task queue manually. (i.e. FTaskGraphInterface::Get().ProcessThreadUntilIdle(ENamedThreads::GameThread) or similar).
+	*/
+	COREUOBJECT_API static UE::Tasks::TTask<FSavePackageResultStruct> SaveAsync(UPackage* InPackage, UObject* InAsset, const TCHAR* InFilename,
+		const FSavePackageArgs& SaveArgs);
+
 	/**
 	* Save one specific object (along with any objects it references contained within the same Outer) into an Unreal package.
 	* 
