Changing some utilities to print full object paths

#Synced-CL 51619035

diff --git a/Engine/Source/Runtime/CoreUObject/Private/Misc/UObjectToken.cpp b/Engine/Source/Runtime/CoreUObject/Private/Misc/UObjectToken.cpp
--- a/Engine/Source/Runtime/CoreUObject/Private/Misc/UObjectToken.cpp
+++ b/Engine/Source/Runtime/CoreUObject/Private/Misc/UObjectToken.cpp
@@ -27,7 +27,7 @@ FUObjectToken::FUObjectToken( const UObject* InObject,  const FText& InLabelOver
 		}
 		else if ( InObject )
 		{
-			CachedText = FText::FromString( InObject->GetName() );
+			CachedText = FText::FromString( InObject->GetPathName() );
 		}
 		else
 		{
diff --git a/Engine/Source/Runtime/CoreUObject/Private/UObject/ScriptCore.cpp b/Engine/Source/Runtime/CoreUObject/Private/UObject/ScriptCore.cpp
--- a/Engine/Source/Runtime/CoreUObject/Private/UObject/ScriptCore.cpp
+++ b/Engine/Source/Runtime/CoreUObject/Private/UObject/ScriptCore.cpp
@@ -3165,7 +3165,7 @@ void UObject::ProcessContextOpcode( FFrame& Stack, RESULT_DECL, bool bCanFailSil
 					? FText::Format(
 						LOCTEXT("AccessNoneUClassContext", "Accessed None trying to read {2} property {0} in {1}"),
 						FText::FromString(Stack.MostRecentProperty->GetName()),
-						FText::FromString(Stack.MostRecentProperty->GetOwner<UClass>()->GetName()),
+						FText::FromString(Stack.MostRecentProperty->GetOwner<UClass>()->GetPathName()),
 						FText::FromString(Stack.MostRecentProperty->HasAllPropertyFlags(CPF_Virtual)?"(virtual)":"(real)"))
 					: FText::Format(
 						LOCTEXT("AccessNoneContext", "Accessed None trying to read {1} property {0} in not an UClass"),
