Add CoreRedirector support for FObjectAndNameAsStringProxyArchive.

#Synced-CL 51626022

diff --git a/Engine/Source/Runtime/CoreUObject/Private/Serialization/ArchiveUObject.cpp b/Engine/Source/Runtime/CoreUObject/Private/Serialization/ArchiveUObject.cpp
--- a/Engine/Source/Runtime/CoreUObject/Private/Serialization/ArchiveUObject.cpp
+++ b/Engine/Source/Runtime/CoreUObject/Private/Serialization/ArchiveUObject.cpp
@@ -4,6 +4,7 @@
 #include "UObject/UObjectGlobals.h"
 #include "UObject/Object.h"
 #include "Serialization/SerializedPropertyScope.h"
+#include "UObject/CoreRedirects.h"
 #include "UObject/ObjectRedirector.h"
 #include "UObject/UnrealType.h"
 #include "Serialization/ObjectAndNameAsStringProxyArchive.h"
@@ -191,6 +192,32 @@ FArchive& FObjectAndNameAsStringProxyArchive::operator<<(UObject*& Obj)
 			Obj = LoadObject<UObject>(nullptr, *LoadedString);
 		}
 
+		// Check for core redirections
+		if (!Obj && bResolveCoreRedirects)
+		{
+			const FCoreRedirectObjectName OldRedirectObjectName(LoadedString);
+
+			// We're matching the logic in FSoftObjectPath::FixupCoreRedirects
+			const bool bIsNative = FPackageName::IsScriptPackage(LoadedString);
+			const FCoreRedirectObjectName NewRedirectObjectName = FCoreRedirects::GetRedirectedName(
+				bIsNative ? ECoreRedirectFlags::Type_AllMask : ECoreRedirectFlags::Type_Object,
+				OldRedirectObjectName
+			);
+
+			if (NewRedirectObjectName != OldRedirectObjectName)
+			{
+				// look up the redirected object by fully qualified pathname
+				LoadedString = NewRedirectObjectName.ToString();
+				Obj = FindObject<UObject>(nullptr, *LoadedString, EFindObjectFlags::None);
+
+				// If we couldn't find it, and we want to load it, do that
+				if (!Obj && bLoadIfFindFails)
+				{
+					Obj = LoadObject<UObject>(nullptr, *LoadedString);
+				}
+			}
+		}
+
 		if (bResolveRedirectors && Obj && Obj->IsA(UObjectRedirector::StaticClass()))
 		{
 			int32 Count = 0;
diff --git a/Engine/Source/Runtime/CoreUObject/Public/Serialization/ObjectAndNameAsStringProxyArchive.h b/Engine/Source/Runtime/CoreUObject/Public/Serialization/ObjectAndNameAsStringProxyArchive.h
--- a/Engine/Source/Runtime/CoreUObject/Public/Serialization/ObjectAndNameAsStringProxyArchive.h
+++ b/Engine/Source/Runtime/CoreUObject/Public/Serialization/ObjectAndNameAsStringProxyArchive.h
@@ -42,6 +42,9 @@ struct FObjectAndNameAsStringProxyArchive : public FNameAsStringProxyArchive
 	 */
 	bool bResolveRedirectors = false;
 
+	/** If enabled, we'll look for core redirects in operator<<(UObject*, FWeakObjectPtr, FObjectPtr if resolved) */
+	bool bResolveCoreRedirects = false;
+
 	COREUOBJECT_API virtual FArchive& operator<<(UObject*& Obj) override;
 	COREUOBJECT_API virtual FArchive& operator<<(FWeakObjectPtr& Obj) override;
 	COREUOBJECT_API virtual FArchive& operator<<(FSoftObjectPtr& Value) override;
