[SavePackage]
- Fix missing AUTORTFM_DISABLE

#Synced-CL 51636719

diff --git a/Engine/Source/Runtime/CoreUObject/Private/UObject/SavePackage2.cpp b/Engine/Source/Runtime/CoreUObject/Private/UObject/SavePackage2.cpp
--- a/Engine/Source/Runtime/CoreUObject/Private/UObject/SavePackage2.cpp
+++ b/Engine/Source/Runtime/CoreUObject/Private/UObject/SavePackage2.cpp
@@ -4069,7 +4069,7 @@ void BeginSave(FSaveContext& SaveContext, FScopedSlowTask* SlowTask = nullptr)
 	}
 }
 
-void InnerSave(FSaveContext& SaveContext, FScopedSlowTask* SlowTask = nullptr)
+AUTORTFM_DISABLE void InnerSave(FSaveContext& SaveContext, FScopedSlowTask* SlowTask = nullptr)
 {
 	TRACE_CPUPROFILER_EVENT_SCOPE(UPackage::InnerSave);
 
@@ -4123,7 +4123,7 @@ UE::Tasks::TTask<FSavePackageResultStruct> UPackage::SaveAsync(UPackage* InPacka
 	// It's important to know if the request was aborted after BeginSave. We cannot
 	// just look for errors in the context because if BeginSave didn't error out,
 	// we need to execute both InnerSave and EndSave even if InnerSave fails.
-	auto AbortFlag   = MakeShared<bool>(false);
+	auto AbortFlag = MakeShared<bool>(false);
 	auto BeginTask =
 		UE::Tasks::Launch(
 			TEXT("BeginSave"),
