[Backout] - CL51612344
#rnx
Original CL Desc
-----------------------------------------------------------------
Performance improvements for async package loading:
- Cache hash indices to avoid redundant binary searches when handling export collisions.

diff --git a/Engine/Source/Runtime/CoreUObject/Private/Serialization/AsyncLoading2.cpp b/Engine/Source/Runtime/CoreUObject/Private/Serialization/AsyncLoading2.cpp
--- a/Engine/Source/Runtime/CoreUObject/Private/Serialization/AsyncLoading2.cpp
+++ b/Engine/Source/Runtime/CoreUObject/Private/Serialization/AsyncLoading2.cpp
@@ -1209,17 +1209,14 @@ private:
 			}
 		}
 
-		void Store(uint64 ExportHash, int32 ObjectIndex, int32 Index = INDEX_NONE)
+		void Store(uint64 ExportHash, int32 ObjectIndex)
 		{
 			checkf(ExportHash, TEXT("Invalid to store a hash of 0 in this map"));
 			checkf(ObjectIndex != InvalidValue, TEXT("Use Remove() instead"));
 
 			TArrayView<uint64> Keys = GetKeys();
 			TArrayView<int32> Values = GetValues();
-			if (Index == INDEX_NONE)
-			{
-				Index = Algo::LowerBound(Keys, ExportHash);
-			}
+			int32 Index = Algo::LowerBound(Keys, ExportHash);
 			if (Index < Capacity && Keys[Index] == ExportHash)
 			{
 				// Slot already exists so reuse it
@@ -1274,15 +1271,11 @@ private:
 			return false;
 		}
 
-		// Returns the value stored if any, returning InvalidValue for not present or removed values
-		int32 Find(uint64 ExportHash, int32* OutHashIndex = nullptr)
+		// Returns the value stored if any, returing InvalidValue for not present or removed values
+		int32 Find(uint64 ExportHash)
 		{
 			TArrayView<uint64> Keys = GetKeys();
 			int32 Index = Algo::LowerBound(Keys, ExportHash);
-			if (OutHashIndex)
-			{
-				*OutHashIndex = Index;
-			}
 			if (Index < Capacity && Keys[Index] == ExportHash)
 			{
 				TArrayView<int32> Values = GetValues();
@@ -1702,10 +1695,10 @@ public:
 		PublicExportMap.PreInsertPublicExports(Exports);
 	}
 
-	void StorePublicExport(uint64 ExportHash, int32 HashIndex, UObject* Object)
+	void StorePublicExport(uint64 ExportHash, UObject* Object)
 	{
 		check(Mutex.IsLocked());
-		PublicExportMap.Store(ExportHash, GUObjectArray.ObjectToIndex(Object), HashIndex);
+		PublicExportMap.Store(ExportHash, GUObjectArray.ObjectToIndex(Object));
 	}
 
 	void RemovePublicExport(uint64 ExportHash, FName ObjectName = NAME_None)
@@ -1721,10 +1714,10 @@ public:
 		}
 	}
 
-	UObject* GetPublicExport(uint64 ExportHash, int32* OutHashIndex = nullptr)
+	UObject* GetPublicExport(uint64 ExportHash)
 	{
 		check(Mutex.IsLocked() || Mutex.IsLockShared());
-		int32 ObjectIndex = PublicExportMap.Find(ExportHash, OutHashIndex);
+		int32 ObjectIndex = PublicExportMap.Find(ExportHash);
 		if (ObjectIndex >= 0)
 		{
 			return static_cast<UObject*>(GUObjectArray.IndexToObject(ObjectIndex)->GetObject());
@@ -2641,12 +2634,11 @@ public:
 		FPublicExportKey Key = FPublicExportKey::MakeKey(PackageId, ExportHash);
 		FPublicExportKey ExistingKey;
 
- 		bool bPreviousExportRemoved = false;
+		bool bPreviousExportRemoved = false;
 		{
 			FWriteLockedPackageRef PackageRef = FindPackageRefCheckedForWrite(Key.GetPackageId());
 
-			int32 HashIndex = INDEX_NONE;
-			UObject* ExistingObject = PackageRef->GetPublicExport(Key.GetExportHash(), &HashIndex);
+			UObject* ExistingObject = PackageRef->GetPublicExport(Key.GetExportHash());
 			if (ExistingObject && ExistingObject != Object)
 			{
 				int32 ExistingObjectIndex = GUObjectArray.ObjectToIndex(ExistingObject);
@@ -2675,7 +2667,7 @@ public:
 				}
 			) > 0;
 
-			PackageRef->StorePublicExport(ExportHash, HashIndex, Object);
+			PackageRef->StorePublicExport(ExportHash, Object);
 			ObjectIndexToPublicExport.Add(ObjectIndex, Key);
 		}
 
@@ -13349,34 +13341,19 @@ bool FPublicExportMapBasicTest::RunTest(const FString& Parameters)
 	typedef FLoadedPackageRef::FPublicExportMap FPublicExportMap;
 
 	FPublicExportMap Map;
-	int32 HashIndex = INDEX_NONE;
 	TestEqual(TEXT("Empty map has no count"), Map.Num(), 0);
 	TestEqual(TEXT("Empty map has no keys"), Map.GetKeys().Num(), 0);
 	TestEqual(TEXT("Empty map has no values"), Map.GetValues().Num(), 0);
 
 	Map.Store(1, 2);
-	TestEqual(TEXT("Find stored key"), Map.Find(1, &HashIndex), 2);
-	TestBetweenInclusive(TEXT("Find stored key hash index"), HashIndex, 0, Map.GetValues().Num());
+	TestEqual(TEXT("Find stored key"), Map.Find(1), 2);
 	TestEqual(TEXT("Store increased size"), Map.Num(), 1);
 
 	// May resize
 	Map.Store(2, 3);
-	HashIndex = INDEX_NONE;
-	TestEqual(TEXT("Find stored key"), Map.Find(2, &HashIndex), 3);
-	TestBetweenInclusive(TEXT("Find stored key hash index"), HashIndex, 0, Map.GetValues().Num());
+	TestEqual(TEXT("Find stored key"), Map.Find(2), 3);
 	TestEqual(TEXT("Store increased size"), Map.Num(), 2);
 
-	HashIndex = INDEX_NONE;
-	TestEqual(TEXT("Find removed key fails"), Map.Find(3, &HashIndex), FPublicExportMap::InvalidValue);
-	Map.Store(3,4,HashIndex);
-	int32 NewHashIndex = INDEX_NONE;
-	TestEqual(TEXT("Find stored key"), Map.Find(3, &NewHashIndex), 4);
-	TestEqual(TEXT("Same hash index"), NewHashIndex, HashIndex);
-	TestBetweenInclusive(TEXT("Find stored key hash index"), NewHashIndex, 0, Map.GetValues().Num());
-	TestEqual(TEXT("Store increased size"), Map.Num(), 3);
-	TestEqual(TEXT("Remove existing key"), Map.Remove(3), true);
-	TestEqual(TEXT("Removed decreased count"), Map.Num(), 2);
-
 	TestEqual(TEXT("Remove existing key"), Map.Remove(1), true);
 	TestEqual(TEXT("Removed decreased count"), Map.Num(), 1);
 	TestEqual(TEXT("Remove removed key fails"), Map.Remove(1), false);
@@ -13397,7 +13374,6 @@ bool FPublicExportMapIterationTest::RunTest(const FString& Parameters)
 	typedef FLoadedPackageRef::FPublicExportMap FPublicExportMap;
 
 	FPublicExportMap Map;
-	int32 HashIndex = INDEX_NONE;
 	TArray<uint64> Keys = { 1, 2, 3, 4, 5, 6, 7, 8 };
 	TArray<int32> Values = { 1, 2, 3, 4, 5, 6, 7, 8 };
 	check(Keys.Num() == Values.Num());
@@ -13433,8 +13409,7 @@ bool FPublicExportMapIterationTest::RunTest(const FString& Parameters)
 
 	for (uint64 Key : Keys)
 	{
-		TestEqual(TEXT("Lookup works"), Map.Find(Key, &HashIndex), Values[(int32)(Key - 1)]);
-		TestBetweenInclusive(TEXT("Find stored key hash index"), HashIndex, 0, Map.GetValues().Num());
+		TestEqual(TEXT("Lookup works"), Map.Find(Key), Values[(int32)(Key - 1)]);
 	}
 	return true;
 }
@@ -13447,7 +13422,6 @@ bool FPublicExportMapPreInsertPublicExportsTest::RunTest(const FString& Paramete
 	// Empty map
 	{
 		FPublicExportMap Map;
-		int32 HashIndex = INDEX_NONE;
 		TArray<FExportMapEntry> Exports = { {.PublicExportHash = 1 }, {.PublicExportHash = 2 }, {.PublicExportHash = 3 }, {.PublicExportHash = 4 } };
 		int32 OriginaMapCount = Map.Num();
 
@@ -13468,15 +13442,13 @@ bool FPublicExportMapPreInsertPublicExportsTest::RunTest(const FString& Paramete
 		// Overwrite preinserted value
 		TestEqual(TEXT("Find stored key with no value set"), Map.Find(Exports[0].PublicExportHash), FPublicExportMap::InvalidValue);
 		Map.Store(Exports[0].PublicExportHash, (int32)Exports[0].PublicExportHash);
-		TestEqual(TEXT("Find stored key"), Map.Find(Exports[0].PublicExportHash, &HashIndex), (int32)Exports[0].PublicExportHash);
-		TestBetweenInclusive(TEXT("Find stored key hash index"), HashIndex, 0, Map.GetValues().Num());
+		TestEqual(TEXT("Find stored key"), Map.Find(Exports[0].PublicExportHash), (int32)Exports[0].PublicExportHash);
 		TestEqual(TEXT("Store increased count"), Map.Num(), OriginaMapCount + 1);
 
 		// Insert new value
 		TestEqual(TEXT("Find missing key"), Map.Find(111), FPublicExportMap::InvalidValue);
 		Map.Store(111, 111);
-		TestEqual(TEXT("Find stored key"), Map.Find(111, &HashIndex), 111);
-		TestBetweenInclusive(TEXT("Find stored key hash index"), HashIndex, 0, Map.GetValues().Num());
+		TestEqual(TEXT("Find stored key"), Map.Find(111), 111);
 		TestEqual(TEXT("Store increased count"), Map.Num(), OriginaMapCount + 2);
 
 		TestEqual(TEXT("Remove existing key"), Map.Remove(111), true);
@@ -13489,7 +13461,6 @@ bool FPublicExportMapPreInsertPublicExportsTest::RunTest(const FString& Paramete
 	// Empty map, populate a single key
 	{
 		FPublicExportMap Map;
-		int32 HashIndex = INDEX_NONE;
 		TArray<FExportMapEntry> Exports = { {.PublicExportHash = 1 } };
 		int32 OriginaMapCount = Map.Num();
 
@@ -13510,15 +13481,13 @@ bool FPublicExportMapPreInsertPublicExportsTest::RunTest(const FString& Paramete
 		// Overwrite preinserted value
 		TestEqual(TEXT("Find stored key with no value set"), Map.Find(Exports[0].PublicExportHash), FPublicExportMap::InvalidValue);
 		Map.Store(Exports[0].PublicExportHash, (int32)Exports[0].PublicExportHash);
-		TestEqual(TEXT("Find stored key"), Map.Find(Exports[0].PublicExportHash, &HashIndex), (int32)Exports[0].PublicExportHash);
-		TestBetweenInclusive(TEXT("Find stored key hash index"), HashIndex, 0, Map.GetValues().Num());
+		TestEqual(TEXT("Find stored key"), Map.Find(Exports[0].PublicExportHash), (int32)Exports[0].PublicExportHash);
 		TestEqual(TEXT("Store increased count"), Map.Num(), OriginaMapCount + 1);
 
 		// Insert new value
 		TestEqual(TEXT("Find missing key"), Map.Find(111), FPublicExportMap::InvalidValue);
 		Map.Store(111, 111);
-		TestEqual(TEXT("Find stored key"), Map.Find(111, &HashIndex), 111);
-		TestBetweenInclusive(TEXT("Find stored key hash index"), HashIndex, 0, Map.GetValues().Num());
+		TestEqual(TEXT("Find stored key"), Map.Find(111), 111);
 		TestEqual(TEXT("Store increased count"), Map.Num(), OriginaMapCount + 2);
 
 		TestEqual(TEXT("Remove existing key"), Map.Remove(111), true);
@@ -13528,10 +13497,10 @@ bool FPublicExportMapPreInsertPublicExportsTest::RunTest(const FString& Paramete
 		TestEqual(TEXT("Find removed key returns no value set"), Map.Find(111), FPublicExportMap::InvalidValue);
 	}
 
+
 	// One existing entry (preinsert multiple keys, some overlapping)
 	{
 		FPublicExportMap Map;
-		int32 HashIndex = INDEX_NONE;
 		TArray<FExportMapEntry> Exports = { {.PublicExportHash = 1 }, {.PublicExportHash = 2 }, {.PublicExportHash = 3 }, {.PublicExportHash = 4 } };
 
 		uint64 PreExistingKey = 2;
@@ -13539,8 +13508,7 @@ bool FPublicExportMapPreInsertPublicExportsTest::RunTest(const FString& Paramete
 		Map.Store(PreExistingKey, PreExistingValue);
 		int32 OriginaMapCount = Map.Num();
 
-		TestEqual(TEXT("Find stored key"), Map.Find(PreExistingKey, &HashIndex), PreExistingValue);
-		TestBetweenInclusive(TEXT("Find stored key hash index"), HashIndex, 0, Map.GetValues().Num());
+		TestEqual(TEXT("Find stored key"), Map.Find(PreExistingKey), PreExistingValue);
 		TestEqual(TEXT("Store increased count"), OriginaMapCount, 1);
 		TestEqual(TEXT("Store increased keys"), Map.GetKeys().Num(), OriginaMapCount);
 		TestEqual(TEXT("Store increased values"), Map.GetValues().Num(), OriginaMapCount);
@@ -13556,7 +13524,7 @@ bool FPublicExportMapPreInsertPublicExportsTest::RunTest(const FString& Paramete
 			if (Key == PreExistingKey)
 			{
 				bFoundPreExistingKey = true;
-				TestEqual(TEXT("Pre-existing keys remain intact after PreInsertPublicExports"), Map.Find(Key, &HashIndex), PreExistingValue);
+				TestEqual(TEXT("Pre-existing keys remain intact after PreInsertPublicExports"), Map.Find(Key), PreExistingValue);
 			}
 			else
 			{
@@ -13569,15 +13537,13 @@ bool FPublicExportMapPreInsertPublicExportsTest::RunTest(const FString& Paramete
 		// Overwrite preinserted value
 		TestEqual(TEXT("Find stored key with no value set"), Map.Find(Exports[0].PublicExportHash), FPublicExportMap::InvalidValue);
 		Map.Store(Exports[0].PublicExportHash, (int32)Exports[0].PublicExportHash);
-		TestEqual(TEXT("Find stored key"), Map.Find(Exports[0].PublicExportHash, &HashIndex), (int32)Exports[0].PublicExportHash);
-		TestBetweenInclusive(TEXT("Find stored key hash index"), HashIndex, 0, Map.GetValues().Num());
+		TestEqual(TEXT("Find stored key"), Map.Find(Exports[0].PublicExportHash), (int32)Exports[0].PublicExportHash);
 		TestEqual(TEXT("Store increased count"), Map.Num(), OriginaMapCount + 1);
 
 		// Insert new value
 		TestEqual(TEXT("Find missing key"), Map.Find(111), FPublicExportMap::InvalidValue);
 		Map.Store(111, 111);
-		TestEqual(TEXT("Find stored key"), Map.Find(111, &HashIndex), 111);
-		TestBetweenInclusive(TEXT("Find stored key hash index"), HashIndex, 0, Map.GetValues().Num());
+		TestEqual(TEXT("Find stored key"), Map.Find(111), 111);
 		TestEqual(TEXT("Store increased count"), Map.Num(), OriginaMapCount + 2);
 
 		TestEqual(TEXT("Remove existing key"), Map.Remove(111), true);
@@ -13590,7 +13556,6 @@ bool FPublicExportMapPreInsertPublicExportsTest::RunTest(const FString& Paramete
 	// One existing entry (preinsert a single entry overlapping)
 	{
 		FPublicExportMap Map;
-		int32 HashIndex = INDEX_NONE;
 		TArray<FExportMapEntry> Exports = { {.PublicExportHash = 1 } };
 
 		uint64 PreExistingKey = 1;
@@ -13598,8 +13563,7 @@ bool FPublicExportMapPreInsertPublicExportsTest::RunTest(const FString& Paramete
 		Map.Store(PreExistingKey, PreExistingValue);
 		int32 OriginaMapCount = Map.Num();
 
-		TestEqual(TEXT("Find stored key"), Map.Find(PreExistingKey, &HashIndex), PreExistingValue);
-		TestBetweenInclusive(TEXT("Find stored key hash index"), HashIndex, 0, Map.GetValues().Num());
+		TestEqual(TEXT("Find stored key"), Map.Find(PreExistingKey), PreExistingValue);
 		TestEqual(TEXT("Store increased count"), OriginaMapCount, 1);
 		TestEqual(TEXT("Store increased keys"), Map.GetKeys().Num(), OriginaMapCount);
 		TestEqual(TEXT("Store increased values"), Map.GetValues().Num(), OriginaMapCount);
@@ -13615,8 +13579,7 @@ bool FPublicExportMapPreInsertPublicExportsTest::RunTest(const FString& Paramete
 			if (Key == PreExistingKey)
 			{
 				bFoundPreExistingKey = true;
-				TestEqual(TEXT("Pre-existing keys remain intact after PreInsertPublicExports"), Map.Find(Key, &HashIndex), PreExistingValue);
-				TestBetweenInclusive(TEXT("Find stored key hash index"), HashIndex, 0, Map.GetValues().Num());
+				TestEqual(TEXT("Pre-existing keys remain intact after PreInsertPublicExports"), Map.Find(Key), PreExistingValue);
 			}
 			else
 			{
@@ -13630,7 +13593,6 @@ bool FPublicExportMapPreInsertPublicExportsTest::RunTest(const FString& Paramete
 	// One existing entry (preinsert a single entry non-overlapping)
 	{
 		FPublicExportMap Map;
-		int32 HashIndex = INDEX_NONE;
 		TArray<FExportMapEntry> Exports = { {.PublicExportHash = 1 } };
 
 		uint64 PreExistingKey = 2;
@@ -13638,8 +13600,7 @@ bool FPublicExportMapPreInsertPublicExportsTest::RunTest(const FString& Paramete
 		Map.Store(PreExistingKey, PreExistingValue);
 		int32 OriginaMapCount = Map.Num();
 
-		TestEqual(TEXT("Find stored key"), Map.Find(PreExistingKey, &HashIndex), PreExistingValue);
-		TestBetweenInclusive(TEXT("Find stored key hash index"), HashIndex, 0, Map.GetValues().Num());
+		TestEqual(TEXT("Find stored key"), Map.Find(PreExistingKey), PreExistingValue);
 		TestEqual(TEXT("Store increased count"), OriginaMapCount, 1);
 		TestEqual(TEXT("Store increased keys"), Map.GetKeys().Num(), OriginaMapCount);
 		TestEqual(TEXT("Store increased values"), Map.GetValues().Num(), OriginaMapCount);
@@ -13655,8 +13616,7 @@ bool FPublicExportMapPreInsertPublicExportsTest::RunTest(const FString& Paramete
 			if (Key == PreExistingKey)
 			{
 				bFoundPreExistingKey = true;
-				TestEqual(TEXT("Pre-existing keys remain intact after PreInsertPublicExports"), Map.Find(Key, &HashIndex), PreExistingValue);
-				TestBetweenInclusive(TEXT("Find stored key hash index"), HashIndex, 0, Map.GetValues().Num());
+				TestEqual(TEXT("Pre-existing keys remain intact after PreInsertPublicExports"), Map.Find(Key), PreExistingValue);
 			}
 			else
 			{
@@ -13669,15 +13629,13 @@ bool FPublicExportMapPreInsertPublicExportsTest::RunTest(const FString& Paramete
 		// Overwrite preinserted value
 		TestEqual(TEXT("Find stored key with no value set"), Map.Find(Exports[0].PublicExportHash), FPublicExportMap::InvalidValue);
 		Map.Store(Exports[0].PublicExportHash, (int32)Exports[0].PublicExportHash);
-		TestEqual(TEXT("Find stored key"), Map.Find(Exports[0].PublicExportHash, &HashIndex), (int32)Exports[0].PublicExportHash);
-		TestBetweenInclusive(TEXT("Find stored key hash index"), HashIndex, 0, Map.GetValues().Num());
+		TestEqual(TEXT("Find stored key"), Map.Find(Exports[0].PublicExportHash), (int32)Exports[0].PublicExportHash);
 		TestEqual(TEXT("Store increased count"), Map.Num(), OriginaMapCount + 1);
 	}
 
 	// One existing entry (preinsert multiple keys, non-overlapping)
 	{
 		FPublicExportMap Map;
-		int32 HashIndex = INDEX_NONE;
 		TArray<FExportMapEntry> Exports = { {.PublicExportHash = 1 }, {.PublicExportHash = 2 }, {.PublicExportHash = 3 }, {.PublicExportHash = 4 } };
 
 		uint64 PreExistingKey = 7;
@@ -13685,8 +13643,7 @@ bool FPublicExportMapPreInsertPublicExportsTest::RunTest(const FString& Paramete
 		Map.Store(PreExistingKey, PreExistingValue);
 		int32 OriginaMapCount = Map.Num();
 
-		TestEqual(TEXT("Find stored key"), Map.Find(PreExistingKey, &HashIndex), PreExistingValue);
-		TestBetweenInclusive(TEXT("Find stored key hash index"), HashIndex, 0, Map.GetValues().Num());
+		TestEqual(TEXT("Find stored key"), Map.Find(PreExistingKey), PreExistingValue);
 		TestEqual(TEXT("Store increased count"), OriginaMapCount, 1);
 		TestEqual(TEXT("Store increased keys"), Map.GetKeys().Num(), OriginaMapCount);
 		TestEqual(TEXT("Store increased values"), Map.GetValues().Num(), OriginaMapCount);
@@ -13701,8 +13658,7 @@ bool FPublicExportMapPreInsertPublicExportsTest::RunTest(const FString& Paramete
 			if (Key == PreExistingKey)
 			{
 				bFoundPreExistingKey = true;
-				TestEqual(TEXT("Pre-existing keys remain intact after PreInsertPublicExports"), Map.Find(Key, &HashIndex), PreExistingValue);
-				TestBetweenInclusive(TEXT("Find stored key hash index"), HashIndex, 0, Map.GetValues().Num());
+				TestEqual(TEXT("Pre-existing keys remain intact after PreInsertPublicExports"), Map.Find(Key), PreExistingValue);
 			}
 			else
 			{
@@ -13715,15 +13671,13 @@ bool FPublicExportMapPreInsertPublicExportsTest::RunTest(const FString& Paramete
 		// Overwrite preinserted value
 		TestEqual(TEXT("Find stored key with no value set"), Map.Find(Exports[0].PublicExportHash), FPublicExportMap::InvalidValue);
 		Map.Store(Exports[0].PublicExportHash, (int32)Exports[0].PublicExportHash);
-		TestEqual(TEXT("Find stored key"), Map.Find(Exports[0].PublicExportHash, &HashIndex), (int32)Exports[0].PublicExportHash);
-		TestBetweenInclusive(TEXT("Find stored key hash index"), HashIndex, 0, Map.GetValues().Num());
+		TestEqual(TEXT("Find stored key"), Map.Find(Exports[0].PublicExportHash), (int32)Exports[0].PublicExportHash);
 		TestEqual(TEXT("Store increased count"), Map.Num(), OriginaMapCount + 1);
 
 		// Insert new value
 		TestEqual(TEXT("Find missing key"), Map.Find(111), FPublicExportMap::InvalidValue);
 		Map.Store(111, 111);
-		TestEqual(TEXT("Find stored key"), Map.Find(111, &HashIndex), 111);
-		TestBetweenInclusive(TEXT("Find stored key hash index"), HashIndex, 0, Map.GetValues().Num());
+		TestEqual(TEXT("Find stored key"), Map.Find(111), 111);
 		TestEqual(TEXT("Store increased count"), Map.Num(), OriginaMapCount + 2);
 
 		TestEqual(TEXT("Remove existing key"), Map.Remove(111), true);
@@ -13736,7 +13690,6 @@ bool FPublicExportMapPreInsertPublicExportsTest::RunTest(const FString& Paramete
 	// Many pre-existing keys
 	{
 		FPublicExportMap Map;
-		int32 HashIndex = INDEX_NONE;
 		// Note we are explicitly overlapping key 1 with the preinsert key 1
 		// Note we are mixing ranges of preinserted and preexisting keys so we are forced to sort the keys together
 		uint64 OverlappingKey = 1;
@@ -13773,13 +13726,11 @@ bool FPublicExportMapPreInsertPublicExportsTest::RunTest(const FString& Paramete
 			if (PreExistingKeys.Contains(Key))
 			{
 				FoundPreExistingKeys.Add(Key);
-				TestEqual(TEXT("Pre-existing keys remain intact after PreInsertPublicExports"), Map.Find(Key, &HashIndex), PreExistingValues[PreExistingKeys.IndexOfByKey(Key)]);
-				TestBetweenInclusive(TEXT("Find stored key hash index"), HashIndex, 0, Map.GetValues().Num());
+				TestEqual(TEXT("Pre-existing keys remain intact after PreInsertPublicExports"), Map.Find(Key), PreExistingValues[PreExistingKeys.IndexOfByKey(Key)]);
 			}
 			else if (Key == OverlappingKey)
 			{
-				TestEqual(TEXT("Ensure preinsert didn't overwrite the existing value"), Map.Find(Key, &HashIndex), (int32)OverlappingKey);
-				TestBetweenInclusive(TEXT("Find stored key hash index"), HashIndex, 0, Map.GetValues().Num());
+				TestEqual(TEXT("Ensure preinsert didn't overwrite the existing value"), Map.Find(Key), (int32)OverlappingKey);
 			}
 			else
 			{
@@ -13792,15 +13743,13 @@ bool FPublicExportMapPreInsertPublicExportsTest::RunTest(const FString& Paramete
 		// Overwrite preinserted value
 		TestEqual(TEXT("Find stored key with no value set"), Map.Find(Exports[0].PublicExportHash), FPublicExportMap::InvalidValue);
 		Map.Store(Exports[0].PublicExportHash, (int32)Exports[0].PublicExportHash);
-		TestEqual(TEXT("Find stored key"), Map.Find(Exports[0].PublicExportHash, &HashIndex), (int32)Exports[0].PublicExportHash);
-		TestBetweenInclusive(TEXT("Find stored key hash index"), HashIndex, 0, Map.GetValues().Num());
+		TestEqual(TEXT("Find stored key"), Map.Find(Exports[0].PublicExportHash), (int32)Exports[0].PublicExportHash);
 		TestEqual(TEXT("Store increased count"), Map.Num(), PreExistingValues.Num() + 1);
 
 		// Insert new value
 		TestEqual(TEXT("Find missing key"), Map.Find(111), FPublicExportMap::InvalidValue);
 		Map.Store(111, 111);
-		TestEqual(TEXT("Find stored key"), Map.Find(111, &HashIndex), 111);
-		TestBetweenInclusive(TEXT("Find stored key hash index"), HashIndex, 0, Map.GetValues().Num());
+		TestEqual(TEXT("Find stored key"), Map.Find(111), 111);
 		TestEqual(TEXT("Store increased count"), Map.Num(), PreExistingValues.Num() + 2);
 
 		TestEqual(TEXT("Remove existing key"), Map.Remove(111), true);
