Add new ESaveFlags::SAVE_KeepPersistentGUID which allows the caller to request that the persistent GUID of a package not be regenerated during save.

#rnx

- When ever a package is saved (not an auto save or procedural save) to a new file path then the persistent guid of the package will be regenerated but we do have some very specific use cases where we do not want this to be the case.
-- This issue was originally found in a system that first saves packages to the temp directory and then moved them back to the original location once a number of steps have been completed. Since the new packages eventually replace the original ones we want the persistent GUIDs to remain unchaged but because the package save targets a different path they were being regenerated.
- The new flag ESaveFlags::SAVE_KeepPersistentGUID replaces ESaveFlags::SAVE_KeepGUID as it was deprecated and due for removal anyway.
-- I have tried to make it clear in the documentation for this new flag that it is potentially dangerous and could result in duplicate GUIDs if abused.
-- Fixed some existing issues with code documentation formatting while I was making changes to the code area.
- ::PreSavePackage was already checking the save context to see if we are performing an auto or procedural save, so I moved those checks and the check for the new flag to a single FSaveContext::ShouldKeepPersistentGUID method.

#Synced-CL 51587407

diff --git a/Engine/Source/Runtime/CoreUObject/Private/UObject/SavePackage/SaveContext.h b/Engine/Source/Runtime/CoreUObject/Private/UObject/SavePackage/SaveContext.h
--- a/Engine/Source/Runtime/CoreUObject/Private/UObject/SavePackage/SaveContext.h
+++ b/Engine/Source/Runtime/CoreUObject/Private/UObject/SavePackage/SaveContext.h
@@ -885,6 +885,11 @@ public:
 		return (SaveArgs.SaveFlags & ESaveFlags::SAVE_RehydratePayloads) != 0;
 	}
 
+	bool ShouldKeepPersistentGUID() const
+	{
+		return IsProceduralSave() || IsFromAutoSave() || (SaveArgs.SaveFlags & ESaveFlags::SAVE_KeepPersistentGUID) != 0;
+	}
+
 	bool GetPostSaveRootRequired() const
 	{
 		return bPostSaveRootRequired;
diff --git a/Engine/Source/Runtime/CoreUObject/Private/UObject/SavePackage2.cpp b/Engine/Source/Runtime/CoreUObject/Private/UObject/SavePackage2.cpp
--- a/Engine/Source/Runtime/CoreUObject/Private/UObject/SavePackage2.cpp
+++ b/Engine/Source/Runtime/CoreUObject/Private/UObject/SavePackage2.cpp
@@ -290,10 +290,10 @@ FORCEINLINE void EnsurePackageLocalization(UPackage* InPackage)
 void PreSavePackage(FSaveContext& SaveContext)
 {
 #if WITH_EDITOR
-	// if the in memory package filename is different the filename we are saving it to,
-	// regenerate a new persistent id for it.
+	// If the in memory package filename is different the filename we are saving it to and if the save context allows for it,
+	// regenerate a new persistent id for the package.
 	UPackage* Package = SaveContext.GetPackage();
-	if (!SaveContext.IsProceduralSave() && !SaveContext.IsFromAutoSave() && !Package->GetLoadedPath().IsEmpty() && Package->GetLoadedPath() != SaveContext.GetTargetPackagePath())
+	if (!SaveContext.ShouldKeepPersistentGUID() && !Package->GetLoadedPath().IsEmpty() && Package->GetLoadedPath() != SaveContext.GetTargetPackagePath())
 	{
 		Package->SetPersistentGuid(FGuid::NewGuid());
 	}
diff --git a/Engine/Source/Runtime/CoreUObject/Public/UObject/ObjectMacros.h b/Engine/Source/Runtime/CoreUObject/Public/UObject/ObjectMacros.h
--- a/Engine/Source/Runtime/CoreUObject/Public/UObject/ObjectMacros.h
+++ b/Engine/Source/Runtime/CoreUObject/Public/UObject/ObjectMacros.h
@@ -99,9 +99,10 @@ enum ESaveFlags
 {
 	SAVE_None						= 0x00000000,	///< No flags
 	SAVE_NoError					= 0x00000001,	///< Don't generate errors on save
-	SAVE_FromAutosave				= 0x00000002,   ///< Used to indicate this save was initiated automatically
+	SAVE_FromAutosave				= 0x00000002,	///< Used to indicate this save was initiated automatically
 	SAVE_KeepDirty					= 0x00000004,	///< Do not clear the dirty flag when saving
-	SAVE_KeepGUID UE_DEPRECATED(5.6, "KeepPackageGUIDOnSave is no longer supported, because package save no longer records arbitrary guids but instead records a hash of the package's bytes - PackageSavedHash.") = 0x00000008,
+	SAVE_KeepPersistentGUID			= 0x00000008,	///< Do not replace the persistent GUID when saving a package to a different path. This is potentially dangerous as it could result in duplicate GUIDs and so should only be used in cases
+													///< where we know that the resulting package is temporary or will eventually replace the original package file.
 	SAVE_Async						= 0x00000010,	///< Save to a memory writer, then actually write to disk async
 	SAVE_Unversioned_Native			= 0x00000020,	///< Save all versions as zero. Upon load this is changed to the current version. This is only reasonable to use with full cooked builds for distribution.
 	SAVE_CutdownPackage				= 0x00000040,	///< Saving cutdown packages in a temp location WITHOUT renaming the package.
@@ -113,10 +114,10 @@ enum ESaveFlags
 	SAVE_CompareLinker UE_DEPRECATED(5.6, "Comparing Linkers is no longer used. Contact Epic if you need this functionality.") = 0x00001000,
 	SAVE_BulkDataByReference		= 0x00002000,	///< When saving to a different file than the package's LoadedPath, point bulkdata in the new file to be loaded from the original file.
 	SAVE_Unversioned_Properties		= 0x00004000,	///< Properties are saved without property name information, and are saved/loaded in the order of the current binary.
-													///This is only reasonable to use with full cooked builds for distribution, or with a domain that selects the payload matching a given version externally.
+													///< This is only reasonable to use with full cooked builds for distribution, or with a domain that selects the payload matching a given version externally.
 	SAVE_Optional					= 0x00008000,	///< Indicate that we to save optional exports. This flag is only valid while cooking. Optional exports are filtered if not specified during cooking. 
-	SAVE_AllowTimeout				= 0x00010000,   ///< If present, if SavePackage is blocked by an asynchronous operation, it will quickly fail with ESavePackageResult::Timeout, otherwise it will wait for a while and then fail with ESavePackageResult::Error.
-	SAVE_RehydratePayloads			= 0x00020000,   ///< Any virtualized payloads in the package should be pulled and stored locally in the package file during the save
+	SAVE_AllowTimeout				= 0x00010000,	///< If present, if SavePackage is blocked by an asynchronous operation, it will quickly fail with ESavePackageResult::Timeout, otherwise it will wait for a while and then fail with ESavePackageResult::Error.
+	SAVE_RehydratePayloads			= 0x00020000,	///< Any virtualized payloads in the package should be pulled and stored locally in the package file during the save
 	SAVE_CookSoftPackageReferences	= 0x00040000,	///< Indicate that we want to save soft package references for cooked packages
 	SAVE_Unversioned = SAVE_Unversioned_Native | SAVE_Unversioned_Properties,
 };
