Don't reset UObject name when unhashing during destruction to allow better debugging

#Synced-CL 51662199

diff --git a/Engine/Source/Runtime/CoreUObject/Private/UObject/Obj.cpp b/Engine/Source/Runtime/CoreUObject/Private/UObject/Obj.cpp
--- a/Engine/Source/Runtime/CoreUObject/Private/UObject/Obj.cpp
+++ b/Engine/Source/Runtime/CoreUObject/Private/UObject/Obj.cpp
@@ -1115,7 +1115,20 @@ void UObject::BeginDestroy()
 	// Remove from linker's export table.
 	SetLinker( NULL, INDEX_NONE );
 
-	LowLevelRename(NAME_None);
+	// Remove from the internal hash tables
+	// The object was hashed in the open (see UObjectBase constructor calling UObjectBase::AddObject) so we need to unhash it in the open too.
+	// There's two places that can call this function:
+	// - GC which always runs in the open
+	// - Constructing an object on top of an existing object (may be run in the closed) where the object is first destroyed (here) and then re-constucted (this is where UObjectBase::AddObject is being called again)
+	// Aborting the reconstuction will always leave the object in a hashed state so no need for an OnAbort handler
+	UE_AUTORTFM_OPEN
+	{
+#if STATS || ENABLE_STATNAMEDEVENTS_UOBJECT
+		((UObject*)this)->ResetStatID(); // reset the stat id since this thing is about to be destroyed
+#endif
+		UnhashObject(this);
+	};
+
 	// Remove any associated external package, at this point
 	SetExternalPackage(nullptr);
 
diff --git a/Engine/Source/Runtime/CoreUObject/Private/UObject/UObjectBase.cpp b/Engine/Source/Runtime/CoreUObject/Private/UObject/UObjectBase.cpp
--- a/Engine/Source/Runtime/CoreUObject/Private/UObject/UObjectBase.cpp
+++ b/Engine/Source/Runtime/CoreUObject/Private/UObject/UObjectBase.cpp
@@ -180,12 +180,13 @@ UObjectBase::UObjectBase(UClass* InClass,
 UObjectBase::~UObjectBase()
 {
 	// If not initialized, skip out.
-	if( UObjectInitialized() && ClassPrivate.GetNoResolve() && !GIsCriticalError)
+	if (UObjectInitialized() && ClassPrivate.GetNoResolve())
 	{
 		// Validate it.
 		check(IsValidLowLevelForDestruction());
-		check(GetFName() == NAME_None);
-		checkf(InternalIndex == INDEX_NONE, TEXT("Object destroyed outside of GC (InternalIndex=%d, expected %d)"), InternalIndex, INDEX_NONE);	
+		constexpr EObjectFlags DestroyedFlags = RF_BeginDestroyed | RF_FinishDestroyed;
+		checkf((GetFlags() & DestroyedFlags) == DestroyedFlags, TEXT("Object %s was not properly destroyed"), *GetFName().ToString());
+		checkf(InternalIndex == INDEX_NONE, TEXT("Object %s destroyed outside of GC (InternalIndex=%d, expected %d)"), *GetFName().ToString(), InternalIndex, INDEX_NONE);
 	}
 
 #if CSV_PROFILER_STATS && CSV_TRACK_UOBJECT_COUNT
