Protect against potential integer overflow when loading a FPackageTrailer.

#rnx

- Since HeaderLength and PayloadsDataLength are read from disk it would be possible for a corrupted file to have incorrect values that would cause an integer overflow and a fail to be caught by the check.
- By using FGuardedInt64 we can catch the overflow as well as incorrect values landing outside of the archive length.

#Synced-CL 51587407

diff --git a/Engine/Source/Runtime/CoreUObject/Private/UObject/PackageTrailer.cpp b/Engine/Source/Runtime/CoreUObject/Private/UObject/PackageTrailer.cpp
--- a/Engine/Source/Runtime/CoreUObject/Private/UObject/PackageTrailer.cpp
+++ b/Engine/Source/Runtime/CoreUObject/Private/UObject/PackageTrailer.cpp
@@ -6,6 +6,7 @@
 #include "Algo/RemoveIf.h"
 #include "CoreGlobals.h"
 #include "HAL/FileManager.h"
+#include "Math/GuardedInt.h"
 #include "Misc/ConfigCacheIni.h"
 #include "Misc/PackagePath.h"
 #include "Misc/PackageSegment.h"
@@ -563,7 +564,8 @@ bool FPackageTrailer::TryLoad(FArchive& Ar)
 	Ar << Header.PayloadsDataLength;
 
 	// The header and the payloads should not exceed the remaining data in the archive
-	if (PackageTrailerOffset + Header.HeaderLength + Header.PayloadsDataLength > (uint64)Ar.TotalSize())
+	const FGuardedInt64 FinalTrailerBytePos = FGuardedInt64(PackageTrailerOffset) + Header.HeaderLength + FGuardedInt64(Header.PayloadsDataLength);
+	if (FinalTrailerBytePos.InvalidOrGreaterThan(Ar.TotalSize()))
 	{
 		Ar.SetError();
 		return false;
