Allow remote object serialization debugging without recompilaiton in dev builds

#Synced-CL 51662199

diff --git a/Engine/Source/Runtime/CoreUObject/Private/UObject/RemoteObjectSerialization.cpp b/Engine/Source/Runtime/CoreUObject/Private/UObject/RemoteObjectSerialization.cpp
--- a/Engine/Source/Runtime/CoreUObject/Private/UObject/RemoteObjectSerialization.cpp
+++ b/Engine/Source/Runtime/CoreUObject/Private/UObject/RemoteObjectSerialization.cpp
@@ -48,6 +48,16 @@ static FAutoConsoleVariableRef CVarUseImmutableArchetypes(
 	GUseImmutableArchetypes,
 	TEXT("Whether to use immutable archetypes when serializing remote object data"));
 
+#if UE_WITH_REMOTEOBJECTARCHIVE_DEBUGGING
+
+int32 GDebugRemoteObjectSerialization = 0;
+static FAutoConsoleVariableRef CVarDebugRemoteObjectSerialization(
+	TEXT("ro.DebugRemoteObjectSerialization"),
+	GDebugRemoteObjectSerialization,
+	TEXT("Whether to dump detailed remote object serialization stats to log"));
+
+#endif // UE_WITH_REMOTEOBJECTARCHIVE_DEBUGGING
+
 const UObject* FindImmutableArchetype(const UObject* InObj);
 
 using FNameIndexType = FRemoteObjectTables::FNameIndexType;
@@ -333,6 +343,11 @@ protected:
 
 	void DumpStatsToLog()
 	{
+		if (!GDebugRemoteObjectSerialization)
+		{
+			return;
+		}
+
 		int64 NameTableSize = 0;
 		{
 			TArray<uint8> NameTableData;
@@ -456,29 +471,39 @@ public:
 #if UE_WITH_REMOTEOBJECTARCHIVE_DEBUGGING
 	void SetSerializationScope(const TCHAR* InScope)
 	{
-		SerializationScope = InScope;
+		if (GDebugRemoteObjectSerialization)
+		{
+			SerializationScope = InScope;
+		}
 	}
 
 	virtual void Serialize(void* Data, int64 Num) override
 	{
-		FProperty* CurrentProperty = nullptr;
-		if (const FArchiveSerializedPropertyChain* PropChain = T::GetSerializedPropertyChain())
+		if (!GDebugRemoteObjectSerialization)
+		{
+			T::Serialize(Data, Num);
+		}
+		else
 		{
-			if (PropChain->GetNumProperties())
+			FProperty* CurrentProperty = nullptr;
+			if (const FArchiveSerializedPropertyChain* PropChain = T::GetSerializedPropertyChain())
 			{
-				CurrentProperty = PropChain->GetPropertyFromRoot(0);
+				if (PropChain->GetNumProperties())
+				{
+					CurrentProperty = PropChain->GetPropertyFromRoot(0);
+				}
 			}
-		}
 
-		TMap<FProperty*, FMigratedPropertyStats>& ObjectStats = ObjectPropertyStats.FindOrAdd(SerializationScope);
-		FMigratedPropertyStats& PropertyStats = ObjectStats.FindOrAdd(CurrentProperty);
+			TMap<FProperty*, FMigratedPropertyStats>& ObjectStats = ObjectPropertyStats.FindOrAdd(SerializationScope);
+			FMigratedPropertyStats& PropertyStats = ObjectStats.FindOrAdd(CurrentProperty);
 
-		const int64 StartPos = T::Tell();
+			const int64 StartPos = T::Tell();
 
-		T::Serialize(Data, Num);
+			T::Serialize(Data, Num);
 
-		PropertyStats.Size += T::Tell() - StartPos;
-		PropertyStats.Count++;
+			PropertyStats.Size += T::Tell() - StartPos;
+			PropertyStats.Count++;
+		}
 	}
 #endif // UE_WITH_REMOTEOBJECTARCHIVE_DEBUGGING
 };
diff --git a/Engine/Source/Runtime/CoreUObject/Private/UObject/RemoteObjectSerialization.h b/Engine/Source/Runtime/CoreUObject/Private/UObject/RemoteObjectSerialization.h
--- a/Engine/Source/Runtime/CoreUObject/Private/UObject/RemoteObjectSerialization.h
+++ b/Engine/Source/Runtime/CoreUObject/Private/UObject/RemoteObjectSerialization.h
@@ -7,7 +7,7 @@
 namespace UE::RemoteObject::Handle { struct FRemoteObjectStub; }
 
 #ifndef UE_WITH_REMOTEOBJECTARCHIVE_DEBUGGING
-#define UE_WITH_REMOTEOBJECTARCHIVE_DEBUGGING 0
+#define UE_WITH_REMOTEOBJECTARCHIVE_DEBUGGING (!(UE_BUILD_SHIPPING || UE_BUILD_TEST) && !NO_LOGGING)
 #endif
 
 namespace UE::RemoteObject::Serialization
