Fix concurrent accesses to PIEGuidMap

#Synced-CL 51591079

diff --git a/Engine/Source/Runtime/CoreUObject/Private/UObject/LazyObjectPtr.cpp b/Engine/Source/Runtime/CoreUObject/Private/UObject/LazyObjectPtr.cpp
--- a/Engine/Source/Runtime/CoreUObject/Private/UObject/LazyObjectPtr.cpp
+++ b/Engine/Source/Runtime/CoreUObject/Private/UObject/LazyObjectPtr.cpp
@@ -5,6 +5,7 @@
 =============================================================================*/
 
 #include "UObject/LazyObjectPtr.h"
+#include "Async/TransactionallySafeMutex.h"
 #include "Misc/CommandLine.h"
 #include "Misc/App.h"
 #include "UObject/Package.h"
@@ -15,6 +16,7 @@
 static FUObjectAnnotationSparseSearchable<FUniqueObjectGuid,true> GuidAnnotation;
 
 static TMap<FGuid, FGuid> PIEGuidMap[MAX_PIE_INSTANCES];
+static FTransactionallySafeMutex PIEGuidMapLock;
 
 static FGuid RemapGuid(const FUniqueObjectGuid& Guid, int32 PIEInstanceID)
 {
@@ -27,6 +29,8 @@ static FGuid RemapGuid(const FUniqueObjectGuid& Guid, int32 PIEInstanceID)
 	check(PIEInstanceID != INDEX_NONE);
 	check(PIEInstanceID < MAX_PIE_INSTANCES);
 
+	UE::TScopeLock Lock(PIEGuidMapLock);
+
 	FGuid& FoundGuid = PIEGuidMap[PIEInstanceID].FindOrAdd(Guid.GetGuid());
 
 	if (!FoundGuid.IsValid())
@@ -197,5 +201,6 @@ void FLazyObjectPtr::ResetPIEFixups()
 	const int32 PlayInEditorID = UE::GetPlayInEditorID();
 	check(PlayInEditorID != -1);
 	check(PlayInEditorID < MAX_PIE_INSTANCES);
+	UE::TScopeLock Lock(PIEGuidMapLock);
 	PIEGuidMap[PlayInEditorID].Reset();
 }
