Fix floating point warning

#Synced-CL 51449702

diff --git a/Engine/Source/Runtime/CoreUObject/Tests/PluginManagerTest.cpp b/Engine/Source/Runtime/CoreUObject/Tests/PluginManagerTest.cpp
--- a/Engine/Source/Runtime/CoreUObject/Tests/PluginManagerTest.cpp
+++ b/Engine/Source/Runtime/CoreUObject/Tests/PluginManagerTest.cpp
@@ -188,7 +188,7 @@ namespace UE::CoreUObject::Private::Tests
 			});
 		FText Reason;
 		CHECK(PluginManager.UnmountExplicitlyLoadedPlugin(TEXT("MyTestPlugin"), &Reason));
-		FTSTicker::GetCoreTicker().Tick(0.01);
+		FTSTicker::GetCoreTicker().Tick(0.01f);
 		CHECK(!FPackageName::MountPointExists(TEXT("/MyTestPlugin")));
 		CHECK(PluginPackageName != PluginPackage->GetFName()); //package should have been renamed.
 		CHECK(PluginPackage->HasAnyInternalFlags(EInternalObjectFlags::Garbage));
