refactor FRemoteExecutor and related subsystems to support testability:
- you can now create multiple executors and subsystem instances and there are ActivateForTesting() functions for re-seating the active global instance
- delegates moved to be per-instance so we take a pointer to a struct of delegates
- replaced calls to FRemoteServerId::GetLocalServerId() to storing it on the instance at init time so we could have multiple existing in a process

#Synced-CL 51619363

diff --git a/Engine/Source/Runtime/CoreUObject/Private/UObject/RemoteExecutor.cpp b/Engine/Source/Runtime/CoreUObject/Private/UObject/RemoteExecutor.cpp
--- a/Engine/Source/Runtime/CoreUObject/Private/UObject/RemoteExecutor.cpp
+++ b/Engine/Source/Runtime/CoreUObject/Private/UObject/RemoteExecutor.cpp
@@ -26,17 +26,7 @@ static FAutoConsoleVariableRef CVarRemoteExecutorAllowSpeculationMode(
 
 namespace UE::RemoteExecutor
 {
-// Prevent us from having to match this definition to the declartion in the header just to initialize the delegate
 #define DEFINE_REMOTE_EXECUTOR_DELEGATE(x) decltype(x) x
-	DEFINE_REMOTE_EXECUTOR_DELEGATE(TickNetworkDelegate);
-	DEFINE_REMOTE_EXECUTOR_DELEGATE(BeginMultiServerCommitDelegate);
-	DEFINE_REMOTE_EXECUTOR_DELEGATE(ReadyMultiServerCommitDelegate);
-	DEFINE_REMOTE_EXECUTOR_DELEGATE(AbandonMultiServerCommitDelegate);
-	DEFINE_REMOTE_EXECUTOR_DELEGATE(EndMultiServerCommitDelegate);
-	DEFINE_REMOTE_EXECUTOR_DELEGATE(ReadyRemoteMultiServerCommitDelegate);
-	DEFINE_REMOTE_EXECUTOR_DELEGATE(AbortRemoteMultiServerCommitDelegate);
-	DEFINE_REMOTE_EXECUTOR_DELEGATE(FetchNextDeferredRPCDelegate);
-
 	DEFINE_REMOTE_EXECUTOR_DELEGATE(OnTransactionQueuedDelegate);
 	DEFINE_REMOTE_EXECUTOR_DELEGATE(OnTransactionStartingDelegate);
 	DEFINE_REMOTE_EXECUTOR_DELEGATE(OnTransactionCompletedDelegate);
@@ -135,11 +125,46 @@ class FRemoteExecutor
 	uint32 IterationsSinceNetworkPump = 0;
 
 public:
+	FRemoteServerId LocalServerId;
+	FRemoteExecutorDelegates* Delegates = nullptr;
+	int32 SpeculationModeScopeCount = 0;
+
+	FRemoteExecutor() = default;
+
+	FRemoteExecutor(FRemoteServerId InLocalServerId, FRemoteExecutorDelegates* InDelegates)
+		: LocalServerId(InLocalServerId)
+		, Delegates(InDelegates)
+	{
+	}
+
+	~FRemoteExecutor()
+	{
+		for (FRemoteSubsystemBase* Subsystem : Subsystems)
+		{
+			delete Subsystem;
+		}
+		Subsystems.Empty();
+	}
+
+	FRemoteServerId GetLocalServerId() const { return LocalServerId; }
+
+	bool HasPendingWork() const
+	{
+		return PendingWorks.Num() > 0
+			|| PendingMigrationWorks.Num() > 0
+			|| PendingWorkLists.Num() > 0
+			|| ActiveRemoteMultiServerCommitRequestId.IsValid();
+	}
+
 	FRemoteExecutorWork* ExecutingWork = nullptr;
 
 	// Pump the network next chance we get
 	bool bShouldPumpNetwork = false;
 
+	// When true, ExecutePendingWork breaks out of its outer loop after each iteration,
+	// allowing the caller to pump other executors before continuing.
+	bool bSingleStepModeForTesting = false;
+
 	// These should only ever be accessed in an AUTORTFM_OPEN scope
 	UE::RemoteExecutor::ERemoteExecutorAbortReason AbortReason = UE::RemoteExecutor::ERemoteExecutorAbortReason::Unspecified;
 	FString AbortReasonDescription;
@@ -161,6 +186,22 @@ public:
 		Subsystems.Add(Subsystem);
 	}
 
+	void ActivateSubsystemsForTesting()
+	{
+		for (FRemoteSubsystemBase* Subsystem : Subsystems)
+		{
+			Subsystem->ActivateForTesting();
+		}
+	}
+
+	void DeactivateSubsystemsForTesting()
+	{
+		for (FRemoteSubsystemBase* Subsystem : Subsystems)
+		{
+			Subsystem->DeactivateForTesting();
+		}
+	}
+
 	FRemoteTransactionId GenerateNextTransactionId()
 	{
 		uint32 RawRequestId = NextTransactionRequestId++;
@@ -229,7 +270,7 @@ private:
 				// pump the network, reset our state variables
 				IterationsSinceNetworkPump = 0;
 				bShouldPumpNetwork = false;
-				UE::RemoteExecutor::TickNetworkDelegate.ExecuteIfBound();
+				Delegates->TickNetworkDelegate.ExecuteIfBound();
 			}
 		}
 	};
@@ -298,7 +339,7 @@ public:
 			if (!GRemoteExecutorNewNetPumpBehavior)
 			{
 				// pump the network
-				UE::RemoteExecutor::TickNetworkDelegate.ExecuteIfBound();
+				Delegates->TickNetworkDelegate.ExecuteIfBound();
 			}
 
 			// check if we're actively servicing a remote multi-server commit
@@ -320,7 +361,7 @@ public:
 				if (GRemoteExecutorNewNetPumpBehavior)
 				{
 					IterationsSinceNetworkPump = 0;
-					UE::RemoteExecutor::TickNetworkDelegate.ExecuteIfBound();
+					Delegates->TickNetworkDelegate.ExecuteIfBound();
 				}
 
 				continue;
@@ -355,11 +396,19 @@ public:
 
 			ExecutePendingWork(PendingMigrationWorks, LastStallPrintTime);
 			ExecutePendingWork(PendingWorks, LastStallPrintTime);
+
+			if (bSingleStepModeForTesting)
+			{
+				break;
+			}
 		}
 
-		PendingWorkLists.Empty();
-		AbortReason = UE::RemoteExecutor::ERemoteExecutorAbortReason::Unspecified;
-		AbortReasonDescription.Reset();
+		if (!HasPendingWork())
+		{
+			PendingWorkLists.Empty();
+			AbortReason = UE::RemoteExecutor::ERemoteExecutorAbortReason::Unspecified;
+			AbortReasonDescription.Reset();
+		}
 	}
 
 private:
@@ -564,7 +613,7 @@ private:
 								// signal to the relevant servers that they need to help us do a multi-server commit
 								UE_AUTORTFM_OPEN
 								{
-									UE::RemoteExecutor::BeginMultiServerCommitDelegate.ExecuteIfBound(ExecutingWork->RequestId, ExecutingWork->Priority, MultiServerCommitRemoteServers);
+									Delegates->BeginMultiServerCommitDelegate.ExecuteIfBound(ExecutingWork->RequestId, ExecutingWork->Priority, MultiServerCommitRemoteServers);
 								};
 
 								// ask each subsystem to send any necessary data as part of this commit
@@ -582,7 +631,7 @@ private:
 								{
 									// tell each server that we're done sending commit data and we are waiting for them
 									// to respond with whether they are ready or we need to abort and retry
-									UE::RemoteExecutor::ReadyMultiServerCommitDelegate.ExecuteIfBound(ExecutingWork->RequestId, MultiServerCommitRemoteServers);
+									Delegates->ReadyMultiServerCommitDelegate.ExecuteIfBound(ExecutingWork->RequestId, MultiServerCommitRemoteServers);
 								};
 
 								// tick the network until everything is ready
@@ -638,7 +687,7 @@ private:
 												}
 
 												// aborted, tell all servers to abandon this commit
-												UE::RemoteExecutor::AbandonMultiServerCommitDelegate.ExecuteIfBound(ExecutingWork->RequestId, MultiServerCommitRemoteServers);
+												Delegates->AbandonMultiServerCommitDelegate.ExecuteIfBound(ExecutingWork->RequestId, MultiServerCommitRemoteServers);
 
 												MultiServerCommitRequestId = FRemoteTransactionId::Invalid();
 												MultiServerCommitReadyServers.Reset();
@@ -668,7 +717,7 @@ private:
 
 									UE_AUTORTFM_OPEN
 									{
-										UE::RemoteExecutor::TickNetworkDelegate.ExecuteIfBound();
+										Delegates->TickNetworkDelegate.ExecuteIfBound();
 									};
 								}
 
@@ -684,7 +733,7 @@ private:
 								UE_AUTORTFM_OPEN
 								{
 									// tell each server that we are committed and they should commit
-									UE::RemoteExecutor::EndMultiServerCommitDelegate.ExecuteIfBound(ExecutingWork->RequestId, MultiServerCommitRemoteServers);
+									Delegates->EndMultiServerCommitDelegate.ExecuteIfBound(ExecutingWork->RequestId, MultiServerCommitRemoteServers);
 								};
 							};
 						});
@@ -704,7 +753,7 @@ private:
 							*AbortReasonDescription);
 
 					bWorkTransactionAborted = true;
-					OnTransactionAbortedDelegate.Broadcast(RequestId, ExecutingWork->ExecutionAttempts, AbortReasonDescription);
+					UE::RemoteExecutor::OnTransactionAbortedDelegate.Broadcast(RequestId, ExecutingWork->ExecutionAttempts, AbortReasonDescription);
 
 					if (AbortReason == ERemoteExecutorAbortReason::AbandonWork) // -V547
 					{
@@ -796,15 +845,23 @@ private:
 	}	
 };
 
-FRemoteExecutor GRemoteExecutor;
+FRemoteExecutorDelegates GDefaultDelegates;
+FRemoteExecutor* GRemoteExecutor = nullptr;
 
 void ExecuteTransactionalInternal(FName WorkName, FRemoteWorkPriority InWorkPriority, const TFunctionRef<void(void)>& InWork)
 {
 #if UE_WITH_REMOTE_OBJECT_HANDLE
 	check(!AutoRTFM::IsClosed());
 
-	GRemoteExecutor.EnqueueWork(WorkName, InWorkPriority, true, INDEX_NONE, [&InWork]() { InWork(); });
-	UE::RemoteExecutor::ExecutePendingWork();
+	if (GRemoteExecutor)
+	{
+		GRemoteExecutor->EnqueueWork(WorkName, InWorkPriority, true, INDEX_NONE, [&InWork]() { InWork(); });
+		UE::RemoteExecutor::ExecutePendingWork();
+	}
+	else
+	{
+		InWork();
+	}
 #else
 	InWork();
 #endif
@@ -812,53 +869,95 @@ void ExecuteTransactionalInternal(FName WorkName, FRemoteWorkPriority InWorkPrio
 
 namespace UE::RemoteExecutor
 {
-int32 SpeculationModeScopeCount = 0;
+FRemoteExecutor* CreateRemoteExecutor(FRemoteServerId LocalServerId, FRemoteExecutorDelegates* Delegates)
+{
+	GRemoteExecutor = new FRemoteExecutor(LocalServerId, Delegates);
+	return GRemoteExecutor;
+}
 
-FSpeculationExecutionScope::FSpeculationExecutionScope()
+void DestroyRemoteExecutor(FRemoteExecutor* Executor)
 {
-	if (GRemoteExecutorAllowSpeculationMode)
+	if (GRemoteExecutor == Executor)
 	{
-		++SpeculationModeScopeCount;
+		GRemoteExecutor = nullptr;
 	}
+
+	delete Executor;
 }
 
-FSpeculationExecutionScope::~FSpeculationExecutionScope()
+void SetRemoteExecutorForTesting(FRemoteExecutor* Executor)
 {
-	if (GRemoteExecutorAllowSpeculationMode)
+	if (GRemoteExecutor)
 	{
-		--SpeculationModeScopeCount;
-		check(SpeculationModeScopeCount >= 0);
+		GRemoteExecutor->DeactivateSubsystemsForTesting();
 	}
-	else
+
+	GRemoteExecutor = Executor;
+	if (GRemoteExecutor)
 	{
-		SpeculationModeScopeCount = 0;
+		GRemoteExecutor->ActivateSubsystemsForTesting();
 	}
+}
 
-#if UE_WITH_REMOTE_OBJECT_HANDLE
-	if (SpeculationModeScopeCount == 0)
+void SetSingleStepModeForTesting(bool bEnable)
+{
+	check(GRemoteExecutor);
+	GRemoteExecutor->bSingleStepModeForTesting = bEnable;
+}
+
+FSpeculationExecutionScope::FSpeculationExecutionScope()
+{
+	if (GRemoteExecutor && GRemoteExecutorAllowSpeculationMode)
 	{
-		if (GRemoteExecutor.ExecutingWork && GRemoteExecutor.ExecutingWork->NumAbortsRequested)
+		++GRemoteExecutor->SpeculationModeScopeCount;
+	}
+}
+
+FSpeculationExecutionScope::~FSpeculationExecutionScope()
+{
+	if (GRemoteExecutor)
+	{
+		if (GRemoteExecutorAllowSpeculationMode)
 		{
-			AutoRTFM::AbortTransaction();
+			--GRemoteExecutor->SpeculationModeScopeCount;
+			check(GRemoteExecutor->SpeculationModeScopeCount >= 0);
+		}
+		else
+		{
+			GRemoteExecutor->SpeculationModeScopeCount = 0;
+		}
+
+#if UE_WITH_REMOTE_OBJECT_HANDLE
+		if (GRemoteExecutor->SpeculationModeScopeCount == 0)
+		{
+			if (GRemoteExecutor->ExecutingWork && GRemoteExecutor->ExecutingWork->NumAbortsRequested)
+			{
+				AutoRTFM::AbortTransaction();
+			}
 		}
-	}
 #endif
+	}
+}
+
+bool IsRemoteExecutorActive()
+{
+	return GRemoteExecutor != nullptr;
 }
 
 bool IsSpeculationMode()
 {
-	check(SpeculationModeScopeCount >= 0);
-	return GRemoteExecutorAllowSpeculationMode && (SpeculationModeScopeCount > 0);
+	check(GRemoteExecutor && GRemoteExecutor->SpeculationModeScopeCount >= 0);
+	return GRemoteExecutor && GRemoteExecutorAllowSpeculationMode && (GRemoteExecutor->SpeculationModeScopeCount > 0);
 }
 
-void RegisterRemoteSubsystem(FRemoteSubsystemBase* Subsystem)
+void RegisterRemoteSubsystem(FRemoteExecutor* Executor, FRemoteSubsystemBase* Subsystem)
 {
-	GRemoteExecutor.RegisterSubsystem(Subsystem);
+	Executor->RegisterSubsystem(Subsystem);
 }
 
 void NotifyNetworkIsWaiting()
 {
-	GRemoteExecutor.bShouldPumpNetwork = true;
+	GRemoteExecutor->bShouldPumpNetwork = true;
 }
 
 void AbortCurrentTransaction(EAbortTransactionBehavior AbortBehavior, ERemoteExecutorAbortReason AbortReason, FStringView AbortDescription)
@@ -867,10 +966,10 @@ void AbortCurrentTransaction(EAbortTransactionBehavior AbortBehavior, ERemoteExe
 
 	UE_AUTORTFM_OPEN
 	{
-		GRemoteExecutor.AbortReason = AbortReason;
-		GRemoteExecutor.AbortReasonDescription = AbortDescription;
+		GRemoteExecutor->AbortReason = AbortReason;
+		GRemoteExecutor->AbortReasonDescription = AbortDescription;
 
-		FRemoteExecutorWork* CurrentWorkItem = GRemoteExecutor.ExecutingWork;
+		FRemoteExecutorWork* CurrentWorkItem = GRemoteExecutor->ExecutingWork;
 		if (ensure(CurrentWorkItem))
 		{
 			CurrentWorkItem->NumAbortsRequested++;
@@ -899,12 +998,17 @@ void AbortTransactionAndAbandonWork(FStringView Description)
 
 void TransactionRequiresMultiServerCommit(FStringView Description)
 {
-	if (GRemoteExecutor.ExecutingWork && !GRemoteExecutor.ExecutingWork->bRequiresMultiServerCommit)
+	if (!GRemoteExecutor)
+	{
+		return;
+	}
+
+	if (GRemoteExecutor->ExecutingWork && !GRemoteExecutor->ExecutingWork->bRequiresMultiServerCommit)
 	{
-		GRemoteExecutor.ExecutingWork->bRequiresMultiServerCommit = true;
-		GRemoteExecutor.ExecutingWork->RequiresMultiServerCommitReason = Description;
+		GRemoteExecutor->ExecutingWork->bRequiresMultiServerCommit = true;
+		GRemoteExecutor->ExecutingWork->RequiresMultiServerCommitReason = Description;
 
-		UE_LOG(LogRemoteExec, VeryVerbose, TEXT("TransactionRequiresMultiServerCommit ACTIVATED because: %s"), *GRemoteExecutor.ExecutingWork->RequiresMultiServerCommitReason);
+		UE_LOG(LogRemoteExec, VeryVerbose, TEXT("TransactionRequiresMultiServerCommit ACTIVATED because: %s"), *GRemoteExecutor->ExecutingWork->RequiresMultiServerCommitReason);
 	}
 }
 
@@ -913,26 +1017,26 @@ void BeginRemoteMultiServerCommit(FRemoteServerId ServerId, FRemoteTransactionId
 	bool bAcceptRequest = false;
 	bool bAbortLocalCommit = false;
 
-	if (GRemoteExecutor.ActiveRemoteMultiServerCommitRequestId.IsValid())
+	if (GRemoteExecutor->ActiveRemoteMultiServerCommitRequestId.IsValid())
 	{
-		if (GRemoteExecutor.bActiveRemoteMultiServerCommitReady)
+		if (GRemoteExecutor->bActiveRemoteMultiServerCommitReady)
 		{
 			// we already told the remote server that we are READY, so we are locked in for a moment until we finish
 			UE_LOG(LogRemoteExec, VeryVerbose, TEXT("BeginRemoteMultiServerCommit %s from %s %s DENYING because we are already READY with remote multi-server commit %s %s"),
 				*RequestId.ToString(),
 				*ServerId.ToString(),
 				*RequestPriority.ToString(),
-				*GRemoteExecutor.ActiveRemoteMultiServerCommitRequestId.ToString(),
-				*GRemoteExecutor.ActiveRemoteMultiServerCommitPriority.ToString());
+				*GRemoteExecutor->ActiveRemoteMultiServerCommitRequestId.ToString(),
+				*GRemoteExecutor->ActiveRemoteMultiServerCommitPriority.ToString());
 		}
-		else if (IsHigherPriority(RequestPriority, GRemoteExecutor.ActiveRemoteMultiServerCommitPriority))
+		else if (IsHigherPriority(RequestPriority, GRemoteExecutor->ActiveRemoteMultiServerCommitPriority))
 		{
 			UE_LOG(LogRemoteExec, VeryVerbose, TEXT("BeginRemoteMultiServerCommit %s from %s %s ACCEPTING because our remote multi-server commit %s is lower priority %s"),
 				*RequestId.ToString(),
 				*ServerId.ToString(),
 				*RequestPriority.ToString(),
-				*GRemoteExecutor.ActiveRemoteMultiServerCommitRequestId.ToString(),
-				*GRemoteExecutor.ActiveRemoteMultiServerCommitPriority.ToString());
+				*GRemoteExecutor->ActiveRemoteMultiServerCommitRequestId.ToString(),
+				*GRemoteExecutor->ActiveRemoteMultiServerCommitPriority.ToString());
 
 			bAcceptRequest = true;
 		}
@@ -942,22 +1046,22 @@ void BeginRemoteMultiServerCommit(FRemoteServerId ServerId, FRemoteTransactionId
 				*RequestId.ToString(),
 				*ServerId.ToString(),
 				*RequestPriority.ToString(),
-				*GRemoteExecutor.ActiveRemoteMultiServerCommitRequestId.ToString(),
-				*GRemoteExecutor.ActiveRemoteMultiServerCommitPriority.ToString());
+				*GRemoteExecutor->ActiveRemoteMultiServerCommitRequestId.ToString(),
+				*GRemoteExecutor->ActiveRemoteMultiServerCommitPriority.ToString());
 		}
 	}
-	else if (GRemoteExecutor.MultiServerCommitRequestId.IsValid())
+	else if (GRemoteExecutor->MultiServerCommitRequestId.IsValid())
 	{
 		// we are in a local multi-server commit - should we abort it in favor of this remote one?
-		check(GRemoteExecutor.ExecutingWork);
-		if (IsHigherPriority(RequestPriority, GRemoteExecutor.ExecutingWork->Priority))
+		check(GRemoteExecutor->ExecutingWork);
+		if (IsHigherPriority(RequestPriority, GRemoteExecutor->ExecutingWork->Priority))
 		{
 			UE_LOG(LogRemoteExec, VeryVerbose, TEXT("BeginRemoteMultiServerCommit %s from %s %s ACCEPTING because our local multi-server commit %s is lower priority %s"),
 				*RequestId.ToString(),
 				*ServerId.ToString(),
 				*RequestPriority.ToString(),
-				*GRemoteExecutor.MultiServerCommitRequestId.ToString(),
-				*GRemoteExecutor.ExecutingWork->Priority.ToString());
+				*GRemoteExecutor->MultiServerCommitRequestId.ToString(),
+				*GRemoteExecutor->ExecutingWork->Priority.ToString());
 
 			bAcceptRequest = true;
 			bAbortLocalCommit = true;
@@ -968,8 +1072,8 @@ void BeginRemoteMultiServerCommit(FRemoteServerId ServerId, FRemoteTransactionId
 				*RequestId.ToString(),
 				*ServerId.ToString(),
 				*RequestPriority.ToString(),
-				*GRemoteExecutor.MultiServerCommitRequestId.ToString(),
-				*GRemoteExecutor.ExecutingWork->Priority.ToString());
+				*GRemoteExecutor->MultiServerCommitRequestId.ToString(),
+				*GRemoteExecutor->ExecutingWork->Priority.ToString());
 		}
 	}
 	else
@@ -987,36 +1091,36 @@ void BeginRemoteMultiServerCommit(FRemoteServerId ServerId, FRemoteTransactionId
 	if (bAcceptRequest)
 	{
 		// do we have to first abandon one we're working on?
-		if (GRemoteExecutor.ActiveRemoteMultiServerCommitRequestId.IsValid())
+		if (GRemoteExecutor->ActiveRemoteMultiServerCommitRequestId.IsValid())
 		{
-			AbortRemoteMultiServerCommitDelegate.Execute(GRemoteExecutor.ActiveRemoteMultiServerCommitRequestId, GRemoteExecutor.ActiveRemoteMultiServerCommitServerId);
+			GRemoteExecutor->Delegates->AbortRemoteMultiServerCommitDelegate.Execute(GRemoteExecutor->ActiveRemoteMultiServerCommitRequestId, GRemoteExecutor->ActiveRemoteMultiServerCommitServerId);
 
-			GRemoteExecutor.ActiveRemoteMultiServerCommitServerId = FRemoteServerId{};
-			GRemoteExecutor.ActiveRemoteMultiServerCommitRequestId = FRemoteTransactionId::Invalid();
-			GRemoteExecutor.ActiveRemoteMultiServerCommitPriority = FRemoteWorkPriority{};
-			GRemoteExecutor.bActiveRemoteMultiServerCommitReady = false;
-			GRemoteExecutor.ActiveRemoteMultiServerCommitDeferredActions.Reset();
+			GRemoteExecutor->ActiveRemoteMultiServerCommitServerId = FRemoteServerId{};
+			GRemoteExecutor->ActiveRemoteMultiServerCommitRequestId = FRemoteTransactionId::Invalid();
+			GRemoteExecutor->ActiveRemoteMultiServerCommitPriority = FRemoteWorkPriority{};
+			GRemoteExecutor->bActiveRemoteMultiServerCommitReady = false;
+			GRemoteExecutor->ActiveRemoteMultiServerCommitDeferredActions.Reset();
 		}
 
-		check(!GRemoteExecutor.ActiveRemoteMultiServerCommitServerId.IsValid());
-		check(!GRemoteExecutor.ActiveRemoteMultiServerCommitRequestId.IsValid());
-		check(!GRemoteExecutor.ActiveRemoteMultiServerCommitPriority.IsValid());
-		check(!GRemoteExecutor.bActiveRemoteMultiServerCommitReady);
-		check(GRemoteExecutor.ActiveRemoteMultiServerCommitDeferredActions.Num() == 0);
+		check(!GRemoteExecutor->ActiveRemoteMultiServerCommitServerId.IsValid());
+		check(!GRemoteExecutor->ActiveRemoteMultiServerCommitRequestId.IsValid());
+		check(!GRemoteExecutor->ActiveRemoteMultiServerCommitPriority.IsValid());
+		check(!GRemoteExecutor->bActiveRemoteMultiServerCommitReady);
+		check(GRemoteExecutor->ActiveRemoteMultiServerCommitDeferredActions.Num() == 0);
 
-		GRemoteExecutor.ActiveRemoteMultiServerCommitServerId = ServerId;
-		GRemoteExecutor.ActiveRemoteMultiServerCommitRequestId = RequestId;
-		GRemoteExecutor.ActiveRemoteMultiServerCommitPriority = RequestPriority;
+		GRemoteExecutor->ActiveRemoteMultiServerCommitServerId = ServerId;
+		GRemoteExecutor->ActiveRemoteMultiServerCommitRequestId = RequestId;
+		GRemoteExecutor->ActiveRemoteMultiServerCommitPriority = RequestPriority;
 	}
 	else
 	{
-		AbortRemoteMultiServerCommitDelegate.Execute(RequestId, ServerId);
+		GRemoteExecutor->Delegates->AbortRemoteMultiServerCommitDelegate.Execute(RequestId, ServerId);
 	}
 
 	if (bAbortLocalCommit)
 	{
 		// we can't immediately abort the transaction, set this flag for it to properly shut down gracefully
-		GRemoteExecutor.bMultiServerCommitRequiresAbort = true;
+		GRemoteExecutor->bMultiServerCommitRequiresAbort = true;
 	}
 }
 
@@ -1024,33 +1128,33 @@ void EndRemoteMultiServerCommit(FRemoteServerId ServerId, FRemoteTransactionId R
 {
 	UE_LOG(LogRemoteExec, VeryVerbose, TEXT("EndRemoteMultiServerCommit : %s %s"), *ServerId.ToString(), *RequestId.ToString());
 
-	check(GRemoteExecutor.ActiveRemoteMultiServerCommitServerId == ServerId);
-	check(GRemoteExecutor.ActiveRemoteMultiServerCommitRequestId == RequestId);
+	check(GRemoteExecutor->ActiveRemoteMultiServerCommitServerId == ServerId);
+	check(GRemoteExecutor->ActiveRemoteMultiServerCommitRequestId == RequestId);
 
-	for (TUniqueFunction<void()>& DeferredAction : GRemoteExecutor.ActiveRemoteMultiServerCommitDeferredActions)
+	for (TUniqueFunction<void()>& DeferredAction : GRemoteExecutor->ActiveRemoteMultiServerCommitDeferredActions)
 	{
 		DeferredAction();
 	}
 
-	GRemoteExecutor.ActiveRemoteMultiServerCommitServerId = FRemoteServerId{};
-	GRemoteExecutor.ActiveRemoteMultiServerCommitRequestId = FRemoteTransactionId::Invalid();
-	GRemoteExecutor.ActiveRemoteMultiServerCommitPriority = FRemoteWorkPriority{};
-	GRemoteExecutor.bActiveRemoteMultiServerCommitReady = false;
-	GRemoteExecutor.ActiveRemoteMultiServerCommitDeferredActions.Reset();
+	GRemoteExecutor->ActiveRemoteMultiServerCommitServerId = FRemoteServerId{};
+	GRemoteExecutor->ActiveRemoteMultiServerCommitRequestId = FRemoteTransactionId::Invalid();
+	GRemoteExecutor->ActiveRemoteMultiServerCommitPriority = FRemoteWorkPriority{};
+	GRemoteExecutor->bActiveRemoteMultiServerCommitReady = false;
+	GRemoteExecutor->ActiveRemoteMultiServerCommitDeferredActions.Reset();
 }
 
 void AbandonRemoteMultiServerCommit(FRemoteServerId ServerId, FRemoteTransactionId RequestId)
 {
-	if ((GRemoteExecutor.ActiveRemoteMultiServerCommitServerId == ServerId) &&
-		GRemoteExecutor.ActiveRemoteMultiServerCommitRequestId == RequestId)
+	if ((GRemoteExecutor->ActiveRemoteMultiServerCommitServerId == ServerId) &&
+		GRemoteExecutor->ActiveRemoteMultiServerCommitRequestId == RequestId)
 	{
 		UE_LOG(LogRemoteExec, VeryVerbose, TEXT("AbandonRemoteMultiServerCommit : %s %s"), *ServerId.ToString(), *RequestId.ToString());
 
-		GRemoteExecutor.ActiveRemoteMultiServerCommitServerId = FRemoteServerId{};
-		GRemoteExecutor.ActiveRemoteMultiServerCommitRequestId = FRemoteTransactionId::Invalid();
-		GRemoteExecutor.ActiveRemoteMultiServerCommitPriority = FRemoteWorkPriority{};
-		GRemoteExecutor.bActiveRemoteMultiServerCommitReady = false;
-		GRemoteExecutor.ActiveRemoteMultiServerCommitDeferredActions.Reset();
+		GRemoteExecutor->ActiveRemoteMultiServerCommitServerId = FRemoteServerId{};
+		GRemoteExecutor->ActiveRemoteMultiServerCommitRequestId = FRemoteTransactionId::Invalid();
+		GRemoteExecutor->ActiveRemoteMultiServerCommitPriority = FRemoteWorkPriority{};
+		GRemoteExecutor->bActiveRemoteMultiServerCommitReady = false;
+		GRemoteExecutor->ActiveRemoteMultiServerCommitDeferredActions.Reset();
 	}
 	else
 	{
@@ -1062,12 +1166,12 @@ void EnqueueRemoteMultiServerCommitAction(FRemoteServerId ServerId, FRemoteTrans
 {
 #if UE_WITH_REMOTE_OBJECT_HANDLE
 	// THIS has to sit in a holding pen...
-	if ((GRemoteExecutor.ActiveRemoteMultiServerCommitServerId == ServerId) &&
-		GRemoteExecutor.ActiveRemoteMultiServerCommitRequestId == RequestId)
+	if ((GRemoteExecutor->ActiveRemoteMultiServerCommitServerId == ServerId) &&
+		GRemoteExecutor->ActiveRemoteMultiServerCommitRequestId == RequestId)
 	{
 		UE_LOG(LogRemoteExec, Verbose, TEXT("EnqueueRemoteMultiServerCommitAction Enqueueing action from %s request %s"), *ServerId.ToString(), *RequestId.ToString());
 
-		GRemoteExecutor.ActiveRemoteMultiServerCommitDeferredActions.Add(MoveTemp(Action));
+		GRemoteExecutor->ActiveRemoteMultiServerCommitDeferredActions.Add(MoveTemp(Action));
 	}
 	else
 	{
@@ -1080,57 +1184,69 @@ void EnqueueRemoteMultiServerCommitAction(FRemoteServerId ServerId, FRemoteTrans
 
 void ReadyMultiServerCommitResponse(FRemoteServerId ServerId, FRemoteTransactionId RequestId)
 {
-	if (GRemoteExecutor.MultiServerCommitRequestId == RequestId)
+	if (GRemoteExecutor->MultiServerCommitRequestId == RequestId)
 	{
 		UE_LOG(LogRemoteExec, VeryVerbose, TEXT("ReadyMultiServerCommitResponse got ready server: %s"), *ServerId.ToString());
-		GRemoteExecutor.MultiServerCommitReadyServers.Add(ServerId);
+		GRemoteExecutor->MultiServerCommitReadyServers.Add(ServerId);
 	}
 	else
 	{
-		UE_LOG(LogRemoteExec, VeryVerbose, TEXT("ReadyMultiServerCommitResponse ignoring %s because we are working on %s"), *RequestId.ToString(), *GRemoteExecutor.MultiServerCommitRequestId.ToString());
+		UE_LOG(LogRemoteExec, VeryVerbose, TEXT("ReadyMultiServerCommitResponse ignoring %s because we are working on %s"), *RequestId.ToString(), *GRemoteExecutor->MultiServerCommitRequestId.ToString());
 	}
 }
 
 void AbortMultiServerCommit(FRemoteServerId ServerId, FRemoteTransactionId RequestId)
 {
-	if (GRemoteExecutor.MultiServerCommitRequestId == RequestId)
+	if (GRemoteExecutor->MultiServerCommitRequestId == RequestId)
 	{
 		UE_LOG(LogRemoteExec, VeryVerbose, TEXT("AbortMultiServerCommit got valid: %s"), *RequestId.ToString());
-		GRemoteExecutor.bMultiServerCommitRequiresAbort = true;
-		GRemoteExecutor.MultiServerCommitReadyServers.Add(ServerId);
+		GRemoteExecutor->bMultiServerCommitRequiresAbort = true;
+		GRemoteExecutor->MultiServerCommitReadyServers.Add(ServerId);
 	}
 	else
 	{
-		UE_LOG(LogRemoteExec, VeryVerbose, TEXT("AbortMultiServerCommit ignoring %s because we are working on %s"), *RequestId.ToString(), *GRemoteExecutor.MultiServerCommitRequestId.ToString());
+		UE_LOG(LogRemoteExec, VeryVerbose, TEXT("AbortMultiServerCommit ignoring %s because we are working on %s"), *RequestId.ToString(), *GRemoteExecutor->MultiServerCommitRequestId.ToString());
 	}
 }
 
 void ReadyRemoteMultiServerCommit(FRemoteServerId ServerId, FRemoteTransactionId RequestId)
 {
-	if ((GRemoteExecutor.ActiveRemoteMultiServerCommitServerId == ServerId) &&
-		(GRemoteExecutor.ActiveRemoteMultiServerCommitRequestId == RequestId))
+	if ((GRemoteExecutor->ActiveRemoteMultiServerCommitServerId == ServerId) &&
+		(GRemoteExecutor->ActiveRemoteMultiServerCommitRequestId == RequestId))
 	{
 		UE_LOG(LogRemoteExec, VeryVerbose, TEXT("ReadyRemoteMultiServerCommit %s from %s"), *RequestId.ToString(), *ServerId.ToString());
 
-		check(!GRemoteExecutor.bActiveRemoteMultiServerCommitReady);
-		GRemoteExecutor.bActiveRemoteMultiServerCommitReady = true;
+		check(!GRemoteExecutor->bActiveRemoteMultiServerCommitReady);
+		GRemoteExecutor->bActiveRemoteMultiServerCommitReady = true;
 
 		// if everything passed OK we need to send a message back saying we're ready
 		// TODO: this is where we'll need to hook code in that verifies everything
 		// we received is able to be accepted
-		ReadyRemoteMultiServerCommitDelegate.Execute(RequestId, ServerId);
+		GRemoteExecutor->Delegates->ReadyRemoteMultiServerCommitDelegate.Execute(RequestId, ServerId);
 	}
 }
 
 FRemoteWorkPriority CreateRootWorkPriority()
 {
-	FRemoteWorkPriority RootWorkPriority = FRemoteWorkPriority::CreateRootWorkPriority(FRemoteServerId::GetLocalServerId(), GRemoteExecutor.GenerateNextTransactionId());
+	FRemoteWorkPriority RootWorkPriority = FRemoteWorkPriority::CreateRootWorkPriority(GRemoteExecutor->GetLocalServerId(), GRemoteExecutor->GenerateNextTransactionId());
 	return RootWorkPriority;
 }
 
 void ExecutePendingWork()
 {
-	GRemoteExecutor.ExecutePendingWork();
+	if (GRemoteExecutor)
+	{
+		GRemoteExecutor->ExecutePendingWork();
+	}
+}
+
+bool HasPendingWork()
+{
+	if (GRemoteExecutor)
+	{
+		return GRemoteExecutor->HasPendingWork();
+	}
+	return false;
 }
 
 void ExecuteTransactionalWithExplicitPriority(FName WorkName, FRemoteWorkPriority WorkPriority, const TFunctionRef<void(void)>& Work)
@@ -1140,14 +1256,28 @@ void ExecuteTransactionalWithExplicitPriority(FName WorkName, FRemoteWorkPriorit
 
 void ExecuteTransactional(FName WorkName, const TFunctionRef<void(void)>& Work)
 {
-	FRemoteWorkPriority RootWorkPriority = FRemoteWorkPriority::CreateRootWorkPriority(FRemoteServerId::GetLocalServerId(), GRemoteExecutor.GenerateNextTransactionId());
-	ExecuteTransactionalWithExplicitPriority(WorkName, RootWorkPriority, Work);
+	if (GRemoteExecutor)
+	{
+		FRemoteWorkPriority RootWorkPriority = FRemoteWorkPriority::CreateRootWorkPriority(GRemoteExecutor->GetLocalServerId(), GRemoteExecutor->GenerateNextTransactionId());
+		ExecuteTransactionalWithExplicitPriority(WorkName, RootWorkPriority, Work);
+	}
+	else
+	{
+		Work();
+	}
 }
 
 void EnqueueWorkWithExplicitPriority(FName WorkName, FRemoteWorkPriority WorkPriority, bool bIsTransactional, TUniqueFunction<void(void)>&& InWork)
 {
 #if UE_WITH_REMOTE_OBJECT_HANDLE && UE_AUTORTFM
-	GRemoteExecutor.EnqueueWork(WorkName, WorkPriority, bIsTransactional, INDEX_NONE, MoveTemp(InWork));
+	if (GRemoteExecutor)
+	{
+		GRemoteExecutor->EnqueueWork(WorkName, WorkPriority, bIsTransactional, INDEX_NONE, MoveTemp(InWork));
+	}
+	else
+	{
+		InWork();
+	}
 #else
 	InWork();
 #endif
@@ -1155,34 +1285,70 @@ void EnqueueWorkWithExplicitPriority(FName WorkName, FRemoteWorkPriority WorkPri
 
 void EnqueueWork(FName WorkName, bool bIsTransactional, TUniqueFunction<void(void)>&& InWork)
 {
-	FRemoteWorkPriority RootWorkPriority = FRemoteWorkPriority::CreateRootWorkPriority(FRemoteServerId::GetLocalServerId(), GRemoteExecutor.GenerateNextTransactionId());
-	EnqueueWorkWithExplicitPriority(WorkName, RootWorkPriority, bIsTransactional, MoveTemp(InWork));
+	if (GRemoteExecutor)
+	{
+		FRemoteWorkPriority RootWorkPriority = FRemoteWorkPriority::CreateRootWorkPriority(GRemoteExecutor->GetLocalServerId(), GRemoteExecutor->GenerateNextTransactionId());
+		EnqueueWorkWithExplicitPriority(WorkName, RootWorkPriority, bIsTransactional, MoveTemp(InWork));
+	}
+	else
+	{
+		InWork();
+	}
 }
 
 void EnqueueMigrationWork(FName WorkName, TUniqueFunction<void(void)>&& InWork)
 {
-	FRemoteWorkPriority RootWorkPriority = FRemoteWorkPriority::CreateRootWorkPriority(FRemoteServerId::GetLocalServerId(), GRemoteExecutor.GenerateNextTransactionId(), ERemoteWorkType::Migration);
-	EnqueueWorkWithExplicitPriority(WorkName, RootWorkPriority, /*bIsTransactional=*/ true, MoveTemp(InWork));
+	if (GRemoteExecutor)
+	{
+		FRemoteWorkPriority RootWorkPriority = FRemoteWorkPriority::CreateRootWorkPriority(GRemoteExecutor->GetLocalServerId(), GRemoteExecutor->GenerateNextTransactionId(), ERemoteWorkType::Migration);
+		EnqueueWorkWithExplicitPriority(WorkName, RootWorkPriority, /*bIsTransactional=*/ true, MoveTemp(InWork));
+	}
+	else
+	{
+		InWork();
+	}
 }
 
 
 void EnqueueWorkListWithExplicitPriority(FName WorkName, FRemoteWorkPriority WorkPriority, bool bIsTransactional, TArray<TUniqueFunction<void(void)>>&& InWorkList)
 {
 #if UE_WITH_REMOTE_OBJECT_HANDLE
-	GRemoteExecutor.EnqueueWorkList(WorkName, WorkPriority, bIsTransactional, MoveTemp(InWorkList));
-#else
+	if (GRemoteExecutor)
+	{
+		GRemoteExecutor->EnqueueWorkList(WorkName, WorkPriority, bIsTransactional, MoveTemp(InWorkList));
+	}
+	else
+	{
+#endif
 	for(int32 WorkIndex = 0; WorkIndex < InWorkList.Num(); WorkIndex++)
 	{
 		const TUniqueFunction<void(void)>& Work = InWorkList[WorkIndex];
 		Work();
 	}
+#if UE_WITH_REMOTE_OBJECT_HANDLE
+	}
 #endif
 }
 
 void EnqueueWorkList(FName WorkName, bool bIsTransactional, TArray<TUniqueFunction<void(void)>>&& InWorkList)
 {
-	FRemoteWorkPriority RootWorkPriority = FRemoteWorkPriority::CreateRootWorkPriority(FRemoteServerId::GetLocalServerId(), GRemoteExecutor.GenerateNextTransactionId());
-	EnqueueWorkListWithExplicitPriority(WorkName, RootWorkPriority, bIsTransactional, MoveTemp(InWorkList));
+#if UE_WITH_REMOTE_OBJECT_HANDLE
+	if (GRemoteExecutor)
+	{
+		FRemoteWorkPriority RootWorkPriority = FRemoteWorkPriority::CreateRootWorkPriority(GRemoteExecutor->GetLocalServerId(), GRemoteExecutor->GenerateNextTransactionId());
+		EnqueueWorkListWithExplicitPriority(WorkName, RootWorkPriority, bIsTransactional, MoveTemp(InWorkList));
+	}
+	else
+	{
+#endif
+	for(int32 WorkIndex = 0; WorkIndex < InWorkList.Num(); WorkIndex++)
+	{
+		const TUniqueFunction<void(void)>& Work = InWorkList[WorkIndex];
+		Work();
+	}
+#if UE_WITH_REMOTE_OBJECT_HANDLE
+	}
+#endif
 }
 
 } // namespace UE::RemoteExecutor
diff --git a/Engine/Source/Runtime/CoreUObject/Private/UObject/RemoteHeap.cpp b/Engine/Source/Runtime/CoreUObject/Private/UObject/RemoteHeap.cpp
--- a/Engine/Source/Runtime/CoreUObject/Private/UObject/RemoteHeap.cpp
+++ b/Engine/Source/Runtime/CoreUObject/Private/UObject/RemoteHeap.cpp
@@ -28,8 +28,6 @@ const uint64 RemoteHeapLineMask = ~(RemoteHeapLineSize-1ull);
 	DEFINE_REMOTE_HEAP_DELEGATE(UE::RemoteHeapDeallocateDelegate);
 #undef DEFINE_REMOTE_HEAP_DELEGATE
 
-UE::FRemoteHeap GRemoteHeap;
-
 static uint64 RemoteHeapGetBaseAddressForServerId(FRemoteServerId ServerId)
 {
 	uint64 ServerIdNum = ServerId.GetIdNumber();
@@ -52,95 +50,99 @@ void RemoteHeapRedirectedLoad(void* DestPointer, uint64_t Size, uint64_t SourceA
 {
 	switch(Size)
 	{
-	case 8: *(uint64*)DestPointer = GRemoteHeap.RemoteHeapRead8(SourceAddress); break;
-	case 4: *(uint32*)DestPointer = GRemoteHeap.RemoteHeapRead4(SourceAddress); break;
-	case 2: *(uint16*)DestPointer = GRemoteHeap.RemoteHeapRead2(SourceAddress); break;
-	case 1: *(uint8*)DestPointer = GRemoteHeap.RemoteHeapRead1(SourceAddress); break;
-	default: GRemoteHeap.RemoteHeapRead(DestPointer, Size, SourceAddress); break;
+	case 8: *(uint64*)DestPointer = FRemoteHeap::Instance->RemoteHeapRead8(SourceAddress); break;
+	case 4: *(uint32*)DestPointer = FRemoteHeap::Instance->RemoteHeapRead4(SourceAddress); break;
+	case 2: *(uint16*)DestPointer = FRemoteHeap::Instance->RemoteHeapRead2(SourceAddress); break;
+	case 1: *(uint8*)DestPointer = FRemoteHeap::Instance->RemoteHeapRead1(SourceAddress); break;
+	default: FRemoteHeap::Instance->RemoteHeapRead(DestPointer, Size, SourceAddress); break;
 	}
 }
 
 void RemoteHeapRedirectedLoad8(void* DestPointer, uint64_t SourceAddress)
 {
-	*(uint64*)DestPointer = GRemoteHeap.RemoteHeapRead8(SourceAddress);
+	*(uint64*)DestPointer = FRemoteHeap::Instance->RemoteHeapRead8(SourceAddress);
 }
 
 void RemoteHeapRedirectedLoad4(void* DestPointer, uint64_t SourceAddress)
 {
-	*(uint32*)DestPointer = GRemoteHeap.RemoteHeapRead4(SourceAddress);
+	*(uint32*)DestPointer = FRemoteHeap::Instance->RemoteHeapRead4(SourceAddress);
 }
 
 void RemoteHeapRedirectedLoad2(void* DestPointer, uint64_t SourceAddress)
 {
-	*(uint16*)DestPointer = GRemoteHeap.RemoteHeapRead2(SourceAddress);
+	*(uint16*)DestPointer = FRemoteHeap::Instance->RemoteHeapRead2(SourceAddress);
 }
 
 void RemoteHeapRedirectedLoad1(void* DestPointer, uint64_t SourceAddress)
 {
-	*(uint8*)DestPointer = GRemoteHeap.RemoteHeapRead1(SourceAddress);
+	*(uint8*)DestPointer = FRemoteHeap::Instance->RemoteHeapRead1(SourceAddress);
 }
 
 void RemoteHeapRedirectedStore(uint64_t DestAddress, uint64_t Size, const void* SourcePointer)
 {
 	switch(Size)
 	{
-	case 8: GRemoteHeap.RemoteHeapWrite8(DestAddress, *(uint64*)SourcePointer); break;
-	case 4: GRemoteHeap.RemoteHeapWrite4(DestAddress, *(uint32*)SourcePointer); break;
-	case 2: GRemoteHeap.RemoteHeapWrite2(DestAddress, *(uint16*)SourcePointer); break;
-	case 1: GRemoteHeap.RemoteHeapWrite1(DestAddress, *(uint8*)SourcePointer); break;
-	default: GRemoteHeap.RemoteHeapWrite(DestAddress, Size, SourcePointer); break;
+	case 8: FRemoteHeap::Instance->RemoteHeapWrite8(DestAddress, *(uint64*)SourcePointer); break;
+	case 4: FRemoteHeap::Instance->RemoteHeapWrite4(DestAddress, *(uint32*)SourcePointer); break;
+	case 2: FRemoteHeap::Instance->RemoteHeapWrite2(DestAddress, *(uint16*)SourcePointer); break;
+	case 1: FRemoteHeap::Instance->RemoteHeapWrite1(DestAddress, *(uint8*)SourcePointer); break;
+	default: FRemoteHeap::Instance->RemoteHeapWrite(DestAddress, Size, SourcePointer); break;
 	}
 }
 
 void RemoteHeapRedirectedStore8(uint64_t DestAddress, const void* SourcePointer)
 {
-	GRemoteHeap.RemoteHeapWrite8(DestAddress, *(uint64*)SourcePointer);
+	FRemoteHeap::Instance->RemoteHeapWrite8(DestAddress, *(uint64*)SourcePointer);
 }
 
 void RemoteHeapRedirectedStore4(uint64_t DestAddress, const void* SourcePointer)
 {
-	GRemoteHeap.RemoteHeapWrite4(DestAddress, *(uint32*)SourcePointer);
+	FRemoteHeap::Instance->RemoteHeapWrite4(DestAddress, *(uint32*)SourcePointer);
 }
 
 void RemoteHeapRedirectedStore2(uint64_t DestAddress, const void* SourcePointer)
 {
-	GRemoteHeap.RemoteHeapWrite2(DestAddress, *(uint16*)SourcePointer);
+	FRemoteHeap::Instance->RemoteHeapWrite2(DestAddress, *(uint16*)SourcePointer);
 }
 
 void RemoteHeapRedirectedStore1(uint64_t DestAddress, const void* SourcePointer)
 {
-	GRemoteHeap.RemoteHeapWrite1(DestAddress, *(uint8*)SourcePointer);
+	FRemoteHeap::Instance->RemoteHeapWrite1(DestAddress, *(uint8*)SourcePointer);
 }
 
-FRemoteHeap& FRemoteHeap::GetInstance()
+void FRemoteHeap::CreateInstance(FRemoteExecutor* Executor, FRemoteServerId LocalServerId)
 {
-	return GRemoteHeap;
-}
+	FRemoteHeap* Heap = new FRemoteHeap(LocalServerId);
+	Instance = Heap;
 
-void FRemoteHeap::Initialize()
-{
 #if UE_WITH_REMOTE_OBJECT_HANDLE
-	RemoteExecutor::RegisterRemoteSubsystem(this);
+	RemoteExecutor::RegisterRemoteSubsystem(Executor, Heap);
 
-	uint64 LocalServerHeapBaseAddress = RemoteHeapGetBaseAddressForServerId(FRemoteServerId::GetLocalServerId());
+	uint64 LocalServerHeapBaseAddress = RemoteHeapGetBaseAddressForServerId(LocalServerId);
 
 	pas_range heap_range = pas_range_create(LocalServerHeapBaseAddress, LocalServerHeapBaseAddress + RemoteHeapSizePerServer);
 	jit_heap_add_fresh_memory(heap_range);
 
-	AutoRTFM::FHeapRedirectCallbacks Callbacks;
-	// the address space is encoded as the bottom nibble of the top byte
-	Callbacks.AddressSpace = (RemoteHeapBaseAddress >> 56) & 0x0F;
-	Callbacks.RedirectedLoad = &RemoteHeapRedirectedLoad;
-	Callbacks.RedirectedLoad8 = &RemoteHeapRedirectedLoad8;
-	Callbacks.RedirectedLoad4 = &RemoteHeapRedirectedLoad4;
-	Callbacks.RedirectedLoad2 = &RemoteHeapRedirectedLoad2;
-	Callbacks.RedirectedLoad1 = &RemoteHeapRedirectedLoad1;
-	Callbacks.RedirectedStore = &RemoteHeapRedirectedStore;
-	Callbacks.RedirectedStore8 = &RemoteHeapRedirectedStore8;
-	Callbacks.RedirectedStore4 = &RemoteHeapRedirectedStore4;
-	Callbacks.RedirectedStore2 = &RemoteHeapRedirectedStore2;
-	Callbacks.RedirectedStore1 = &RemoteHeapRedirectedStore1;
-	AutoRTFM::RegisterHeapRedirectCallbacks(Callbacks);
+	// Register the heap redirect callbacks once (they route through FRemoteHeap::Instance)
+	static bool bCallbacksRegistered = false;
+	if (!bCallbacksRegistered)
+	{
+		AutoRTFM::FHeapRedirectCallbacks Callbacks;
+		// the address space is encoded as the bottom nibble of the top byte
+		Callbacks.AddressSpace = (RemoteHeapBaseAddress >> 56) & 0x0F;
+		Callbacks.RedirectedLoad = &RemoteHeapRedirectedLoad;
+		Callbacks.RedirectedLoad8 = &RemoteHeapRedirectedLoad8;
+		Callbacks.RedirectedLoad4 = &RemoteHeapRedirectedLoad4;
+		Callbacks.RedirectedLoad2 = &RemoteHeapRedirectedLoad2;
+		Callbacks.RedirectedLoad1 = &RemoteHeapRedirectedLoad1;
+		Callbacks.RedirectedStore = &RemoteHeapRedirectedStore;
+		Callbacks.RedirectedStore8 = &RemoteHeapRedirectedStore8;
+		Callbacks.RedirectedStore4 = &RemoteHeapRedirectedStore4;
+		Callbacks.RedirectedStore2 = &RemoteHeapRedirectedStore2;
+		Callbacks.RedirectedStore1 = &RemoteHeapRedirectedStore1;
+		AutoRTFM::RegisterHeapRedirectCallbacks(Callbacks);
+		bCallbacksRegistered = true;
+	}
 
 	UE_LOG(LogRemoteHeap, Display, TEXT("FRemoteHeap Initialized: BaseAddress %llx"), LocalServerHeapBaseAddress);
 #endif
@@ -359,7 +361,7 @@ void FRemoteHeap::TickRequest()
 					{
 						// send a request and bump the priority
 						PendingRequest->RequestPriority = HighestPriorityRequest->Priority;
-						RemoteHeapRequestLineDelegate.Execute(Address, ExpectedOwnerServerId, PendingRequest->RequestPriority, FRemoteServerId::GetLocalServerId());
+						RemoteHeapRequestLineDelegate.Execute(Address, ExpectedOwnerServerId, PendingRequest->RequestPriority, LocalServerId);
 					}
 				}
 			}
@@ -370,7 +372,7 @@ void FRemoteHeap::TickRequest()
 				PendingRequest->Address = Address;
 				PendingRequest->RequestPriority = HighestPriorityRequest->Priority;
 
-				RemoteHeapRequestLineDelegate.Execute(Address, ExpectedOwnerServerId, PendingRequest->RequestPriority, FRemoteServerId::GetLocalServerId());
+				RemoteHeapRequestLineDelegate.Execute(Address, ExpectedOwnerServerId, PendingRequest->RequestPriority, LocalServerId);
 			}
 		}
 	}
@@ -833,7 +835,7 @@ void* FRemoteHeap::RemoteHeapAllocate(uint64 Size)
 
 	while(LineAddressCursor < LineAddressEnd)
 	{
-		SRemoteHeapLineEntry& NewHeapLine = HeapLines.Emplace_GetRef();
+		SRemoteHeapLineEntry& NewHeapLine = FRemoteHeap::Instance->HeapLines.Emplace_GetRef();
 		NewHeapLine.BaseAddress = LineAddressCursor;
 		NewHeapLine.State = ERemoteHeapLineState::Exclusive;
 		NewHeapLine.Data.SetNum(RemoteHeapLineSize);
@@ -860,9 +862,9 @@ void FRemoteHeap::RemoteHeapDeallocate(void* Ptr)
 
 	FRemoteServerId OwnerServer = RemoteHeapGetOwnerServerId(Address);
 
-	UE_AUTORTFM_ONCOMMIT(Ptr, Address, OwnerServer)
+	UE_AUTORTFM_ONCOMMIT(this, Ptr, Address, OwnerServer)
 	{
-		if (OwnerServer == FRemoteServerId::GetLocalServerId())
+		if (OwnerServer == LocalServerId)
 		{
 			UE_LOG(LogRemoteHeap, Verbose, TEXT("RemoteHeapDeallocate (local): %llx"), Address);
 			jit_heap_deallocate((void*)Ptr);
@@ -879,43 +881,43 @@ void FRemoteHeap::RemoteHeapDeallocate(void* Ptr)
 
 void RemoteHeapOnLineReceived(uint64 Address, const TArray<uint8>& Data)
 {
-	GRemoteHeap.OnLineReceived(Address, Data);
+	FRemoteHeap::Instance->OnLineReceived(Address, Data);
 }
 
 void RemoteHeapOnLineRequested(uint64 Address, FRemoteWorkPriority Priority, FRemoteServerId Destination)
 {
-	GRemoteHeap.OnLineRequested(Address, Priority, Destination);
+	FRemoteHeap::Instance->OnLineRequested(Address, Priority, Destination);
 }
 
 void RemoteHeapDeallocate(uint64 Address)
 {
-	GRemoteHeap.RemoteHeapDeallocate((void*)Address);
+	FRemoteHeap::Instance->RemoteHeapDeallocate((void*)Address);
 }
 
 void FRemoteHeapAllocator::ForAnyElementType::ResizeAllocation(FRemoteHeapAllocator::SizeType CurrentNum, FRemoteHeapAllocator::SizeType NewMax, FRemoteHeapAllocator::SizeType NumBytesPerElement)
 {
 	if (NewMax == 0)
 	{
-		GRemoteHeap.RemoteHeapDeallocate(Allocation);
+		FRemoteHeap::Instance->RemoteHeapDeallocate(Allocation);
 		Allocation = nullptr;
 	}
 	else
 	{
-		void* NewMemory = GRemoteHeap.RemoteHeapAllocate(NewMax * NumBytesPerElement);
-		
+		void* NewMemory = FRemoteHeap::Instance->RemoteHeapAllocate(NewMax * NumBytesPerElement);
+
 		if (CurrentNum != 0)
 		{
 			FMemory::Memcpy(NewMemory, Allocation, CurrentNum * NumBytesPerElement); //-V575
 		}
 
-		GRemoteHeap.RemoteHeapDeallocate(Allocation);
+		FRemoteHeap::Instance->RemoteHeapDeallocate(Allocation);
 		Allocation = NewMemory;
 	}
 }
 
 FRemoteHeapAllocator::ForAnyElementType::~ForAnyElementType()
 {
-	GRemoteHeap.RemoteHeapDeallocate(Allocation);
+	FRemoteHeap::Instance->RemoteHeapDeallocate(Allocation);
 	Allocation = nullptr;
 }
 
diff --git a/Engine/Source/Runtime/CoreUObject/Private/UObject/RemoteObject.cpp b/Engine/Source/Runtime/CoreUObject/Private/UObject/RemoteObject.cpp
--- a/Engine/Source/Runtime/CoreUObject/Private/UObject/RemoteObject.cpp
+++ b/Engine/Source/Runtime/CoreUObject/Private/UObject/RemoteObject.cpp
@@ -282,8 +282,6 @@ void InitRemoteObjects()
 		InitServerId();
 	}
 
-	UE::RemoteObject::Transfer::InitRemoteObjectTransfer();
-
 	if (!UE::RemoteObject::Transfer::RemoteObjectTransferDelegate.IsBound())
 	{
 		UE::RemoteObject::Transfer::RemoteObjectTransferDelegate.BindStatic(&UE::RemoteObject::Serialization::Disk::SaveObjectToDisk);
@@ -311,11 +309,6 @@ void InitRemoteObjects()
 		UE::RemoteObject::Transfer::RestoreRemoteObjectDataDelegate.BindStatic(&UE::RemoteObject::Serialization::Disk::LoadObjectFromDisk);
 	}
 
-	if (!UE::RemoteExecutor::FetchNextDeferredRPCDelegate.IsBound())
-	{
-		UE::RemoteExecutor::FetchNextDeferredRPCDelegate.BindStatic([]() { return TOptional<TTuple<FName, FRemoteWorkPriority, bool, TUniqueFunction<void(void)>>>(); });
-	}
-
 #if UE_WITH_REMOTE_OBJECT_HANDLE
 	// We rely on garbage elimination being disabled because we don't allow it inside of borrowed objects
 	if (UObjectBaseUtility::IsGarbageEliminationEnabled())
diff --git a/Engine/Source/Runtime/CoreUObject/Private/UObject/RemoteObjectTransfer.cpp b/Engine/Source/Runtime/CoreUObject/Private/UObject/RemoteObjectTransfer.cpp
--- a/Engine/Source/Runtime/CoreUObject/Private/UObject/RemoteObjectTransfer.cpp
+++ b/Engine/Source/Runtime/CoreUObject/Private/UObject/RemoteObjectTransfer.cpp
@@ -128,7 +128,7 @@ public:
 	FRemoteWorkPriority RequestPriority;
 };
 
-class FRemoteObjectTransferQueue : public FRemoteSubsystem<FRemoteObjectRequest>
+class FRemoteObjectTransferQueue : public FRemoteSubsystem<FRemoteObjectTransferQueue, FRemoteObjectRequest>
 {
 	// these are objects that we have outstanding requests for (across all requests)
 	TArray<FPendingObjectRequest> PendingObjectRequests;
@@ -954,8 +954,6 @@ public:
 	}
 };
 
-FRemoteObjectTransferQueue GTransferQueue;
-
 TRefCountPtr<FTransactionalMigrationData> StoreMigratedObjectData(FRemoteObjectData& ObjectData, FUObjectMigrationContext& MigrationContext)
 {
 	using namespace UE::RemoteObject::Serialization;
@@ -1003,9 +1001,10 @@ TRefCountPtr<FTransactionalMigrationData> StoreMigratedObjectData(FRemoteObjectD
 
 	checkf(RootStub->MigratedData == MigratedData, TEXT("Data migrated for %s was not correctly stored in its stub"), *RootStub->Id.ToString());
 
-	for (int32 RequestIndex = 0; RequestIndex < GTransferQueue.GetRequestCount(); RequestIndex++)
+	FRemoteObjectTransferQueue* TransferQueue = FRemoteObjectTransferQueue::Instance;
+	for (int32 RequestIndex = 0; RequestIndex < TransferQueue->GetRequestCount(); RequestIndex++)
 	{
-		FRemoteObjectRequest* ExistingRequest = GTransferQueue.GetRequestByIndex(RequestIndex);
+		FRemoteObjectRequest* ExistingRequest = TransferQueue->GetRequestByIndex(RequestIndex);
 
 		for (FRemoteObjectId SerializedId : MigratedData->Data)
 		{
@@ -1180,9 +1179,10 @@ UObject* PostMigrateObjectsFromMigratedData(FTransactionalMigrationData& InOutMi
 	// if so, we can now update the canonical root id
 	FRemoteObjectId RootObjectRemoteId = UE::CoreUObject::Private::FObjectHandleUtils::GetRemoteId(ReceivedObjects[0]);
 
-	for (int32 RequestIndex = 0; RequestIndex < GTransferQueue.GetRequestCount(); RequestIndex++)
+	FRemoteObjectTransferQueue* TransferQueue = FRemoteObjectTransferQueue::Instance;
+	for (int32 RequestIndex = 0; RequestIndex < TransferQueue->GetRequestCount(); RequestIndex++)
 	{
-		FRemoteObjectRequest* ExistingRequest = GTransferQueue.GetRequestByIndex(RequestIndex);
+		FRemoteObjectRequest* ExistingRequest = TransferQueue->GetRequestByIndex(RequestIndex);
 
 		for (FRemoteObjectId SubObjectId : ReceivedObjectRemoteIds)
 		{
@@ -1302,6 +1302,7 @@ void MigrateObjectFromRemoteServer(FRemoteObjectId ObjectId, FRemoteServerId Cur
 		return;
 	}
 
+	FRemoteObjectTransferQueue* TransferQueue = FRemoteObjectTransferQueue::Instance;
 #if !UE_AUTORTFM
 	{
 		FRemoteServerId LastKnownResidentServerId;
@@ -1309,34 +1310,34 @@ void MigrateObjectFromRemoteServer(FRemoteObjectId ObjectId, FRemoteServerId Cur
 		{
 			LastKnownResidentServerId = Stub->ResidentServerId;
 		}
-		GTransferQueue.SendRemoteObjectRequest(FRemoteWorkPriority(), ObjectId, LastKnownResidentServerId, FRemoteServerId::GetLocalServerId());
+		TransferQueue->SendRemoteObjectRequest(FRemoteWorkPriority(), ObjectId, LastKnownResidentServerId, FRemoteServerId::GetLocalServerId());
 	}
 #else
-	checkf(GTransferQueue.ActiveRequest, TEXT("Attempting to access remote object %s but we are outside of a transaction"), *ObjectId.ToString());
+	checkf(TransferQueue->ActiveRequest, TEXT("Attempting to access remote object %s but we are outside of a transaction"), *ObjectId.ToString());
 
 	// ensure the request is added in the Open so after we abort it is preserved
 	UE_AUTORTFM_OPEN
 	{
 		// check bInMultiServerCommit here to ensure the UsedObjects list doesn't 
 		// accidentally mutate during the actual multi-server commit process
-		if (!GTransferQueue.bInMultiServerCommit && !GTransferQueue.ActiveRequest->UsedObjects.Contains(ObjectId))
+		if (!TransferQueue->bInMultiServerCommit && !TransferQueue->ActiveRequest->UsedObjects.Contains(ObjectId))
 		{
-			GTransferQueue.ActiveRequest->UsedObjects.Add(ObjectId);
+			TransferQueue->ActiveRequest->UsedObjects.Add(ObjectId);
 
-			OnObjectTouchedDelegate.Broadcast(GTransferQueue.ActiveRequest->RequestId, ObjectId);
+			OnObjectTouchedDelegate.Broadcast(TransferQueue->ActiveRequest->RequestId, ObjectId);
 		}
 
 		// add it to the active request's list of required objects
 		bool bAlreadyAdded = false;
-		GTransferQueue.ActiveRequest->RequiredObjects.Add(ObjectId, &bAlreadyAdded);
+		TransferQueue->ActiveRequest->RequiredObjects.Add(ObjectId, &bAlreadyAdded);
 
 		if (!bAlreadyAdded)
 		{
 			// does any other existing request also require this object? if so, try to
 			// grab the cached canonical id from them
-			for (int32 RequestIndex = 0; RequestIndex < GTransferQueue.GetRequestCount(); RequestIndex++)
+			for (int32 RequestIndex = 0; RequestIndex < TransferQueue->GetRequestCount(); RequestIndex++)
 			{
-				FRemoteObjectRequest* ExistingRequest = GTransferQueue.GetRequestByIndex(RequestIndex);
+				FRemoteObjectRequest* ExistingRequest = TransferQueue->GetRequestByIndex(RequestIndex);
 
 				if (FRemoteObjectId* FoundRequiredRemoteObjectID = ExistingRequest->RequiredObjects.Find(ObjectId))
 				{
@@ -1344,8 +1345,8 @@ void MigrateObjectFromRemoteServer(FRemoteObjectId ObjectId, FRemoteServerId Cur
 
 					if (CanonicalRootObjectId)
 					{
-						GTransferQueue.ActiveRequest->RequiredObjectsCanonicalRootByObjectID.Add(ObjectId, *CanonicalRootObjectId);
-						GTransferQueue.ActiveRequest->RequiredObjectsCanonicalRoots.Add(*CanonicalRootObjectId);
+						TransferQueue->ActiveRequest->RequiredObjectsCanonicalRootByObjectID.Add(ObjectId, *CanonicalRootObjectId);
+						TransferQueue->ActiveRequest->RequiredObjectsCanonicalRoots.Add(*CanonicalRootObjectId);
 						break;
 					}
 				}
@@ -1380,29 +1381,35 @@ void TouchResidentObject(UObject* Object)
 		return;
 	}
 
+	FRemoteObjectTransferQueue* TransferQueue = FRemoteObjectTransferQueue::Instance;
+
+	if (!TransferQueue)
+	{
+		return;
+	}
+
 	// check bInMultiServerCommit here to ensure the UsedObjects list doesn't 
 	// accidentally mutate during the actual multi-server commit process
-	if (GTransferQueue.bInMultiServerCommit)
+	if (TransferQueue->bInMultiServerCommit)
 	{
 		return;
 	}
 
-	if (AutoRTFM::IsClosed() && GTransferQueue.ActiveRequest)
+	if (AutoRTFM::IsClosed() && TransferQueue->ActiveRequest)
 	{
 		UE_AUTORTFM_OPEN
 		{
 			const FRemoteObjectId ObjectId(Object);
 
 			bool bAlreadyInUsedObjectsList = false;
-			GTransferQueue.ActiveRequest->UsedObjects.Add(ObjectId, &bAlreadyInUsedObjectsList);
+			TransferQueue->ActiveRequest->UsedObjects.Add(ObjectId, &bAlreadyInUsedObjectsList);
 
 			if (!bAlreadyInUsedObjectsList)
 			{
-				OnObjectTouchedDelegate.Broadcast(GTransferQueue.ActiveRequest->RequestId, ObjectId);
+				OnObjectTouchedDelegate.Broadcast(TransferQueue->ActiveRequest->RequestId, ObjectId);
 			}
 
-
-			GTransferQueue.ActiveRequest->NewRequiredObject = ObjectId;
+			TransferQueue->ActiveRequest->NewRequiredObject = ObjectId;
 		};
 	}
 }
@@ -1461,7 +1468,7 @@ void TransferObjectOwnershipToRemoteServer(UObject* Object, FRemoteServerId Dest
 	}
 
 	FMigrateSendParams SendRemoteObjectParams = { .MigrationContext = MigrationContext, .ObjectData = ObjectData };
-	GTransferQueue.SendRemoteObject(SendRemoteObjectParams);
+	FRemoteObjectTransferQueue::Instance->SendRemoteObject(SendRemoteObjectParams);
 
 	checkf(SentObjects.Num() > 0, TEXT("PeerId:%s unable to migrate object %s to PeerId:%s"), *FRemoteServerId::GetLocalServerId().ToString(), *UE::CoreUObject::Private::FObjectHandleUtils::GetRemoteId(Object).ToString(), *DestinationServerId.ToString());
 
@@ -1503,9 +1510,10 @@ void RegisterRemoteObjectId(FRemoteObjectId Id, FRemoteServerId ResidentServerId
 
 void GetAllBorrowedObjects(TArray<FRemoteObjectId>& OutBorrowedObjectIds)
 {
-	check(GTransferQueue.ActiveRequest);
+	FRemoteObjectTransferQueue* TransferQueue = FRemoteObjectTransferQueue::Instance;
+	check(TransferQueue->ActiveRequest);
 
-	for (const FRemoteObjectId& Id : GTransferQueue.ActiveRequest->UsedObjects)
+	for (const FRemoteObjectId& Id : TransferQueue->ActiveRequest->UsedObjects)
 	{
 		if (!UE::RemoteObject::Handle::IsOwned(Id))
 		{
@@ -1522,7 +1530,7 @@ void MigrateObjectToRemoteServer(FRemoteObjectId ObjectId, FRemoteServerId Desti
 
 void MigrateObjectToRemoteServerWithExplicitPriority(FRemoteWorkPriority RequestPriority, FRemoteObjectId ObjectId, FRemoteServerId DestinationServerId)
 {
-	GTransferQueue.RequestObjectMigration(RequestPriority, ObjectId, DestinationServerId);
+	FRemoteObjectTransferQueue::Instance->RequestObjectMigration(RequestPriority, ObjectId, DestinationServerId);
 }
 
 int32 ReturnObjectDataToOwnedServers()
@@ -1566,18 +1574,20 @@ void OnObjectDataReceived(FRemoteServerId OwnerServerId, FRemoteObjectId ObjectI
 	MigrationContext.CacheObjectStatus(Data);
 	FScopedObjectMigrationContext ScopedMigrationContext(MigrationContext);
 
-	GTransferQueue.FulfillReceiveRequest(Data, MigrationContext);
+	FRemoteObjectTransferQueue::Instance->FulfillReceiveRequest(Data, MigrationContext);
 }
 
 void OnObjectDataDenied(FRemoteObjectId ObjectId, FRemoteServerId RemoteServerId)
 {
 	check(!AutoRTFM::IsClosed());
-	GTransferQueue.DenyReceiveRequest(ObjectId);
+	FRemoteObjectTransferQueue::Instance->DenyReceiveRequest(ObjectId);
 }
 
-void InitRemoteObjectTransfer()
+void InitRemoteObjectTransfer(FRemoteExecutor* Executor)
 {
-	UE::RemoteExecutor::RegisterRemoteSubsystem(&GTransferQueue);
+	FRemoteObjectTransferQueue* Queue = new FRemoteObjectTransferQueue();
+	FRemoteObjectTransferQueue::Instance = Queue;
+	UE::RemoteExecutor::RegisterRemoteSubsystem(Executor, Queue);
 }
 
 } // namespace UE::RemoteObject::Transfer
diff --git a/Engine/Source/Runtime/CoreUObject/Public/UObject/RemoteExecutor.h b/Engine/Source/Runtime/CoreUObject/Public/UObject/RemoteExecutor.h
--- a/Engine/Source/Runtime/CoreUObject/Public/UObject/RemoteExecutor.h
+++ b/Engine/Source/Runtime/CoreUObject/Public/UObject/RemoteExecutor.h
@@ -109,12 +109,16 @@ class FRemoteSubsystemBase
 {
 protected:
 	TArray<TTuple<FRemoteTransactionId, void*>> Requests;
+	FRemoteServerId LocalServerId;
 
 public:
 
 	FRemoteSubsystemBase() = default;
+	explicit FRemoteSubsystemBase(FRemoteServerId InLocalServerId) : LocalServerId(InLocalServerId) {}
 	virtual ~FRemoteSubsystemBase() = default;
 
+	FRemoteServerId GetLocalServerId() const { return LocalServerId; }
+
 	int32 GetRequestCount() const { return Requests.Num(); }
 
 	// implemented in FRemoteSubsystem<>
@@ -123,6 +127,9 @@ public:
 	virtual void SetActiveRequest(FRemoteTransactionId RequestId) = 0;
 	virtual void ClearActiveRequest() = 0;
 
+	virtual void ActivateForTesting() = 0;
+	virtual void DeactivateForTesting() = 0;
+
 	// for subsystem implementor to override:
 	virtual const TCHAR* NameForDebug() = 0;
 	virtual void BeginRequest() = 0;
@@ -138,10 +145,32 @@ public:
 	virtual void EndRequest(bool bTransactionCommitted) = 0;
 };
 
-template<typename RequestType>
+template<typename SubsystemType, typename RequestType>
 class FRemoteSubsystem : public FRemoteSubsystemBase
 {
 public:
+	using FRemoteSubsystemBase::FRemoteSubsystemBase;
+
+	static inline SubsystemType* Instance = nullptr;
+
+	~FRemoteSubsystem()
+	{
+		if (Instance == this)
+		{
+			Instance = nullptr;
+		}
+	}
+
+	void ActivateForTesting() final
+	{
+		Instance = static_cast<SubsystemType*>(this);
+	}
+
+	void DeactivateForTesting() final
+	{
+		Instance = nullptr;
+	}
+
 	RequestType* ActiveRequest = nullptr;
 
 	RequestType* GetRequestByIndex(int32 RequestIndex) const
@@ -209,6 +238,20 @@ public:
 	}
 };
 
+class FRemoteExecutor;
+
+struct FRemoteExecutorDelegates
+{
+	// Protocol delegates
+	TDelegate<void()> TickNetworkDelegate;
+	TDelegate<void(FRemoteTransactionId, FRemoteWorkPriority, const TArray<FRemoteServerId>&)> BeginMultiServerCommitDelegate;
+	TDelegate<void(FRemoteTransactionId, const TArray<FRemoteServerId>&)> ReadyMultiServerCommitDelegate;
+	TDelegate<void(FRemoteTransactionId, const TArray<FRemoteServerId>&)> AbandonMultiServerCommitDelegate;
+	TDelegate<void(FRemoteTransactionId, const TArray<FRemoteServerId>&)> EndMultiServerCommitDelegate;
+	TDelegate<void(FRemoteTransactionId, FRemoteServerId)> ReadyRemoteMultiServerCommitDelegate;
+	TDelegate<void(FRemoteTransactionId, FRemoteServerId)> AbortRemoteMultiServerCommitDelegate;
+};
+
 namespace UE::RemoteExecutor
 {
 	enum class EAbortTransactionBehavior : uint8
@@ -242,19 +285,6 @@ namespace UE::RemoteExecutor
 		COREUOBJECT_API ~FSpeculationExecutionScope();
 	};
 
-
-	/** Delegate that pumps the network and executes RPC handlers */
-	extern COREUOBJECT_API TDelegate<void()> TickNetworkDelegate;
-	extern COREUOBJECT_API TDelegate<void(FRemoteTransactionId, FRemoteWorkPriority, const TArray<FRemoteServerId>&)> BeginMultiServerCommitDelegate;
-	extern COREUOBJECT_API TDelegate<void(FRemoteTransactionId, const TArray<FRemoteServerId>&)> ReadyMultiServerCommitDelegate;
-	extern COREUOBJECT_API TDelegate<void(FRemoteTransactionId, const TArray<FRemoteServerId>&)> AbandonMultiServerCommitDelegate;
-	extern COREUOBJECT_API TDelegate<void(FRemoteTransactionId, const TArray<FRemoteServerId>&)> EndMultiServerCommitDelegate;
-	extern COREUOBJECT_API TDelegate<void(FRemoteTransactionId, FRemoteServerId)> ReadyRemoteMultiServerCommitDelegate;
-	extern COREUOBJECT_API TDelegate<void(FRemoteTransactionId, FRemoteServerId)> AbortRemoteMultiServerCommitDelegate;
-
-	/** Fetch the next deferred RPC request **/
-	extern COREUOBJECT_API TDelegate<TOptional<TTuple<FName, FRemoteWorkPriority, bool, TUniqueFunction<void(void)>>>(void)> FetchNextDeferredRPCDelegate;
-
 	/** Delegate executed when a transaction is queued in the executor for execution */
 	extern COREUOBJECT_API TMulticastDelegate<void(FRemoteTransactionId, FName)> OnTransactionQueuedDelegate;
 
@@ -273,8 +303,19 @@ namespace UE::RemoteExecutor
 	extern COREUOBJECT_API TDelegate<void(FName /*Text*/)> OnRegionBeginDelegate;
 	extern COREUOBJECT_API TDelegate<void(const FString& /*Text*/)> OnRegionEndDelegate;
 
+	COREUOBJECT_API FRemoteExecutor* CreateRemoteExecutor(
+		FRemoteServerId LocalServerId,
+		FRemoteExecutorDelegates* Delegates);
+
+	COREUOBJECT_API void DestroyRemoteExecutor(FRemoteExecutor* Executor);
+
+	COREUOBJECT_API void SetRemoteExecutorForTesting(FRemoteExecutor* Executor);
+	COREUOBJECT_API void SetSingleStepModeForTesting(bool bEnable);
+	
+	COREUOBJECT_API bool IsRemoteExecutorActive();
+
 	COREUOBJECT_API bool IsSpeculationMode();
-	COREUOBJECT_API void RegisterRemoteSubsystem(FRemoteSubsystemBase* Subsystem);
+	COREUOBJECT_API void RegisterRemoteSubsystem(FRemoteExecutor* Executor, FRemoteSubsystemBase* Subsystem);
 
 	/** Notify the executor that the network is waiting (we should tick it next chance we get) */
 	COREUOBJECT_API void NotifyNetworkIsWaiting();
@@ -287,6 +328,7 @@ namespace UE::RemoteExecutor
 	COREUOBJECT_API void EnqueueWorkListWithExplicitPriority(FName WorkName, FRemoteWorkPriority WorkPriority, bool bIsTransactional, TArray<TUniqueFunction<void(void)>>&& WorkList);
 	COREUOBJECT_API void EnqueueWorkList(FName WorkName, bool bIsTransactional, TArray<TUniqueFunction<void(void)>>&& WorkList);
 	COREUOBJECT_API void ExecutePendingWork();
+	COREUOBJECT_API bool HasPendingWork();
 
 	COREUOBJECT_API void AbortCurrentTransaction(EAbortTransactionBehavior AbortBehavior, ERemoteExecutorAbortReason AbortReason, FStringView AbortDescription);
 
diff --git a/Engine/Source/Runtime/CoreUObject/Public/UObject/RemoteHeap.h b/Engine/Source/Runtime/CoreUObject/Public/UObject/RemoteHeap.h
--- a/Engine/Source/Runtime/CoreUObject/Public/UObject/RemoteHeap.h
+++ b/Engine/Source/Runtime/CoreUObject/Public/UObject/RemoteHeap.h
@@ -73,8 +73,10 @@ public:
 	FRemoteWorkPriority RequestPriority;
 };
 
-class FRemoteHeap : public FRemoteSubsystem<FRemoteHeapRequest>
+class FRemoteHeap : public FRemoteSubsystem<FRemoteHeap, FRemoteHeapRequest>
 {
+	using FRemoteSubsystem::FRemoteSubsystem;
+
 	TArray<SRemoteHeapLineEntry> HeapLines;
 
 	// this is a list of heap lines that we have locally that 
@@ -86,8 +88,7 @@ class FRemoteHeap : public FRemoteSubsystem<FRemoteHeapRequest>
 	TArray<FRemoteHeapLineRequest> PendingLineRequests;
 
 public:
-	static FRemoteHeap& GetInstance();
-	void Initialize();
+	static void CreateInstance(class FRemoteExecutor* Executor, FRemoteServerId LocalServerId);
 
 	const TCHAR* NameForDebug() final { return TEXT("RemoteHeap"); }
 	void BeginRequest() final;
diff --git a/Engine/Source/Runtime/CoreUObject/Public/UObject/RemoteObjectTransfer.h b/Engine/Source/Runtime/CoreUObject/Public/UObject/RemoteObjectTransfer.h
--- a/Engine/Source/Runtime/CoreUObject/Public/UObject/RemoteObjectTransfer.h
+++ b/Engine/Source/Runtime/CoreUObject/Public/UObject/RemoteObjectTransfer.h
@@ -290,7 +290,7 @@ namespace UE::RemoteObject::Transfer
 
 	extern COREUOBJECT_API const FRemoteServerId DatabaseId;
 
-	COREUOBJECT_API void InitRemoteObjectTransfer();
+	COREUOBJECT_API void InitRemoteObjectTransfer(class FRemoteExecutor* Executor);
 
 	/**
 	* Returns object data that has not yet been deserialized on this server to their respective owners
