When constructing subobjects from a template, copy over transient flags to handle cases where there are mixed transient objects in a persistent template
#Synced-CL 51587375

diff --git a/Engine/Source/Runtime/CoreUObject/Private/UObject/CoreNative.cpp b/Engine/Source/Runtime/CoreUObject/Private/UObject/CoreNative.cpp
--- a/Engine/Source/Runtime/CoreUObject/Private/UObject/CoreNative.cpp
+++ b/Engine/Source/Runtime/CoreUObject/Private/UObject/CoreNative.cpp
@@ -342,7 +342,10 @@ UObject* FObjectInstancingGraph::GetInstancedSubobject( UObject* SourceSubobject
 								FStaticConstructObjectParameters Params(SourceSubobject->GetClass());
 								Params.Outer = SubobjectOuter;
 								Params.Name = SubobjectName;
-								Params.SetFlags = SubobjectOuter->GetMaskedFlags(RF_PropagateToSubObjects);
+
+								// Copy transient flags from template to avoid saving subobjects that shouldn't be saved
+								EObjectFlags SourceFlagsToCopy = SourceSubobject->GetFlags() & (EObjectFlags)(RF_Transient | RF_TextExportTransient);
+								Params.SetFlags = SubobjectOuter->GetMaskedFlags(RF_PropagateToSubObjects) | SourceFlagsToCopy;
 								Params.Template = SourceSubobject;
 								Params.bCopyTransientsFromClassDefaults = true;
 								Params.InstanceGraph = this;
