[AutoRTFM] Mark the garbage collection logic as AutoRTFM-disabled.

Garbage collection should never happen within a transaction, so it's unnecessary to compile a closed form version of the garbage collector.

#Synced-CL 51590581

diff --git a/Engine/Source/Runtime/CoreUObject/Private/UObject/GarbageCollection.cpp b/Engine/Source/Runtime/CoreUObject/Private/UObject/GarbageCollection.cpp
--- a/Engine/Source/Runtime/CoreUObject/Private/UObject/GarbageCollection.cpp
+++ b/Engine/Source/Runtime/CoreUObject/Private/UObject/GarbageCollection.cpp
@@ -805,6 +805,8 @@ void OnDisregardForGCSetDisabled(int32 NumObjects)
 	using namespace UE::GC::Private;
 	using FDisregardSetDisabledState = TThreadedGather<TSet<int32>>;
 
+	AutoRTFM::UnreachableIfTransactional();
+
 	FDisregardSetDisabledState DisregardSetDisabledState;
 	
 	DisregardSetDisabledState.Start(EGatherOptions::Parallel, NumObjects);
@@ -3167,7 +3169,7 @@ public:
 
 /** Batches up references before dispatching them to the processor, unlike TDirectDispatcher */
 template <class ProcessorType>
-struct TBatchDispatcher
+struct AUTORTFM_DISABLE TBatchDispatcher
 {
 	static constexpr bool bBatching = true;
 	static constexpr bool bParallel = IsParallel(ProcessorType::Options);
@@ -3514,7 +3516,7 @@ private:
 //////////////////////////////////////////////////////////////////////////
 
 template <EGCOptions Options>
-class TReachabilityCollectorBase : public FReferenceCollector
+class AUTORTFM_DISABLE TReachabilityCollectorBase : public FReferenceCollector
 {
 protected:
 	bool bAllowEliminatingReferences = true;
@@ -3542,7 +3544,7 @@ public:
 };
 
 template <EGCOptions Options>
-class TReachabilityCollector final : public TReachabilityCollectorBase<Options>
+class AUTORTFM_DISABLE TReachabilityCollector final : public TReachabilityCollectorBase<Options>
 {
 	using ProcessorType = TReachabilityProcessor<Options>;
 	using DispatcherType = TBatchDispatcher<ProcessorType>;
@@ -3749,7 +3751,7 @@ private:
 };
 
 template <EGCOptions Options>
-class TDebugReachabilityCollector final : public TReachabilityCollectorBase<Options>
+class AUTORTFM_DISABLE TDebugReachabilityCollector final : public TReachabilityCollectorBase<Options>
 {
 	using Super = TReachabilityCollectorBase<Options>;
 	using Super::bAllowEliminatingReferences;
@@ -4145,7 +4147,7 @@ static FORCEINLINE void StopVerseGC()
 namespace UE::GC
 {
 
-class FRealtimeGC : public FGarbageCollectionTracer
+class AUTORTFM_DISABLE FRealtimeGC : public FGarbageCollectionTracer
 {
 	typedef void(FRealtimeGC::*ReachabilityAnalysisFn)(FWorkerContext&);
 
@@ -4215,7 +4217,7 @@ class FRealtimeGC : public FGarbageCollectionTracer
 	}
 
 	template <EGCOptions Options>
-	void PerformReachabilityAnalysisOnObjectsInternal(FWorkerContext& Context)
+	AUTORTFM_DISABLE void PerformReachabilityAnalysisOnObjectsInternal(FWorkerContext& Context)
 	{
 		TRACE_CPUPROFILER_EVENT_SCOPE(PerformReachabilityAnalysisOnObjectsInternal);
 
@@ -5263,7 +5265,7 @@ static FAutoConsoleVariableRef CVarFlushStreamingOnGC(
 	ECVF_Default
 	);
 
-void DissolveUnreachableClusters(UE::GC::EGatherOptions Options)
+AUTORTFM_DISABLE static void DissolveUnreachableClusters(UE::GC::EGatherOptions Options)
 {
 	const int32 InitialNumClusters = GUObjectClusters.GetNumAllocatedClusters();
 	if (!InitialNumClusters)
@@ -5468,7 +5470,7 @@ void GatherUnreachableObjects(bool bForceSingleThreaded)
 namespace UE::GC
 {
 
-class FWeakReferenceEliminator final : public TReachabilityCollectorBase<EGCOptions::None>
+class AUTORTFM_DISABLE FWeakReferenceEliminator final : public TReachabilityCollectorBase<EGCOptions::None>
 {
 
 public:
@@ -5625,9 +5627,9 @@ static void UpdateGCHistory(TConstArrayView<TUniquePtr<FWorkerContext>> Contexts
 }
 
 template<bool bPerformFullPurge>
-FORCENOINLINE void CollectGarbageImpl(EObjectFlags KeepFlags);
+AUTORTFM_DISABLE FORCENOINLINE void CollectGarbageImpl(EObjectFlags KeepFlags);
 
-FORCENOINLINE static void CollectGarbageIncremental(EObjectFlags KeepFlags)
+AUTORTFM_DISABLE FORCENOINLINE static void CollectGarbageIncremental(EObjectFlags KeepFlags)
 {
 	SCOPE_TIME_GUARD(TEXT("Collect Garbage Incremental"));
 	SCOPED_NAMED_EVENT(CollectGarbageIncremental, FColor::Red);
@@ -5637,7 +5639,7 @@ FORCENOINLINE static void CollectGarbageIncremental(EObjectFlags KeepFlags)
 	CollectGarbageImpl<false>(KeepFlags);
 }
 
-FORCENOINLINE static void CollectGarbageFull(EObjectFlags KeepFlags)
+AUTORTFM_DISABLE FORCENOINLINE static void CollectGarbageFull(EObjectFlags KeepFlags)
 {
 	SCOPE_TIME_GUARD(TEXT("Collect Garbage Full"));
 	SCOPED_NAMED_EVENT(CollectGarbageFull, FColor::Red);
@@ -5647,7 +5649,7 @@ FORCENOINLINE static void CollectGarbageFull(EObjectFlags KeepFlags)
 	CollectGarbageImpl<true>(KeepFlags);
 }
 
-FORCEINLINE void CollectGarbageInternal(EObjectFlags KeepFlags, bool bPerformFullPurge)
+AUTORTFM_DISABLE FORCEINLINE void CollectGarbageInternal(EObjectFlags KeepFlags, bool bPerformFullPurge)
 {
 	const double StartTime = FPlatformTime::Seconds();
 
@@ -5682,7 +5684,7 @@ EGatherOptions GetObjectGatherOptions()
  * @param	bPerformFullPurge	if true, perform a full purge after the mark pass
  */
 template<bool bPerformFullPurge>
-void PreCollectGarbageImpl(EObjectFlags KeepFlags)
+AUTORTFM_DISABLE void PreCollectGarbageImpl(EObjectFlags KeepFlags)
 {
 	using namespace UE::GC::Private;
 
@@ -5806,7 +5808,7 @@ void PreCollectGarbageImpl(EObjectFlags KeepFlags)
 }
 
 template<bool bPerformFullPurge>
-void CollectGarbageImpl(EObjectFlags KeepFlags)
+AUTORTFM_DISABLE void CollectGarbageImpl(EObjectFlags KeepFlags)
 {
 	{
 		// Reachability analysis.
@@ -5821,7 +5823,7 @@ void CollectGarbageImpl(EObjectFlags KeepFlags)
 }
 
 template<bool bPerformFullPurge>
-void PostCollectGarbageImpl(EObjectFlags KeepFlags)
+AUTORTFM_DISABLE void PostCollectGarbageImpl(EObjectFlags KeepFlags)
 {
 	const double PostCollectStartTime = FPlatformTime::Seconds();
 
@@ -6165,6 +6167,8 @@ bool IsIncrementalReachabilityAnalysisPending()
 
 void PerformIncrementalReachabilityAnalysis(double TimeLimit)
 {
+	AutoRTFM::UnreachableIfTransactional();
+
 	checkf(UE::GC::GIsIncrementalReachabilityPending, TEXT("Incremental reachability must be pending to perform its next iteration"));
 	// When performing Reachability Analysis iterations start the internal timer before acquiring GC lock 
 	// so that we don't spend more time than GIncrementalReachabilityTimeLimit on fully completing an iteration
@@ -6360,6 +6364,8 @@ bool TryCollectGarbage(EObjectFlags KeepFlags, bool bPerformFullPurge)
 		return false;
 	}
 
+	AutoRTFM::UnreachableIfTransactional();
+
 	// No other thread may be performing UObject operations while we're running so try to acquire GC lock
 	if (UE::GC::GIsIncrementalReachabilityPending)
 	{
@@ -7265,7 +7271,7 @@ ELoot StealWork(FWorkerContext& Context, FReferenceCollector& Collector, FWorkBl
 	return ELoot::Nothing;
 }
 
-TArrayView<FWorkerContext*> InitializeAsyncProcessingContexts(FWorkerContext& InContext)
+AUTORTFM_DISABLE static TArrayView<FWorkerContext*> InitializeAsyncProcessingContexts(FWorkerContext& InContext)
 {
 	TArrayView<FWorkerContext*> Contexts;
 
diff --git a/Engine/Source/Runtime/CoreUObject/Private/UObject/ReachabilityAnalysisState.h b/Engine/Source/Runtime/CoreUObject/Private/UObject/ReachabilityAnalysisState.h
--- a/Engine/Source/Runtime/CoreUObject/Private/UObject/ReachabilityAnalysisState.h
+++ b/Engine/Source/Runtime/CoreUObject/Private/UObject/ReachabilityAnalysisState.h
@@ -18,7 +18,7 @@ namespace UE::GC
 /**
 * Reachability analysis state holds information about suspended worker contexts.
 */
-class FReachabilityAnalysisState
+class AUTORTFM_DISABLE FReachabilityAnalysisState
 {
 public:
 
@@ -154,7 +154,7 @@ inline int32 GetNumThreadsForGather(const EGatherOptions Options, const int32 Nu
 template <typename PayloadType>
 struct TGatherIterator
 {
-	/** Inital start index */
+	/** Initial start index */
 	int32 StartIndex = 0;
 	/** Current index of an object */
 	int32 Index = 0;
@@ -193,7 +193,7 @@ struct TDefaultPayloadOps
 
 /** Helper class that holds the current state of Object Gathering phases */
 template <typename PayloadType, typename PayloadOps = TDefaultPayloadOps<PayloadType>>
-class TThreadedGather
+class AUTORTFM_DISABLE TThreadedGather
 {
 public:
 
diff --git a/Engine/Source/Runtime/CoreUObject/Public/UObject/FastReferenceCollector.h b/Engine/Source/Runtime/CoreUObject/Public/UObject/FastReferenceCollector.h
--- a/Engine/Source/Runtime/CoreUObject/Public/UObject/FastReferenceCollector.h
+++ b/Engine/Source/Runtime/CoreUObject/Public/UObject/FastReferenceCollector.h
@@ -888,7 +888,7 @@ COREUOBJECT_API ELoot StealWork(FWorkerContext& Context, FReferenceCollector& Co
 COREUOBJECT_API void SuspendWork(FWorkerContext& Context);
 
 /** Allocates contexts and coordinator, kicks worker tasks that also call ProcessSync. Processor is type-erased to void* to avoid templated code. */
-COREUOBJECT_API void ProcessAsync(void (*ProcessSync)(void*, FWorkerContext&), void* Processor, FWorkerContext& InitialContext);
+AUTORTFM_DISABLE COREUOBJECT_API void ProcessAsync(void (*ProcessSync)(void*, FWorkerContext&), void* Processor, FWorkerContext& InitialContext);
 
 //////////////////////////////////////////////////////////////////////////
 
