Add option to always reset package localization id on package save

[RN] A new editor INI setting was introduced to force a reset to a deterministic localization package id on save for all packages

#Synced-CL 51591193

#Synced-CL 50983539

diff --git a/Engine/Plugins/Editor/ContentBrowser/ContentBrowserAssetDataSource/Source/ContentBrowserAssetDataSource/Private/AssetFileContextMenu.cpp b/Engine/Plugins/Editor/ContentBrowser/ContentBrowserAssetDataSource/Source/ContentBrowserAssetDataSource/Private/AssetFileContextMenu.cpp
--- a/Engine/Plugins/Editor/ContentBrowser/ContentBrowserAssetDataSource/Source/ContentBrowserAssetDataSource/Private/AssetFileContextMenu.cpp
+++ b/Engine/Plugins/Editor/ContentBrowser/ContentBrowserAssetDataSource/Source/ContentBrowserAssetDataSource/Private/AssetFileContextMenu.cpp
@@ -1902,9 +1902,9 @@ void FAssetFileContextMenu::ExecuteGoToDocsForAsset(UClass* SelectedClass, const
 	}
 }
 
+#if USE_STABLE_LOCALIZATION_KEYS
 void FAssetFileContextMenu::ExecuteResetLocalizationId()
 {
-#if USE_STABLE_LOCALIZATION_KEYS
 	const FText ResetLocalizationIdMsg = LOCTEXT("ResetLocalizationIdMsg", "This will reset the localization ID of the selected assets and cause all text within them to lose their existing translations.\n\nAre you sure you want to do this?");
 	if (FMessageDialog::Open(EAppMsgType::YesNo, ResetLocalizationIdMsg) != EAppReturnType::Yes)
 	{
@@ -1916,13 +1916,15 @@ void FAssetFileContextMenu::ExecuteResetLocalizationId()
 		UObject* Asset = AssetData.GetAsset();
 		if (Asset)
 		{
-			Asset->Modify();
-			TextNamespaceUtil::ClearPackageNamespace(Asset);
-			TextNamespaceUtil::EnsurePackageNamespace(Asset);
+			if (!TextNamespaceUtil::HasDeterministicPackageNamespace(Asset))
+			{
+				Asset->Modify();
+				TextNamespaceUtil::ResetPackageNamespace(Asset);
+			}
 		}
 	}
-#endif // USE_STABLE_LOCALIZATION_KEYS
 }
+#endif // USE_STABLE_LOCALIZATION_KEYS
 
 void FAssetFileContextMenu::ExecuteShowLocalizationCache(const FAssetData InAsset, const FString InPackageFilename)
 {
diff --git a/Engine/Plugins/Editor/ContentBrowser/ContentBrowserAssetDataSource/Source/ContentBrowserAssetDataSource/Public/AssetFileContextMenu.h b/Engine/Plugins/Editor/ContentBrowser/ContentBrowserAssetDataSource/Source/ContentBrowserAssetDataSource/Public/AssetFileContextMenu.h
--- a/Engine/Plugins/Editor/ContentBrowser/ContentBrowserAssetDataSource/Source/ContentBrowserAssetDataSource/Public/AssetFileContextMenu.h
+++ b/Engine/Plugins/Editor/ContentBrowser/ContentBrowserAssetDataSource/Source/ContentBrowserAssetDataSource/Public/AssetFileContextMenu.h
@@ -166,9 +166,11 @@ private:
 
 	/** Handler for GoToAssetDocs */
 	UE_API void ExecuteGoToDocsForAsset(UClass* SelectedClass, const FString ExcerptSection);
-	
+
+#if USE_STABLE_LOCALIZATION_KEYS
 	/** Handler for resetting the localization ID of the current selection */
 	UE_API void ExecuteResetLocalizationId();
+#endif
 
 	/** Handler for showing the cached list of localized texts stored in the package header */
 	UE_API void ExecuteShowLocalizationCache(const FAssetData InAsset, const FString InPackageFilename);
diff --git a/Engine/Source/Editor/UnrealEd/Private/UserDefinedStructEditorData.cpp b/Engine/Source/Editor/UnrealEd/Private/UserDefinedStructEditorData.cpp
--- a/Engine/Source/Editor/UnrealEd/Private/UserDefinedStructEditorData.cpp
+++ b/Engine/Source/Editor/UnrealEd/Private/UserDefinedStructEditorData.cpp
@@ -11,11 +11,46 @@
 #include "Blueprint/BlueprintSupport.h"
 #include "Editor.h"
 #include "EdGraphSchema_K2.h"
+#include "Internationalization/TextPackageNamespaceUtil.h"
 
 #include UE_INLINE_GENERATED_CPP_BY_NAME(UserDefinedStructEditorData)
 
 #define LOCTEXT_NAMESPACE "UserDefinedStructEditorData"
 
+#if USE_STABLE_LOCALIZATION_KEYS 
+
+namespace
+{
+	void LocalizationUpdateTextsPackageNamespaceCallback(UObject* Object, TextNamespaceUtil::FUpdateTextsPackageNamespace& UpdateTextsPackageNamespace)
+	{
+		// User Defined Structs need some special handling as they store their default data in a way that serialize doesn't pick up
+		if (UUserDefinedStruct* const UserDefinedStruct = Cast<UUserDefinedStruct>(Object))
+		{
+			if (UUserDefinedStructEditorData* UDSEditorData = Cast<UUserDefinedStructEditorData>(UserDefinedStruct->EditorData))
+			{
+				for (FStructVariableDescription& StructVariableDesc : UDSEditorData->VariablesDescriptions)
+				{
+					static const FName TextCategory = TEXT("text"); // Must match UEdGraphSchema_K2::PC_Text
+					if (StructVariableDesc.Category == UEdGraphSchema_K2::PC_Text/*TextCategory*/)
+					{
+						FText StructVariableValue;
+						if (FTextStringHelper::ReadFromBuffer(*StructVariableDesc.DefaultValue, StructVariableValue) && UpdateTextsPackageNamespace.UpdateSingleText(StructVariableValue))
+						{
+							FTextStringHelper::WriteToBuffer(StructVariableDesc.DefaultValue, StructVariableValue);
+						}
+					}
+				}
+			}
+		}
+
+		UpdateTextsPackageNamespace.Serialize(Object);
+	}
+
+	static const TextNamespaceUtil::FAutoRegisterLocalizationUpdateTextsPackageNamespaceCallback AutomaticRegistrationOfLocalizationUpdateTextsPackageNamespace(UUserDefinedStruct::StaticClass(), &LocalizationUpdateTextsPackageNamespaceCallback);
+}
+
+#endif // USE_STABLE_LOCALIZATION_KEYS
+
 const FName FStructVariableMetaData::ClampMin("ClampMin");
 const FName FStructVariableMetaData::ClampMax("ClampMax");
 const FName FStructVariableMetaData::UIMin("UIMin");
diff --git a/Engine/Source/Runtime/CoreUObject/Private/Internationalization/TextPackageNamespaceUtil.cpp b/Engine/Source/Runtime/CoreUObject/Private/Internationalization/TextPackageNamespaceUtil.cpp
--- a/Engine/Source/Runtime/CoreUObject/Private/Internationalization/TextPackageNamespaceUtil.cpp
+++ b/Engine/Source/Runtime/CoreUObject/Private/Internationalization/TextPackageNamespaceUtil.cpp
@@ -9,6 +9,7 @@
 #include "UObject/Class.h"
 #include "UObject/MetaData.h"
 #include "UObject/Package.h"
+#include "UObject/PackageFileSummary.h"
 #include "UObject/TextProperty.h"
 
 #if USE_STABLE_LOCALIZATION_KEYS
@@ -108,8 +109,162 @@ FString GenerateDeterministicPackageNamespace(const FString& InPackagePath, cons
 	return NewGuid.ToString();
 }
 
+bool HasDeterministicPackageNamespace(const FString& InPackagePath, const FString& InCurrentPackageNamespace, const FGuid& InPersistentGuid)
+{
+	if (!InCurrentPackageNamespace.IsEmpty())
+	{
+		if (FPackageName::IsScriptPackage(InPackagePath))
+		{
+			return true;
+		}
+
+		const FString& DeterministicPackageNamespace = GenerateDeterministicPackageNamespace(InPackagePath, InPersistentGuid);
+		if (InCurrentPackageNamespace.Equals(DeterministicPackageNamespace))
+		{
+			return true;
+		}
+	}
+
+	return false;
+}
+
 } // TextNamespaceUtilImpl
 
+namespace TextNamespaceUtil
+{
+	class FUpdateTextsPackageNamespace::FImpl : public FArchiveUObject
+	{
+	public:
+		FImpl(UPackage* InPackage)
+			: Package(InPackage)
+			, ActualPackageNamespace(TextNamespaceUtil::EnsurePackageNamespace(InPackage))
+		{
+		}
+
+		void UpdateTextsPackageNamespace(TextNamespaceUtil::FUpdateTextsPackageNamespace* InOriginalUpdateTextsPackageNamespace)
+		{
+			this->SetIsSaving(true);
+
+			TArray<UObject*> AllObjectsInPackage;
+			GetObjectsWithOuter(Package, AllObjectsInPackage, true, RF_Transient, EInternalObjectFlags::Garbage);
+
+			for (UObject* Obj : AllObjectsInPackage)
+			{
+				UClass* CurrentClass = nullptr;
+				if (InOriginalUpdateTextsPackageNamespace)
+				{
+					CurrentClass = Obj->GetClass();
+				}
+
+				while (CurrentClass)
+				{
+					const TextNamespaceUtil::FUpdateTextsPackageNamespace::FProcessObjectCallback* CustomCallback =
+						TextNamespaceUtil::FUpdateTextsPackageNamespace::GetTypeSpecificProcessObjectCallbacks().Find(CurrentClass);
+					if (CustomCallback)
+					{
+						(*CustomCallback)(Obj, *InOriginalUpdateTextsPackageNamespace);
+						break;
+					}
+					else
+					{
+						CurrentClass = CurrentClass->GetSuperClass();
+					}
+
+				}
+
+				if (CurrentClass == nullptr) // No callback was used for that class (and all super classes)
+				{
+					Obj->Serialize(*this);
+				}
+			}
+		}
+
+		bool UpdateSingleText(FText& InOutText)
+		{
+			if (!InOutText.ShouldGatherForLocalization())
+			{
+				return false;
+			}
+
+			const FTextId TextId = FTextInspector::GetTextId(InOutText);
+			if (TextId.IsEmpty())
+			{
+				return false;
+			}
+
+			const FString CurrentNamespace = TextId.GetNamespace().ToString();
+			const FString ExpectedNamespace = TextNamespaceUtil::BuildFullNamespace(CurrentNamespace, ActualPackageNamespace);
+			if (CurrentNamespace.Equals(ExpectedNamespace, ESearchCase::CaseSensitive))
+			{
+				return false;
+			}
+
+			FTextKey NewNamespace;
+			FTextKey NewKey;
+			if (const FTextId* ExistingMapping = PackageTextIdMap.Find(TextId))
+			{
+				NewNamespace = ExistingMapping->GetNamespace();
+				NewKey = ExistingMapping->GetKey();
+			}
+			else
+			{
+				NewNamespace = ExpectedNamespace;
+				NewKey = FGuid::NewGuid().ToString();
+
+				PackageTextIdMap.Add(TextId, FTextId(NewNamespace, NewKey));
+			}
+
+			InOutText = FText::ChangeKey(NewNamespace, NewKey, InOutText);
+			return true;
+		}
+
+		void CallBackSerialize(UObject* InObject)
+		{
+			InObject->Serialize(*this);
+		}
+
+		using FArchiveUObject::operator<<; // For visibility of the overloads we don't override
+
+		virtual FArchive& operator<<(FText& Value) override
+		{
+			UpdateSingleText(Value);
+			return *this;
+		}
+
+	private:
+		UPackage* Package;
+		FString ActualPackageNamespace;
+		TMap<FTextId, FTextId> PackageTextIdMap;
+	};
+} // TextNamespaceUtil
+
+void TextNamespaceUtil::FUpdateTextsPackageNamespace::UpdatePackage(UPackage* InPackage)
+{
+	FUpdateTextsPackageNamespace UpdateTextsPackageNamespace(InPackage);
+}
+
+TextNamespaceUtil::FUpdateTextsPackageNamespace::FUpdateTextsPackageNamespace(UPackage* InPackage)
+{
+	PImpl = MakePimpl<FImpl>(InPackage);
+	PImpl->UpdateTextsPackageNamespace(this);
+}
+
+bool TextNamespaceUtil::FUpdateTextsPackageNamespace::UpdateSingleText(FText& InOutText)
+{
+	return PImpl->UpdateSingleText(InOutText);
+}
+
+void TextNamespaceUtil::FUpdateTextsPackageNamespace::Serialize(UObject* InObject)
+{
+	PImpl->CallBackSerialize(InObject);
+}
+
+TextNamespaceUtil::FUpdateTextsPackageNamespace::FProcessObjectCallbackMap& TextNamespaceUtil::FUpdateTextsPackageNamespace::GetTypeSpecificProcessObjectCallbacks()
+{
+	static TextNamespaceUtil::FUpdateTextsPackageNamespace::FProcessObjectCallbackMap CallbackMap;
+	return CallbackMap;
+}
+
 FString TextNamespaceUtil::GenerateDeterministicPackageNamespace(const UPackage* InPackage)
 {
 	FString PackagePath = InPackage->GetName();
@@ -117,6 +272,56 @@ FString TextNamespaceUtil::GenerateDeterministicPackageNamespace(const UPackage*
 	return TextNamespaceUtilImpl::GenerateDeterministicPackageNamespace(PackagePath, NonUniqueGuid);
 }
 
+bool TextNamespaceUtil::HasDeterministicPackageNamespace(const UPackage* InPackage)
+{
+	if (InPackage)
+	{
+		FString PackagePath = InPackage->GetName();
+		FString CurrentPackageNamespace = TextNamespaceUtil::GetPackageNamespace(InPackage);
+		FGuid PersistentGuid = InPackage->GetPersistentGuid(); // Can be the same for duplicated packages
+		return TextNamespaceUtilImpl::HasDeterministicPackageNamespace(PackagePath, CurrentPackageNamespace, PersistentGuid);
+	}
+	else
+	{
+		return false;
+	}
+}
+
+bool TextNamespaceUtil::HasDeterministicPackageNamespace(UObject* InObject)
+{
+	UPackage* Package = InObject ? InObject->GetOutermost() : nullptr;
+	return HasDeterministicPackageNamespace(Package);
+}
+
+bool TextNamespaceUtil::HasDeterministicPackageNamespace(const FString& InPackagePath, const FPackageFileSummary& InPackageFileSummary)
+{
+	const FString& CurrentPackageNamespace = InPackageFileSummary.LocalizationId;
+	const FGuid& PersistentGuid = InPackageFileSummary.PersistentGuid;
+	return TextNamespaceUtilImpl::HasDeterministicPackageNamespace(InPackagePath, CurrentPackageNamespace, PersistentGuid);
+}
+
+bool TextNamespaceUtil::ResetPackageNamespace(UPackage* InPackage)
+{
+	if (InPackage && !HasDeterministicPackageNamespace(InPackage))
+	{
+		TextNamespaceUtil::ClearPackageNamespace(InPackage);
+		TextNamespaceUtil::EnsurePackageNamespace(InPackage);
+
+		// We also need to update package namespace on all texts within the package. Otherwise, the cache could become invalid on save.
+		TextNamespaceUtil::FUpdateTextsPackageNamespace::UpdatePackage(InPackage);
+
+		return true;
+	}
+
+	return false;
+}
+
+bool TextNamespaceUtil::ResetPackageNamespace(UObject* InObject)
+{
+	UPackage* Package = InObject ? InObject->GetOutermost() : nullptr;
+	return ResetPackageNamespace(Package);
+}
+
 FString TextNamespaceUtil::GetPackageNamespace(const UPackage* InPackage)
 {
 	return TextNamespaceUtilImpl::FindOrAddPackageNamespace(const_cast<UPackage*>(InPackage), /*bCanAdd*/false);
diff --git a/Engine/Source/Runtime/CoreUObject/Private/UObject/SavePackage2.cpp b/Engine/Source/Runtime/CoreUObject/Private/UObject/SavePackage2.cpp
--- a/Engine/Source/Runtime/CoreUObject/Private/UObject/SavePackage2.cpp
+++ b/Engine/Source/Runtime/CoreUObject/Private/UObject/SavePackage2.cpp
@@ -275,14 +275,25 @@ ESavePackageResult ValidatePackage(FSaveContext& SaveContext)
 	return ReturnSuccessOrCancel();
 }
 
-FORCEINLINE void EnsurePackageLocalization(UPackage* InPackage)
+void EnsurePackageLocalization(UPackage* InPackage, const FSaveContext& InSaveContext)
 {
 #if USE_STABLE_LOCALIZATION_KEYS
 	if (GIsEditor)
 	{
-		// We need to ensure that we have a package localization namespace as the package loading will need it
-		// We need to do this before entering the UE::IsSavingPackage() block as it may change the package meta-data
-		TextNamespaceUtil::EnsurePackageNamespace(InPackage);
+		bool bAlwaysResetPackageLocIdOnSave = false;
+		GConfig->GetBool(TEXT("Internationalization"), TEXT("bAlwaysResetPackageLocIdOnSave"), bAlwaysResetPackageLocIdOnSave, GEditorIni);
+		if (bAlwaysResetPackageLocIdOnSave && !InSaveContext.IsCooking())
+		{
+			// Package localization IDs are meant to be stable by default. But to automatically fix conflicts, having them reset on save can be useful in some contexts.
+			// The downside being that text keys might get regenerated (they might feel unstable). Having translation memory with external tools can mitigate that downside though.
+			TextNamespaceUtil::ResetPackageNamespace(InPackage);
+		}
+		else
+		{
+			// In any case, we need to ensure that we have a package localization namespace as the package loading will need it
+			// We need to do this before entering the UE::IsSavingPackage() block as it may change the package meta-data
+			TextNamespaceUtil::EnsurePackageNamespace(InPackage);
+		}
 	}
 #endif // USE_STABLE_LOCALIZATION_KEYS
 }
@@ -4014,7 +4025,7 @@ void BeginSave(FSaveContext& SaveContext, FScopedSlowTask* SlowTask = nullptr)
 			ConditionalFlushAsyncLoadingForSave(Package);
 			(*GFlushStreamingFunc)();
 
-			EnsurePackageLocalization(Package);
+			EnsurePackageLocalization(Package, SaveContext);
 
 			// FullyLoad the package's Loader, so that anything we need to serialize (bulkdata, thumbnails) is available
 			EnsureLoadingComplete(Package);
@@ -4240,7 +4251,7 @@ ESavePackageResult UPackage::SaveConcurrent(TArrayView<FPackageSaveInfo> InPacka
 			}
 
 			// Ensures
-			EnsurePackageLocalization(SaveContext.GetPackage());
+			EnsurePackageLocalization(SaveContext.GetPackage(), SaveContext);
 			EnsureLoadingComplete(SaveContext.GetPackage()); //@todo: needed??
 
 			// PreSave Asset
diff --git a/Engine/Source/Runtime/CoreUObject/Public/Internationalization/TextPackageNamespaceUtil.h b/Engine/Source/Runtime/CoreUObject/Public/Internationalization/TextPackageNamespaceUtil.h
--- a/Engine/Source/Runtime/CoreUObject/Public/Internationalization/TextPackageNamespaceUtil.h
+++ b/Engine/Source/Runtime/CoreUObject/Public/Internationalization/TextPackageNamespaceUtil.h
@@ -7,16 +7,57 @@
 #include "Internationalization/Text.h"
 #include "Internationalization/TextNamespaceFwd.h"
 #include "Internationalization/TextNamespaceUtil.h"
+#include "Templates/PimplPtr.h"
 
 class FTextProperty;
 class UObject;
 class UPackage;
+struct FPackageFileSummary;
 
 namespace TextNamespaceUtil
 {
 
 #if USE_STABLE_LOCALIZATION_KEYS
 
+/**
+* Given a package, all texts found in all objects within the package will be updated to actually use
+* the package namespace on its text.
+*/
+class FUpdateTextsPackageNamespace
+{
+public:
+	typedef TFunction<void(UObject*, FUpdateTextsPackageNamespace&)> FProcessObjectCallback;
+	typedef TMap<const UClass*, FProcessObjectCallback> FProcessObjectCallbackMap;
+
+	// Main Entry point to this class
+	static void UpdatePackage(UPackage* InPackage);
+
+	// Callback map to customize this class
+	COREUOBJECT_API static FProcessObjectCallbackMap& GetTypeSpecificProcessObjectCallbacks();
+
+	// This function can be used during a custom callback to update one FText
+	COREUOBJECT_API bool UpdateSingleText(FText& InOutText);
+
+	// This function can be used during a custom callback to use the default serialize
+	COREUOBJECT_API void Serialize(UObject* InObject);
+
+private:
+	FUpdateTextsPackageNamespace() = delete;
+	FUpdateTextsPackageNamespace(UPackage* InPackage);
+
+	class FImpl;
+	TPimplPtr<FImpl> PImpl;
+};
+
+/** Struct to automatically register a UpdateTextsPackageNamespace callback when it's constructed */
+struct FAutoRegisterLocalizationUpdateTextsPackageNamespaceCallback
+{
+	UE_FORCEINLINE_HINT FAutoRegisterLocalizationUpdateTextsPackageNamespaceCallback(const UClass* InClass, const FUpdateTextsPackageNamespace::FProcessObjectCallback& InCallback)
+	{
+		FUpdateTextsPackageNamespace::GetTypeSpecificProcessObjectCallbacks().Add(InClass, InCallback);
+	}
+};
+
 /**
  * Generate a deterministic package namespace based on the given package info.
  * @note This key will be formatted like a GUID, but the value will actually be based on deterministic hashes.
@@ -25,6 +66,48 @@ namespace TextNamespaceUtil
  */
 COREUOBJECT_API FString GenerateDeterministicPackageNamespace(const UPackage* InPackage);
 
+/**
+ * Check if the given package uses its deterministictly generatable namespace.
+ *
+ * @param InPackage					The package to check the namespace for
+ */
+COREUOBJECT_API bool HasDeterministicPackageNamespace(const UPackage* InPackage);
+
+/**
+ * Check if the given UObject uses its deterministictly generatable namespace.
+ *
+ * @param InObject					The UObject to check the namespace for (from its owner package)
+ */
+COREUOBJECT_API bool HasDeterministicPackageNamespace(UObject* InObject);
+
+/**
+ * Check if the given package file summary uses its deterministictly generatable namespace.
+ *
+ * @param InPackagePath					The package to check the namespace for
+ * @param InPackageFileSummary				The package current Loc ID.
+ */
+COREUOBJECT_API bool HasDeterministicPackageNamespace(const FString& InPackagePath, const FPackageFileSummary& InPackageFileSummary);
+
+/**
+ * Given a package, make sure it uses its deterministic package namespace if USE_STABLE_LOCALIZATION_KEYS is enabled.
+ * If not or if USE_STABLE_LOCALIZATION_KEYS is disabled, reset it and fix up any text in the package.
+ *
+ * @param InPackage			The package to reset the namespace for.
+ *
+ * @return True if the reset occurred. False if the package namespace is already the deterministic one.
+ */
+COREUOBJECT_API bool ResetPackageNamespace(UPackage* InPackage);
+
+/**
+ * Given an object, make sure it uses its deterministic package namespace (from its owner package) if USE_STABLE_LOCALIZATION_KEYS is enabled.
+ * If not or if USE_STABLE_LOCALIZATION_KEYS is disabled, reset it and fix up any text in the package.
+ *
+ * @param InObject			The object to clear the namespace for.
+ *
+ * @return True if the reset occurred. False if the package namespace is already the deterministic one..
+ */
+COREUOBJECT_API bool ResetPackageNamespace(UObject* InObject);
+
 /**
  * Given a package, try and get the namespace it should use for localization.
  *
