Log an error if the class import doesn't have an outer (i.e reported as a package) when calling GetExportClassPackage

#Synced-CL 51651818

diff --git a/Engine/Source/Runtime/CoreUObject/Private/UObject/LinkerLoad.cpp b/Engine/Source/Runtime/CoreUObject/Private/UObject/LinkerLoad.cpp
--- a/Engine/Source/Runtime/CoreUObject/Private/UObject/LinkerLoad.cpp
+++ b/Engine/Source/Runtime/CoreUObject/Private/UObject/LinkerLoad.cpp
@@ -3440,7 +3440,13 @@ FName FLinkerLoad::GetExportClassPackage( int32 i )
 	if( Export.ClassIndex.IsImport() )
 	{
 		FObjectImport& Import = Imp(Export.ClassIndex);
-		return ImpExp(Import.OuterIndex).ObjectName;
+		if (!Import.OuterIndex.IsNull())
+		{
+			return ImpExp(Import.OuterIndex).ObjectName;
+		}
+		// Report an error if the report class index doesn't have an outer
+		UE_LOGF(LogLinker, Error, "Export '%s' Class reported as a package '%s' which might indicate corruption", *Export.ObjectName.ToUtf8String(), *Import.ObjectName.ToUtf8String());
+		return Import.ObjectName;
 	}
 	else if ( !Export.ClassIndex.IsNull() )
 	{
