[AutoRTFM] Mark the garbage collection logic as AutoRTFM-disabled.

Garbage collection should never happen within a transaction, so it's unnecessary to compile a closed form version of the garbage collector.

#was-backout 51700208
#Synced-CL 51590581

diff --git a/Engine/Source/Runtime/CoreUObject/Private/UObject/GarbageCollection.cpp b/Engine/Source/Runtime/CoreUObject/Private/UObject/GarbageCollection.cpp
--- a/Engine/Source/Runtime/CoreUObject/Private/UObject/GarbageCollection.cpp
+++ b/Engine/Source/Runtime/CoreUObject/Private/UObject/GarbageCollection.cpp
@@ -113,7 +113,7 @@ static bool GObjCurrentPurgeObjectIndexNeedsReset = true;
 
 /** Contains a list of objects that stayed marked as unreachable after the last reachability analysis */
 static TArray<UE::GC::FUnreachableObject> GUnreachableObjects;
-static int32 GUnrechableObjectIndex = 0;
+static int32 GUnreachableObjectIndex = 0;
 
 static UE::GC::Private::FStats GGCStats;
 
@@ -719,7 +719,7 @@ namespace UE::GC::Private
 	}
 }
 
-static bool GatherUnreachableObjects(UE::GC::EGatherOptions Options, double TimeLimit = 0.0);
+AUTORTFM_DISABLE static bool GatherUnreachableObjects(UE::GC::EGatherOptions Options, double TimeLimit = 0.0);
 
 void FUObjectItem::MarkRootAsDirty()
 {
@@ -799,7 +799,7 @@ namespace UE::GC::Test
 } // namespace UE::GC::Test
 // END Hidden functions
 
-void OnDisregardForGCSetDisabled(int32 NumObjects)
+AUTORTFM_DISABLE void OnDisregardForGCSetDisabled(int32 NumObjects)
 {
 	using namespace UE::GC;
 	using namespace UE::GC::Private;
@@ -3167,7 +3167,7 @@ public:
 
 /** Batches up references before dispatching them to the processor, unlike TDirectDispatcher */
 template <class ProcessorType>
-struct TBatchDispatcher
+struct AUTORTFM_DISABLE TBatchDispatcher
 {
 	static constexpr bool bBatching = true;
 	static constexpr bool bParallel = IsParallel(ProcessorType::Options);
@@ -3514,7 +3514,7 @@ private:
 //////////////////////////////////////////////////////////////////////////
 
 template <EGCOptions Options>
-class TReachabilityCollectorBase : public FReferenceCollector
+class AUTORTFM_DISABLE TReachabilityCollectorBase : public FReferenceCollector
 {
 protected:
 	bool bAllowEliminatingReferences = true;
@@ -3542,7 +3542,7 @@ public:
 };
 
 template <EGCOptions Options>
-class TReachabilityCollector final : public TReachabilityCollectorBase<Options>
+class AUTORTFM_DISABLE TReachabilityCollector final : public TReachabilityCollectorBase<Options>
 {
 	using ProcessorType = TReachabilityProcessor<Options>;
 	using DispatcherType = TBatchDispatcher<ProcessorType>;
@@ -3749,7 +3749,7 @@ private:
 };
 
 template <EGCOptions Options>
-class TDebugReachabilityCollector final : public TReachabilityCollectorBase<Options>
+class AUTORTFM_DISABLE TDebugReachabilityCollector final : public TReachabilityCollectorBase<Options>
 {
 	using Super = TReachabilityCollectorBase<Options>;
 	using Super::bAllowEliminatingReferences;
@@ -4145,7 +4145,7 @@ static FORCEINLINE void StopVerseGC()
 namespace UE::GC
 {
 
-class FRealtimeGC : public FGarbageCollectionTracer
+class AUTORTFM_DISABLE FRealtimeGC : public FGarbageCollectionTracer
 {
 	typedef void(FRealtimeGC::*ReachabilityAnalysisFn)(FWorkerContext&);
 
@@ -4215,7 +4215,7 @@ class FRealtimeGC : public FGarbageCollectionTracer
 	}
 
 	template <EGCOptions Options>
-	void PerformReachabilityAnalysisOnObjectsInternal(FWorkerContext& Context)
+	AUTORTFM_DISABLE void PerformReachabilityAnalysisOnObjectsInternal(FWorkerContext& Context)
 	{
 		TRACE_CPUPROFILER_EVENT_SCOPE(PerformReachabilityAnalysisOnObjectsInternal);
 
@@ -4781,6 +4781,15 @@ void IncrementalPurgeGarbage(bool bUseTimeLimit, double TimeLimit)
 		return;
 	}
 
+	if (AutoRTFM::IsTransactional())
+	{
+		// Memory cannot be freed within a transaction as this would prevent us from rolling back to the initial state.
+		UE_LOG(LogGarbage, Log, TEXT("IncrementalPurgeGarbage: skipping garbage collection because an AutoRTFM transaction is active."));
+		return;
+	}
+
+	AutoRTFM::UnreachableIfTransactional();
+
 	SCOPED_NAMED_EVENT(IncrementalPurgeGarbage, FColor::Red);
 	DECLARE_SCOPE_CYCLE_COUNTER(TEXT("IncrementalPurgeGarbage"), STAT_IncrementalPurgeGarbage, STATGROUP_GC);
 	CSV_SCOPED_TIMING_STAT_EXCLUSIVE(GarbageCollection);
@@ -4825,7 +4834,7 @@ void IncrementalPurgeGarbage(bool bUseTimeLimit, double TimeLimit)
 		{
 			bTimeLimitReached = UnhashUnreachableObjects(bUseTimeLimit, TimeLimit);
 
-			if (GUnrechableObjectIndex >= GUnreachableObjects.Num())
+			if (GUnreachableObjectIndex >= GUnreachableObjects.Num())
 			{
 				FScopedCBDProfile::DumpProfile();
 			}
@@ -4923,7 +4932,7 @@ bool IncrementalDestroyGarbage(bool bUseTimeLimit, double TimeLimit)
 
 	if( !GObjFinishDestroyHasBeenRoutedToAllObjects && !bTimeLimitReached )
 	{
-		check(GUnrechableObjectIndex >= GUnreachableObjects.Num());
+		check(GUnreachableObjectIndex >= GUnreachableObjects.Num());
 
 		// Try to dispatch all FinishDestroy messages to unreachable objects.  We'll iterate over every
 		// single object and destroy any that are ready to be destroyed.  The objects that aren't yet
@@ -5263,7 +5272,7 @@ static FAutoConsoleVariableRef CVarFlushStreamingOnGC(
 	ECVF_Default
 	);
 
-void DissolveUnreachableClusters(UE::GC::EGatherOptions Options)
+AUTORTFM_DISABLE static void DissolveUnreachableClusters(UE::GC::EGatherOptions Options)
 {
 	const int32 InitialNumClusters = GUObjectClusters.GetNumAllocatedClusters();
 	if (!InitialNumClusters)
@@ -5356,7 +5365,7 @@ bool GatherUnreachableObjects(UE::GC::EGatherOptions Options, double TimeLimit /
 	const double GatherStartTime = FPlatformTime::Seconds();
 
 	GUnreachableObjects.Reset();
-	GUnrechableObjectIndex = 0;
+	GUnreachableObjectIndex = 0;
 
 	FGatherUnreachableObjectsState::FThreadIterators& ThreadIterators = GGatherUnreachableObjectsState.GetThreadIterators();
 	if (!GGatherUnreachableObjectsState.IsSuspended())
@@ -5468,7 +5477,7 @@ void GatherUnreachableObjects(bool bForceSingleThreaded)
 namespace UE::GC
 {
 
-class FWeakReferenceEliminator final : public TReachabilityCollectorBase<EGCOptions::None>
+class AUTORTFM_DISABLE FWeakReferenceEliminator final : public TReachabilityCollectorBase<EGCOptions::None>
 {
 
 public:
@@ -5625,9 +5634,9 @@ static void UpdateGCHistory(TConstArrayView<TUniquePtr<FWorkerContext>> Contexts
 }
 
 template<bool bPerformFullPurge>
-FORCENOINLINE void CollectGarbageImpl(EObjectFlags KeepFlags);
+AUTORTFM_DISABLE FORCENOINLINE void CollectGarbageImpl(EObjectFlags KeepFlags);
 
-FORCENOINLINE static void CollectGarbageIncremental(EObjectFlags KeepFlags)
+AUTORTFM_DISABLE FORCENOINLINE static void CollectGarbageIncremental(EObjectFlags KeepFlags)
 {
 	SCOPE_TIME_GUARD(TEXT("Collect Garbage Incremental"));
 	SCOPED_NAMED_EVENT(CollectGarbageIncremental, FColor::Red);
@@ -5637,7 +5646,7 @@ FORCENOINLINE static void CollectGarbageIncremental(EObjectFlags KeepFlags)
 	CollectGarbageImpl<false>(KeepFlags);
 }
 
-FORCENOINLINE static void CollectGarbageFull(EObjectFlags KeepFlags)
+AUTORTFM_DISABLE FORCENOINLINE static void CollectGarbageFull(EObjectFlags KeepFlags)
 {
 	SCOPE_TIME_GUARD(TEXT("Collect Garbage Full"));
 	SCOPED_NAMED_EVENT(CollectGarbageFull, FColor::Red);
@@ -5647,7 +5656,7 @@ FORCENOINLINE static void CollectGarbageFull(EObjectFlags KeepFlags)
 	CollectGarbageImpl<true>(KeepFlags);
 }
 
-FORCEINLINE void CollectGarbageInternal(EObjectFlags KeepFlags, bool bPerformFullPurge)
+AUTORTFM_DISABLE FORCEINLINE void CollectGarbageInternal(EObjectFlags KeepFlags, bool bPerformFullPurge)
 {
 	const double StartTime = FPlatformTime::Seconds();
 
@@ -5682,7 +5691,7 @@ EGatherOptions GetObjectGatherOptions()
  * @param	bPerformFullPurge	if true, perform a full purge after the mark pass
  */
 template<bool bPerformFullPurge>
-void PreCollectGarbageImpl(EObjectFlags KeepFlags)
+AUTORTFM_DISABLE void PreCollectGarbageImpl(EObjectFlags KeepFlags)
 {
 	using namespace UE::GC::Private;
 
@@ -5806,7 +5815,7 @@ void PreCollectGarbageImpl(EObjectFlags KeepFlags)
 }
 
 template<bool bPerformFullPurge>
-void CollectGarbageImpl(EObjectFlags KeepFlags)
+AUTORTFM_DISABLE void CollectGarbageImpl(EObjectFlags KeepFlags)
 {
 	{
 		// Reachability analysis.
@@ -5821,7 +5830,7 @@ void CollectGarbageImpl(EObjectFlags KeepFlags)
 }
 
 template<bool bPerformFullPurge>
-void PostCollectGarbageImpl(EObjectFlags KeepFlags)
+AUTORTFM_DISABLE void PostCollectGarbageImpl(EObjectFlags KeepFlags)
 {
 	const double PostCollectStartTime = FPlatformTime::Seconds();
 
@@ -6165,6 +6174,15 @@ bool IsIncrementalReachabilityAnalysisPending()
 
 void PerformIncrementalReachabilityAnalysis(double TimeLimit)
 {
+	if (AutoRTFM::IsTransactional())
+	{
+		// The garbage collector is not usable from within a transaction.
+		UE_LOG(LogGarbage, Log, TEXT("PerformIncrementalReachabilityAnalysis: skipping garbage collection because an AutoRTFM transaction is active."));
+		return;
+	}
+
+	AutoRTFM::UnreachableIfTransactional();
+
 	checkf(UE::GC::GIsIncrementalReachabilityPending, TEXT("Incremental reachability must be pending to perform its next iteration"));
 	// When performing Reachability Analysis iterations start the internal timer before acquiring GC lock 
 	// so that we don't spend more time than GIncrementalReachabilityTimeLimit on fully completing an iteration
@@ -6221,7 +6239,7 @@ double GetLastGCDuration()
 
 bool IsIncrementalUnhashPending()
 {
-	return GUnrechableObjectIndex < GUnreachableObjects.Num() || UE::GC::Private::GGatherUnreachableObjectsState.IsPending();
+	return GUnreachableObjectIndex < GUnreachableObjects.Num() || UE::GC::Private::GGatherUnreachableObjectsState.IsPending();
 }
 
 bool UnhashUnreachableObjects(bool bUseTimeLimit, double TimeLimit)
@@ -6231,6 +6249,15 @@ bool UnhashUnreachableObjects(bool bUseTimeLimit, double TimeLimit)
 
 	bool bTimeLimitReached = false;
 
+	if (AutoRTFM::IsTransactional())
+	{
+		// Memory cannot be freed within a transaction as this would prevent us from rolling back to the initial state.
+		UE_LOG(LogGarbage, Log, TEXT("UnhashUnreachableObjects: skipping garbage collection because an AutoRTFM transaction is active."));
+		return bTimeLimitReached;
+	}
+
+	AutoRTFM::UnreachableIfTransactional();
+
 	if (GGatherUnreachableObjectsState.IsPending())
 	{
 		// Incremental Gather needs to be called from UnhashUnreachableObjects to match changes in IsIncrementalUnhashPending() (and not introduce IsIncrementalGatherPending())
@@ -6265,15 +6292,15 @@ bool UnhashUnreachableObjects(bool bUseTimeLimit, double TimeLimit)
 	double LastPollTime = 0.0;
 	const int32 TimeLimitEnforcementGranularityForBeginDestroy = 10;
 	int32 TimePollCounter = 0;
-	const bool bFirstIteration = (GUnrechableObjectIndex == 0);
+	const bool bFirstIteration = (GUnreachableObjectIndex == 0);
 
 	{
 		TRACE_CPUPROFILER_EVENT_SCOPE(ConditionalBeginDestroy);
-		while (GUnrechableObjectIndex < GUnreachableObjects.Num())
+		while (GUnreachableObjectIndex < GUnreachableObjects.Num())
 		{
 			//@todo UE - A prefetch was removed here. Re-add it. It wasn't right anyway, since it was ten items ahead and the consoles on have 8 prefetch slots
 
-			FUObjectItem* ObjectItem = GUnreachableObjects[GUnrechableObjectIndex++].ObjectItem;
+			FUObjectItem* ObjectItem = GUnreachableObjects[GUnreachableObjectIndex++].ObjectItem;
 			{
 				UObject* Object = static_cast<UObject*>(ObjectItem->GetObject());
 				FScopedCBDProfile Profile(Object);
@@ -6293,7 +6320,7 @@ bool UnhashUnreachableObjects(bool bUseTimeLimit, double TimeLimit)
 		}
 	}
 
-	bTimeLimitReached = (GUnrechableObjectIndex < GUnreachableObjects.Num());
+	bTimeLimitReached = (GUnreachableObjectIndex < GUnreachableObjects.Num());
 
 	const double ElapsedTime = (LastPollTime == 0.0 ? FPlatformTime::Seconds() : LastPollTime) - StartTime;
 	GGCStats.UnhashingTime += ElapsedTime;
@@ -6335,6 +6362,15 @@ void CollectGarbage(EObjectFlags KeepFlags, bool bPerformFullPurge)
 		return;
 	}
 
+	if (AutoRTFM::IsTransactional())
+	{
+		// Memory cannot be freed within a transaction as this would prevent us from rolling back to the initial state.
+		UE_LOG(LogGarbage, Log, TEXT("TryCollectGarbage: skipping garbage collection because an AutoRTFM transaction is active."));
+		return;
+	}
+
+	AutoRTFM::UnreachableIfTransactional();
+
 	// No other thread may be performing UObject operations while we're running
 	AcquireGCLock();
 
@@ -6360,6 +6396,8 @@ bool TryCollectGarbage(EObjectFlags KeepFlags, bool bPerformFullPurge)
 		return false;
 	}
 
+	AutoRTFM::UnreachableIfTransactional();
+
 	// No other thread may be performing UObject operations while we're running so try to acquire GC lock
 	if (UE::GC::GIsIncrementalReachabilityPending)
 	{
@@ -6393,7 +6431,7 @@ bool TryCollectGarbage(EObjectFlags KeepFlags, bool bPerformFullPurge)
 /**
  * Destroys all existing UObjects when the engine shuts down. Does NOT perform reachability analysis to do so.
  */
-void PurgeAllUObjectsOnExit()
+AUTORTFM_DISABLE void PurgeAllUObjectsOnExit()
 {
 	using namespace UE::GC::Private;
 
@@ -7265,7 +7303,7 @@ ELoot StealWork(FWorkerContext& Context, FReferenceCollector& Collector, FWorkBl
 	return ELoot::Nothing;
 }
 
-TArrayView<FWorkerContext*> InitializeAsyncProcessingContexts(FWorkerContext& InContext)
+AUTORTFM_DISABLE static TArrayView<FWorkerContext*> InitializeAsyncProcessingContexts(FWorkerContext& InContext)
 {
 	TArrayView<FWorkerContext*> Contexts;
 
diff --git a/Engine/Source/Runtime/CoreUObject/Private/UObject/Obj.cpp b/Engine/Source/Runtime/CoreUObject/Private/UObject/Obj.cpp
--- a/Engine/Source/Runtime/CoreUObject/Private/UObject/Obj.cpp
+++ b/Engine/Source/Runtime/CoreUObject/Private/UObject/Obj.cpp
@@ -5866,7 +5866,7 @@ bool StaticExec( UWorld* InWorld, const TCHAR* Cmd, FOutputDevice& Ar )
 
 void StaticUObjectInit();
 void InitUObject();
-void StaticExit();
+AUTORTFM_DISABLE void StaticExit();
 
 void InitUObject()
 {
@@ -5962,7 +5962,7 @@ void StaticUObjectInit()
 void ShutdownGarbageCollection();
 void CleanupLinkerAnnotations();
 void CleanupCachedArchetypes();
-void PurgeAllUObjectsOnExit();
+AUTORTFM_DISABLE void PurgeAllUObjectsOnExit();
 
 //
 // Shut down the object manager.
diff --git a/Engine/Source/Runtime/CoreUObject/Private/UObject/ReachabilityAnalysisState.h b/Engine/Source/Runtime/CoreUObject/Private/UObject/ReachabilityAnalysisState.h
--- a/Engine/Source/Runtime/CoreUObject/Private/UObject/ReachabilityAnalysisState.h
+++ b/Engine/Source/Runtime/CoreUObject/Private/UObject/ReachabilityAnalysisState.h
@@ -18,7 +18,7 @@ namespace UE::GC
 /**
 * Reachability analysis state holds information about suspended worker contexts.
 */
-class FReachabilityAnalysisState
+class AUTORTFM_DISABLE FReachabilityAnalysisState
 {
 public:
 
@@ -154,7 +154,7 @@ inline int32 GetNumThreadsForGather(const EGatherOptions Options, const int32 Nu
 template <typename PayloadType>
 struct TGatherIterator
 {
-	/** Inital start index */
+	/** Initial start index */
 	int32 StartIndex = 0;
 	/** Current index of an object */
 	int32 Index = 0;
@@ -193,7 +193,7 @@ struct TDefaultPayloadOps
 
 /** Helper class that holds the current state of Object Gathering phases */
 template <typename PayloadType, typename PayloadOps = TDefaultPayloadOps<PayloadType>>
-class TThreadedGather
+class AUTORTFM_DISABLE TThreadedGather
 {
 public:
 
@@ -210,7 +210,7 @@ public:
 		return ThreadIterators.Num() > 0;
 	}
 
-	FORCEINLINE bool IsPending() const
+	AUTORTFM_ENABLE FORCEINLINE bool IsPending() const
 	{
 		return bIsPending;
 	}
diff --git a/Engine/Source/Runtime/CoreUObject/Private/UObject/UObjectArray.cpp b/Engine/Source/Runtime/CoreUObject/Private/UObject/UObjectArray.cpp
--- a/Engine/Source/Runtime/CoreUObject/Private/UObject/UObjectArray.cpp
+++ b/Engine/Source/Runtime/CoreUObject/Private/UObject/UObjectArray.cpp
@@ -221,8 +221,11 @@ void FUObjectArray::DisableDisregardForGC()
 	if (!GExitPurge && NumDisregardForGCObjects > 0)
 	{
 		checkf(!IsAsyncLoading(), TEXT("Disregard for GC Set can't be safely disabled while async loading. Consider calling FlushAsyncLoading() first or using gc.MaxObjectsNotConsideredByGC=0 ini setting instead."));
+		checkf(!AutoRTFM::IsTransactional(), TEXT("Disregard for GC can't be changed while inside of a transaction."));
 
-		void OnDisregardForGCSetDisabled(int32 NumObjects);
+		AutoRTFM::UnreachableIfTransactional();
+
+		AUTORTFM_DISABLE void OnDisregardForGCSetDisabled(int32 NumObjects);
 		OnDisregardForGCSetDisabled(NumDisregardForGCObjects);
 	}
 }
diff --git a/Engine/Source/Runtime/CoreUObject/Public/UObject/FastReferenceCollector.h b/Engine/Source/Runtime/CoreUObject/Public/UObject/FastReferenceCollector.h
--- a/Engine/Source/Runtime/CoreUObject/Public/UObject/FastReferenceCollector.h
+++ b/Engine/Source/Runtime/CoreUObject/Public/UObject/FastReferenceCollector.h
@@ -883,12 +883,12 @@ struct TGetDispatcherType
 //////////////////////////////////////////////////////////////////////////
 
 enum class ELoot { Nothing, Block, ARO, Context };
-COREUOBJECT_API ELoot StealWork(FWorkerContext& Context, FReferenceCollector& Collector, FWorkBlock*& OutBlock, EGCOptions Options);
+AUTORTFM_DISABLE COREUOBJECT_API ELoot StealWork(FWorkerContext& Context, FReferenceCollector& Collector, FWorkBlock*& OutBlock, EGCOptions Options);
 
 COREUOBJECT_API void SuspendWork(FWorkerContext& Context);
 
 /** Allocates contexts and coordinator, kicks worker tasks that also call ProcessSync. Processor is type-erased to void* to avoid templated code. */
-COREUOBJECT_API void ProcessAsync(void (*ProcessSync)(void*, FWorkerContext&), void* Processor, FWorkerContext& InitialContext);
+AUTORTFM_DISABLE COREUOBJECT_API void ProcessAsync(void (*ProcessSync)(void*, FWorkerContext&), void* Processor, FWorkerContext& InitialContext);
 
 //////////////////////////////////////////////////////////////////////////
 
diff --git a/Engine/Source/Runtime/CoreUObject/Public/UObject/UObjectGlobals.h b/Engine/Source/Runtime/CoreUObject/Public/UObject/UObjectGlobals.h
--- a/Engine/Source/Runtime/CoreUObject/Public/UObject/UObjectGlobals.h
+++ b/Engine/Source/Runtime/CoreUObject/Public/UObject/UObjectGlobals.h
@@ -989,7 +989,7 @@ COREUOBJECT_API bool IsIncrementalPurgePending();
  *
  * @param bForceSingleThreaded true to force the process to just one thread
  */
-COREUOBJECT_API void GatherUnreachableObjects(bool bForceSingleThreaded);
+AUTORTFM_DISABLE COREUOBJECT_API void GatherUnreachableObjects(bool bForceSingleThreaded);
 
 /**
  * Returns whether an incremental reachability analysis is still pending/ in progress.
