Protect against possible invalid offset values when trying to load a payload directly from a package trailer.

#rnx

-It might be possible that the offset will be invalid if the package trailer is corrupted on disk which might cause a seek outside the range of the archive.
- With the previous code this would cause either an assert or be entirely ignored if assets are disabled.
- We can use FGuardedInt64 to check if the corrupted values cause an integer overflow and check if the final offset is within the archive, if not then we can log a warning and return a null FCompressedBuffer.
-- It would then be up to the calling code to handle the failure to load data, which is should be doing anyway.

#Synced-CL 51587407

diff --git a/Engine/Source/Runtime/CoreUObject/Private/UObject/PackageTrailer.cpp b/Engine/Source/Runtime/CoreUObject/Private/UObject/PackageTrailer.cpp
--- a/Engine/Source/Runtime/CoreUObject/Private/UObject/PackageTrailer.cpp
+++ b/Engine/Source/Runtime/CoreUObject/Private/UObject/PackageTrailer.cpp
@@ -675,8 +675,18 @@ FCompressedBuffer FPackageTrailer::LoadLocalPayload(const FIoHash& Id, FArchive&
 		return FCompressedBuffer();
 	}
 
-	const int64 PayloadOffsetInFile = PackageTrailerOffset + Header.HeaderLength + Entry->OffsetInFile;
-	Ar.Seek(PayloadOffsetInFile);
+	const FGuardedInt64 PayloadOffsetInFile = FGuardedInt64(PackageTrailerOffset) + Header.HeaderLength + Entry->OffsetInFile;
+	if (PayloadOffsetInFile.InvalidOrLessThan(0) || PayloadOffsetInFile.InvalidOrGreaterOrEqual(Ar.TotalSize()))
+	{
+		UE_LOG(LogSerialization, Warning, TEXT("Unable to load payload '%s' from disk, it's offset is outside of the valid range (0-%lld) from archive '%s'"),
+			*LexToString(Id),
+			Ar.Tell(),
+			*Ar.GetArchiveName());
+
+		return FCompressedBuffer();
+	}
+
+	Ar.Seek(PayloadOffsetInFile.Get(0));
 
 	FCompressedBuffer Payload;
 	Ar << Payload;
