Adding the existing migrated data state to the assert when the migrated data already exists but we received another copy

#Synced-CL 51761520

diff --git a/Engine/Source/Runtime/CoreUObject/Private/UObject/RemoteObjectTransfer.cpp b/Engine/Source/Runtime/CoreUObject/Private/UObject/RemoteObjectTransfer.cpp
--- a/Engine/Source/Runtime/CoreUObject/Private/UObject/RemoteObjectTransfer.cpp
+++ b/Engine/Source/Runtime/CoreUObject/Private/UObject/RemoteObjectTransfer.cpp
@@ -972,7 +972,7 @@ TRefCountPtr<FTransactionalMigrationData> StoreMigratedObjectData(FRemoteObjectD
 		FRemoteObjectId SerializedId = It.GetId();
 		FRemoteObjectStub* Stub = FindOrAddRemoteObjectStub(SerializedId, FRemoteServerId::GetLocalServerId());
 
-		checkf(!Stub->MigratedData, TEXT("Received remote object %s (%s) from server %s bundled with %s (%s) but the latter object's data already exists locally (received from server %s and bundled with %s %s)"), 
+		checkf(!Stub->MigratedData, TEXT("Received remote object %s (%s) from server %s bundled with %s (%s) but the latter object's data already exists locally (received from server %s and bundled with %s %s, state: %s)"), 
 			*MigratedData->Context.ObjectId.ToString(), 
 			*MigratedData->Data.GetName(MigratedData->Context.ObjectId).ToString(),
 			*MigratedData->Context.RemoteServerId.ToString(),
@@ -980,7 +980,8 @@ TRefCountPtr<FTransactionalMigrationData> StoreMigratedObjectData(FRemoteObjectD
 			*MigratedData->Data.GetName(SerializedId).ToString(),
 			*Stub->MigratedData->Context.RemoteServerId.ToString(),
 			*Stub->MigratedData->Context.ObjectId.ToString(),
-			*Stub->MigratedData->Data.GetName(Stub->MigratedData->Context.ObjectId).ToString());
+			*Stub->MigratedData->Data.GetName(Stub->MigratedData->Context.ObjectId).ToString(),
+			EnumToString(Stub->MigratedData->State));
 
 		Stub->OwningServerId = MigratedData->Context.OwnerServerId;
 		Stub->MigratedData = MigratedData;
