[AutoRTFM] Mark the DebugReachabilityProcessors as AutoRTFM-disabled.

This CL is just a small followup to https://p4-swarm.epicgames.net/changes/51732300 which missed these classes.

#Synced-CL 51590581

diff --git a/Engine/Source/Runtime/CoreUObject/Private/UObject/GarbageCollection.cpp b/Engine/Source/Runtime/CoreUObject/Private/UObject/GarbageCollection.cpp
--- a/Engine/Source/Runtime/CoreUObject/Private/UObject/GarbageCollection.cpp
+++ b/Engine/Source/Runtime/CoreUObject/Private/UObject/GarbageCollection.cpp
@@ -3438,7 +3438,7 @@ FORCEINLINE static FName GetReferencerName(const UObject* Referencer, FMemberId
 
 #if !UE_BUILD_SHIPPING && (WITH_VERSE_VM || defined(__INTELLISENSE__))
 template <EGCOptions Options>
-struct TVerseDebugReachabilityVisitor : public Verse::FAbstractVisitor
+struct AUTORTFM_DISABLE TVerseDebugReachabilityVisitor : public Verse::FAbstractVisitor
 {
 	UE_NONCOPYABLE(TVerseDebugReachabilityVisitor);
 
@@ -3607,7 +3607,7 @@ public:
 #endif
 
 template <EGCOptions InOptions>
-class TDebugReachabilityProcessor 
+class AUTORTFM_DISABLE TDebugReachabilityProcessor
 {
 public:
 	FORCEINLINE void BeginTimingObject(UObject* CurrentObject)
