[Backout] - CL51791046, just ScriptCore and KismetCompilerVMBackend due to faulty literal logic. Not backing out whole CL due to native version change
#rnx
Original CL Desc
-----------------------------------------------------------------
Localization: Add a "Dev Notes" field to localized text so developers can provide additional context for translators

This is exposed as an extra text field in the editor and with new macros in C++ (NOTELOCTEXT, NSNOTELOCTEXT)

PR #11531

diff --git a/Engine/Source/Editor/KismetCompiler/Private/KismetCompilerVMBackend.cpp b/Engine/Source/Editor/KismetCompiler/Private/KismetCompilerVMBackend.cpp
--- a/Engine/Source/Editor/KismetCompiler/Private/KismetCompilerVMBackend.cpp
+++ b/Engine/Source/Editor/KismetCompiler/Private/KismetCompilerVMBackend.cpp
@@ -730,12 +730,6 @@ public:
 						EmitStringLiteral(*SourceString);
 						EmitStringLiteral(Key);
 						EmitStringLiteral(Namespace);
-#if WITH_EDITORONLY_DATA
-						if (!Writer.IsFilterEditorOnly() && !Writer.IsCooking() && !IsRunningCookCommandlet())
-						{
-							EmitStringLiteral(Term->TextLiteral.GetDevNotes());
-						}
-#endif
 					}
 					else
 					{
diff --git a/Engine/Source/Runtime/CoreUObject/Private/UObject/ScriptCore.cpp b/Engine/Source/Runtime/CoreUObject/Private/UObject/ScriptCore.cpp
--- a/Engine/Source/Runtime/CoreUObject/Private/UObject/ScriptCore.cpp
+++ b/Engine/Source/Runtime/CoreUObject/Private/UObject/ScriptCore.cpp
@@ -3481,19 +3481,8 @@ DEFINE_FUNCTION(UObject::execTextConst)
 			
 			FString Namespace;
 			Stack.Step(Stack.Object, &Namespace);
-#if WITH_EDITORONLY_DATA
-			if (Stack.Object && Stack.Object->GetPackage() && !Stack.Object->GetPackage()->HasAnyPackageFlags(PKG_Cooked))
-			{
-				FString DevNotes;
-				Stack.Step(Stack.Object, &DevNotes);
 
-				*(FText*)RESULT_PARAM = FText::AsLocalizable_Advanced(Namespace, KeyString, MoveTemp(SourceString), DevNotes);
-			}
-			else
-#endif
-			{
-				*(FText*)RESULT_PARAM = FText::AsLocalizable_Advanced(Namespace, KeyString, MoveTemp(SourceString));
-			}
+			*(FText*)RESULT_PARAM = FText::AsLocalizable_Advanced(Namespace, KeyString, MoveTemp(SourceString));
 		}
 		break;
 
