[Backout] - CL51791046. Just part of PropertyLocalizationDataGathering.cpp related to text literals. The rest of the literal change was backed out in previous CLs
#rnx
Original CL Desc
-----------------------------------------------------------------
Localization: Add a "Dev Notes" field to localized text so developers can provide additional context for translators

This is exposed as an extra text field in the editor and with new macros in C++ (NOTELOCTEXT, NSNOTELOCTEXT)

PR #11531

diff --git a/Engine/Source/Runtime/CoreUObject/Private/Serialization/PropertyLocalizationDataGathering.cpp b/Engine/Source/Runtime/CoreUObject/Private/Serialization/PropertyLocalizationDataGathering.cpp
--- a/Engine/Source/Runtime/CoreUObject/Private/Serialization/PropertyLocalizationDataGathering.cpp
+++ b/Engine/Source/Runtime/CoreUObject/Private/Serialization/PropertyLocalizationDataGathering.cpp
@@ -755,23 +755,9 @@ private:
 				SerializeExpr(iCode, Ar);
 				const FString TextNamespace = MoveTemp(LastParsedString);
 
-#if WITH_EDITORONLY_DATA
-				FString DevNotes;
-				const bool bIsCooked = PropertyLocalizationDataGatherer.GetPackage()->HasAnyPackageFlags(PKG_Cooked);
-				if (!bIsCooked)
-				{
-					SerializeExpr(iCode, Ar);
-					DevNotes = MoveTemp(LastParsedString);
-				}
-#endif
-
 				bIsParsingText = false;
 
-#if WITH_EDITORONLY_DATA
-				const FText TextInstance = FText::AsLocalizable_Advanced(TextNamespace, TextKey, MoveTemp(SourceString), DevNotes);
-#else
 				const FText TextInstance = FText::AsLocalizable_Advanced(TextNamespace, TextKey, MoveTemp(SourceString));
-#endif
 				if (!PropertyLocalizationDataGatherer.IsDefaultTextInstance(TextInstance))
 				{
 					PropertyLocalizationDataGatherer.GatherTextInstance(TextInstance, FString::Printf(TEXT("%s [Script Bytecode]"), SourceDescription), bTreatAsEditorOnlyData);
