Switch from `TSortedMap` -> `TMap` with an occassional key sort (if required for deterministic hashes) to improve performance by 2.91x.

#Synced-CL 51702299

diff --git a/Engine/Source/Runtime/CoreUObject/Internal/UObject/CoreRedirects/RedirectionSummary.h b/Engine/Source/Runtime/CoreUObject/Internal/UObject/CoreRedirects/RedirectionSummary.h
--- a/Engine/Source/Runtime/CoreUObject/Internal/UObject/CoreRedirects/RedirectionSummary.h
+++ b/Engine/Source/Runtime/CoreUObject/Internal/UObject/CoreRedirects/RedirectionSummary.h
@@ -57,7 +57,7 @@ private:
 		void CalculateHash();
 		void AppendHashWithoutDirtyCheck(FBlake3& Hasher) const;
 	private:
-		TSortedMap<FCoreRedirect, bool, FDefaultAllocator, FCompareRedirect> Redirects;
+		TMap<FCoreRedirect, bool> Redirects;
 		FBlake3Hash Hash;
 		bool bHashDirty = false;
 	};
diff --git a/Engine/Source/Runtime/CoreUObject/Private/UObject/CoreRedirects.cpp b/Engine/Source/Runtime/CoreUObject/Private/UObject/CoreRedirects.cpp
--- a/Engine/Source/Runtime/CoreUObject/Private/UObject/CoreRedirects.cpp
+++ b/Engine/Source/Runtime/CoreUObject/Private/UObject/CoreRedirects.cpp
@@ -1580,6 +1580,34 @@ int FCoreRedirect::Compare(const FCoreRedirect& Other) const
 	return 0;
 }
 
+int32 FCoreRedirect::GetTypeHash() const
+{
+	uint32 Hashes[] = {
+		GetTypeHashHelper(RedirectFlags),
+		GetTypeHashHelper(OldName),
+		GetTypeHashHelper(NewName),
+		GetTypeHashHelper(OverrideClassName),
+	};
+
+	uint32 SmooshedHash = 0;
+
+	for (const uint32 Hash : Hashes)
+	{
+		SmooshedHash = HashCombineFast(SmooshedHash, Hash);
+	}
+
+	TArray<TPair<FString, FString>> ValueArray = ValueChanges.Array();
+	Algo::Sort(ValueArray);
+
+	for (const TPair<FString, FString>& Pair : ValueArray)
+	{
+		SmooshedHash = HashCombineFast(SmooshedHash, GetTypeHashHelper(Pair.Key));
+		SmooshedHash = HashCombineFast(SmooshedHash, GetTypeHashHelper(Pair.Value));
+	}
+
+	return SmooshedHash;
+}
+
 void FCoreRedirectObjectName::AppendHash(FBlake3& Hasher) const
 {
 	FNameBuilder NameStr;
diff --git a/Engine/Source/Runtime/CoreUObject/Private/UObject/RedirectionSummary.cpp b/Engine/Source/Runtime/CoreUObject/Private/UObject/RedirectionSummary.cpp
--- a/Engine/Source/Runtime/CoreUObject/Private/UObject/RedirectionSummary.cpp
+++ b/Engine/Source/Runtime/CoreUObject/Private/UObject/RedirectionSummary.cpp
@@ -305,6 +305,10 @@ void FRedirectionSummary::FRedirectContainer::DebugPrint() const
 
 void FRedirectionSummary::FRedirectContainer::CalculateHash()
 {
+	// For hashing we need determinism across multiple runs, so since our hash is dirty
+	// we have to sort the redirects on keys now to ensure we have a deterministic hash.
+	Redirects.KeyStableSort(TLess<>());
+
 	FBlake3 Hasher;
 	for (const TPair<FCoreRedirect,bool>& Pair: Redirects)
 	{
diff --git a/Engine/Source/Runtime/CoreUObject/Public/UObject/CoreRedirects.h b/Engine/Source/Runtime/CoreUObject/Public/UObject/CoreRedirects.h
--- a/Engine/Source/Runtime/CoreUObject/Public/UObject/CoreRedirects.h
+++ b/Engine/Source/Runtime/CoreUObject/Public/UObject/CoreRedirects.h
@@ -284,6 +284,23 @@ struct FCoreRedirect
 	/** Returns -,0,+ based on a full lexical-fnames compare of all fields on the two CoreRedirects. */
 	COREUOBJECT_API int Compare(const FCoreRedirect& Other) const;
 
+	COREUOBJECT_API int32 GetTypeHash() const;
+
+	friend int32 GetTypeHash(const FCoreRedirect& CoreRedirect)
+	{
+		return CoreRedirect.GetTypeHash();
+	}
+
+	bool operator==(const FCoreRedirect& Other) const
+	{
+		return Compare(Other) == 0;
+	}
+
+	bool operator<(const FCoreRedirect& Other) const
+	{
+		return Compare(Other) < 0;
+	}
+
 private:
 	friend struct FCoreRedirects;
 	friend class FRedirectionSummary;
