Add a new ensure to check if VirtualExportsFileOffset is getting out of sync with the archive file pos.

- In CL 51492546 we added checks for corrupted packages that seem to have the wrong offset in the package file summary to the package trailer, but we are still not sure why this is occurring.
- One potential cause of the problem is that the offset is derived from a virtual offset we build during the saving of a package.
- This check should catch if it is the case and make an error report so that we can track this.
- It seems unlikely that this is the problem so I am not making an attempt to 'fix' the issue here, just track it so we can prove/disprove the theory. The work in 51492546 should allow the package to be loaded anyway.
- The ensure can be disabled by setting the cvar 'save.ValidateVirtualExportsFileOffset' to false.

#Synced-CL 51587407

diff --git a/Engine/Source/Runtime/CoreUObject/Private/UObject/SavePackage2.cpp b/Engine/Source/Runtime/CoreUObject/Private/UObject/SavePackage2.cpp
--- a/Engine/Source/Runtime/CoreUObject/Private/UObject/SavePackage2.cpp
+++ b/Engine/Source/Runtime/CoreUObject/Private/UObject/SavePackage2.cpp
@@ -105,6 +105,12 @@ static FAutoConsoleVariableRef CVar_MemoryMappedBulkDataAlignment(
 	TEXT("Any other non-zero value - overrides the memory mapped data alignment, implies padding")
 );
 
+bool GValidateVirtualExportsFileOffset = true;
+static FAutoConsoleVariableRef CVarValidateVirtualExportsFileOffset(
+	TEXT("save.ValidateVirtualExportsFileOffset"),
+	GValidateVirtualExportsFileOffset,
+	TEXT("When true additional validation will be run when saving a package file."));
+
 ESavePackageResult WriteAdditionalFiles(FSaveContext& SaveContext, FScopedSlowTask& SlowTask, int64& VirtualExportsFileOffset);
 
 ESavePackageResult ReturnSuccessOrCancel()
@@ -3834,11 +3840,18 @@ AUTORTFM_DISABLE ESavePackageResult SaveHarvestedRealms(FSaveContext& SaveContex
 	{
 		return SaveContext.Result;
 	}
+
 	if (PackageWriter)
 	{
 		checkf(SaveContext.GetLinker()->Tell() == EndOfExportsOffset,
 			TEXT("The writing of additional files is not allowed to append to the LinkerSave when using a PackageWriter."));
 	}
+	else
+	{
+		ensureMsgf(!GValidateVirtualExportsFileOffset || SaveContext.GetLinker()->Tell() == VirtualExportsFileOffset,
+			TEXT("VirtualExportsFileOffset is no longer in sync (%lld vs %lld) when writing out '%s'"),
+			VirtualExportsFileOffset, SaveContext.GetLinker()->Tell(), * SaveContext.GetPackage()->GetName());
+	}
 
 	// Store the package header and export size of the non optional realm
 	if (SaveContext.GetCurrentHarvestingRealm() != ESaveRealm::Optional)
