Added accessor for the root overridable node of a property set.

#rnx

#Synced-CL 51585801

diff --git a/Engine/Source/Runtime/CoreUObject/Public/UObject/OverriddenPropertySet.h b/Engine/Source/Runtime/CoreUObject/Public/UObject/OverriddenPropertySet.h
--- a/Engine/Source/Runtime/CoreUObject/Public/UObject/OverriddenPropertySet.h
+++ b/Engine/Source/Runtime/CoreUObject/Public/UObject/OverriddenPropertySet.h
@@ -264,7 +264,7 @@ struct FOverriddenPropertyPath
 		return Path.IsEmpty();
 	}
 
-	FString ToString() const;
+	COREUOBJECT_API FString ToString() const;
 
 private:
 	UPROPERTY()
@@ -536,6 +536,22 @@ public:
 	UE_API const FOverriddenPropertyNode* GetOverriddenPropertyNode(const FArchiveSerializedPropertyChain* CurrentPropertyChain) const;
 	UE_API const FOverriddenPropertyNode* GetOverriddenPropertyNode(const FOverriddenPropertyNodeID& NodeID) const;
 
+	/**
+	 * Retrieve the root overridden property node (read-only variant).
+	 * @return the root node of the property set */
+	const FOverriddenPropertyNode& GetRootOverriddenPropertyNode() const 
+	{ 
+		return RootNode; 
+	}
+
+	/**
+	 * Retrieve the root overridden property node (mutable variant).
+	 * @return the root node of the property set */
+	FOverriddenPropertyNode& GetRootOverriddenPropertyNode() 
+	{ 
+		return RootNode; 
+	}
+
 	/**
 	 * Retrieve the overridable operation given the property key
 	 * @param NodeID that uniquely identify the property within the object 
