[CoreUObject] Add MoveAssignScriptStruct to UScriptStruct for move-semantic support in type-erased struct operations

Add a new WithMoveAssign trait to TStructOpsTypeTraitsBase2, following the
exact pattern of the existing WithCopy trait:

1. New trait: TStructOpsTypeTraitsBase2::WithMoveAssign (default false)
2. New struct flag: STRUCT_MoveAssignNative
3. New capability: FCapabilities::HasMoveAssign
4. New fake vtable entry: MoveAssign free function template
5. New ICppStructOps methods: HasMoveAssign() / MoveAssign()
6. New UScriptStruct method: MoveAssignScriptStruct()
   - STRUCT_MoveAssignNative: dispatches to C++ operator=(T&&) via vtable
   - STRUCT_IsPlainOldData: memcpy (same as copy, move is identical)
   - Fallback: calls CopyScriptStruct (zero regression for non-opted types)

#Synced-CL 51590771

diff --git a/Engine/Source/Runtime/CoreUObject/Private/UObject/Class.cpp b/Engine/Source/Runtime/CoreUObject/Private/UObject/Class.cpp
--- a/Engine/Source/Runtime/CoreUObject/Private/UObject/Class.cpp
+++ b/Engine/Source/Runtime/CoreUObject/Private/UObject/Class.cpp
@@ -3210,6 +3210,17 @@ void UScriptStruct::PrepareCppStructOps()
 			UE_LOG(LogClass, Verbose, TEXT("Native struct %s has a native copy."),*GetName());
 			StructFlags = EStructFlags(StructFlags | STRUCT_CopyNative);
 		}
+		if (CppStructOps->HasMoveAssign())
+		{
+			UE_LOG(LogClass, Verbose, TEXT("Native struct %s has a native move assign."),*GetName());
+			StructFlags = EStructFlags(StructFlags | STRUCT_MoveAssignNative);
+
+			// Move-only types are not yet supported. MoveAssignScriptStruct falls back to
+			// CopyScriptStruct for code paths that haven't been updated, so WithCopy must
+			// also be set to ensure correct behavior in those fallback paths.
+			checkf(CppStructOps->HasCopy(), TEXT("Native struct %s has WithMoveAssign but not WithCopy. "
+				"Move-only types are not yet supported; WithCopy must also be set."), *GetName());
+		}
 		if (!CppStructOps->HasDestructor())
 		{
 			UE_LOG(LogClass, Verbose, TEXT("Native struct %s has no destructor."),*GetName());
@@ -3714,6 +3725,32 @@ void UScriptStruct::CopyScriptStruct(void* InDest, void const* InSrc, int32 Arra
 	}
 }
 
+void UScriptStruct::MoveAssignScriptStruct(void* InDest, void* InSrc, int32 ArrayDim) const
+{
+	uint8* Dest = (uint8*)InDest;
+	check(Dest);
+	uint8* Src = (uint8*)InSrc;
+	check(Src);
+
+	int32 Stride = GetStructureSize();
+
+	if (StructFlags & STRUCT_MoveAssignNative)
+	{
+		check(!(StructFlags & STRUCT_IsPlainOldData)); // should not have both
+		UScriptStruct::ICppStructOps* TheCppStructOps = GetCppStructOps();
+		check(TheCppStructOps);
+		check(Stride == TheCppStructOps->GetSize() && PropertiesSize == Stride);
+		if (TheCppStructOps->MoveAssign(Dest, Src, ArrayDim))
+		{
+			return;
+		}
+	}
+
+	// Fallback: use copy. For POD this is memcpy; for non-POD this is copy assignment.
+	// Source is not modified — caller's destructor handles cleanup as before.
+	CopyScriptStruct(Dest, Src, ArrayDim);
+}
+
 uint32 UScriptStruct::GetStructTypeHash(const void* Src) const
 {
 	// Calling GetStructTypeHash on struct types that doesn't provide a native
diff --git a/Engine/Source/Runtime/CoreUObject/Public/UObject/Class.h b/Engine/Source/Runtime/CoreUObject/Public/UObject/Class.h
--- a/Engine/Source/Runtime/CoreUObject/Public/UObject/Class.h
+++ b/Engine/Source/Runtime/CoreUObject/Public/UObject/Class.h
@@ -1126,7 +1126,10 @@ enum EStructFlags
 
 	/** If set, this struct will have PostLoad on it in the editor. */
 	STRUCT_PostLoad UE_INTERNAL = 0x08000000,
-	
+
+	/** If set, this struct will be move-assigned using the CPP operator=(T&&) */
+	STRUCT_MoveAssignNative		= 0x10000000,
+
 	// !!!!!!!!!!!!!!!!!!!!!!!!!!!!!!!!!!!!!!!!!!!!!!!!!!!!!!!!!!!!!!!!!!!!!!!!!!!!!!!!!!!!!!!!!!!!!!!!!!!!!!!!!!!!!!!!!!!!!!!!!!!!!
 	// DO NOT ADD composite flags to EStructFlags. Debugger visualizations have trouble displaying composite flag values as text.
 	// !!!!!!!!!!!!!!!!!!!!!!!!!!!!!!!!!!!!!!!!!!!!!!!!!!!!!!!!!!!!!!!!!!!!!!!!!!!!!!!!!!!!!!!!!!!!!!!!!!!!!!!!!!!!!!!!!!!!!!!!!!!!!
@@ -1136,7 +1139,7 @@ enum EStructFlags
 #define STRUCT_Inherit			(STRUCT_HasInstancedReference | STRUCT_Atomic)
 
 /** Flags that are always computed, never loaded or done with code generation */
-#define STRUCT_ComputedFlags	(STRUCT_NetDeltaSerializeNative | STRUCT_NetSerializeNative | STRUCT_SerializeNative | STRUCT_PostSerializeNative | STRUCT_CopyNative | STRUCT_IsPlainOldData | STRUCT_NoDestructor | STRUCT_ZeroConstructor | STRUCT_IdenticalNative | STRUCT_AddStructReferencedObjects | STRUCT_ExportTextItemNative | STRUCT_ImportTextItemNative | STRUCT_SerializeFromMismatchedTag | STRUCT_PostScriptConstruct | STRUCT_NetSharedSerialization)
+#define STRUCT_ComputedFlags	(STRUCT_NetDeltaSerializeNative | STRUCT_NetSerializeNative | STRUCT_SerializeNative | STRUCT_PostSerializeNative | STRUCT_CopyNative | STRUCT_MoveAssignNative | STRUCT_IsPlainOldData | STRUCT_NoDestructor | STRUCT_ZeroConstructor | STRUCT_IdenticalNative | STRUCT_AddStructReferencedObjects | STRUCT_ExportTextItemNative | STRUCT_ImportTextItemNative | STRUCT_SerializeFromMismatchedTag | STRUCT_PostScriptConstruct | STRUCT_NetSharedSerialization)
 
 
 #if CHECK_PUREVIRTUALS
@@ -1184,6 +1187,7 @@ namespace UE::CoreUObject::Private
 		bool IsUECoreType : 1;
 		bool IsUECoreVariant : 1;
 		bool HasCopy : 1;
+		bool HasMoveAssign : 1;
 		bool HasIdentical : 1;
 		bool HasExportTextItem : 1;
 		bool HasImportTextItem : 1;
@@ -1230,6 +1234,7 @@ namespace UE::CoreUObject::Private
 			.IsUECoreType								= TIsUECoreType<CppStruct>::Value,
 			.IsUECoreVariant							= TIsUECoreVariant<CppStruct>::Value,
 			.HasCopy									= TraitsType::WithCopy,
+			.HasMoveAssign								= TraitsType::WithMoveAssign,
 			.HasIdentical								= TraitsType::WithIdentical || TraitsType::WithIdenticalViaEquality,
 			.HasExportTextItem							= TraitsType::WithExportTextItem,
 			.HasImportTextItem							= TraitsType::WithImportTextItem,
@@ -1385,6 +1390,22 @@ namespace UE::CoreUObject::Private
 		return true;
 	}
 
+	template <typename CppStruct>
+		requires (!!TStructOpsTypeTraits<CppStruct>::WithMoveAssign)
+	bool MoveAssign(void* Dest, void* Src, int32 ArrayDim)
+	{
+		static_assert((!TIsPODType<CppStruct>::Value), "You probably don't want custom move for a POD type.");
+
+		CppStruct* TypedDest = (CppStruct*)Dest;
+		CppStruct* TypedSrc  = (CppStruct*)Src;
+
+		for (; ArrayDim; --ArrayDim)
+		{
+			*TypedDest++ = MoveTemp(*TypedSrc++);
+		}
+		return true;
+	}
+
 	// If you get an overload resolution error here, it's likely you've specified both WithIdentical and WithIdenticalViaEquality
 	// in your TStructOpsTypeTraits, which is not supported.
 	template <typename CppStruct>
@@ -1572,22 +1593,23 @@ namespace UE::CoreUObject::Private
 		IndexNameFunctype(8,	PostScriptConstruct,					void(void* Data)) \
 		IndexNameFunctype(9,	GetPreloadDependencies,					void(void* Data, TArray<UObject*>& OutDeps)) \
 		IndexNameFunctype(10,	Copy,									bool(void* Dest, void const* Src, int32 ArrayDim)) \
-		IndexNameFunctype(11,	Identical,								bool(const void* A, const void* B, uint32 PortFlags, bool& bOutResult)) \
-		IndexNameFunctype(12,	ExportTextItem,							bool(FString& ValueStr, const void* PropertyValue, const void* DefaultValue, UObject* Parent, int32 PortFlags, UObject* ExportRootScope)) \
-		IndexNameFunctype(13,	ImportTextItem,							bool(const TCHAR*& Buffer, void* Data, int32 PortFlags, UObject* OwnerObject, FOutputDevice* ErrorText)) \
-		IndexNameFunctype(14,	FindInnerPropertyInstance,				bool(FName PropertyName, const void* Data, const FProperty*& OutProp, const void*& OutData)) \
-		IndexNameFunctype(15,	AddStructReferencedObjects,				void(*())(void* A, FReferenceCollector& Collector)) \
-		IndexNameFunctype(16,	SerializeFromMismatchedTag,				bool(FPropertyTag const& Tag, FArchive& Ar, void* Data)) \
-		IndexNameFunctype(17,	StructuredSerializeFromMismatchedTag,	bool(FPropertyTag const& Tag, FStructuredArchive::FSlot Slot, void* Data)) \
-		IndexNameFunctype(18,	GetStructTypeHash,						uint32(const void* Src)) \
-		IndexNameFunctype(19,	InitializeIntrusiveUnsetOptionalValue,	void(void* Data)) \
-		IndexNameFunctype(20,	IsIntrusiveOptionalValueSet,			bool(const void* Data)) \
-		IndexNameFunctype(21,	ClearIntrusiveOptionalValue,			void(void* Data)) \
-		IndexNameFunctype(22,	Visit,									EPropertyVisitorControlFlow(FPropertyVisitorContext& Context, const TFunctionRef<EPropertyVisitorControlFlow(const FPropertyVisitorContext& Context)> InFunc)) \
-		IndexNameFunctype(23,	ResolveVisitedPathInfo,					void*(void* Data, const FPropertyVisitorInfo& Info)) \
+		IndexNameFunctype(11,	MoveAssign,								bool(void* Dest, void* Src, int32 ArrayDim)) \
+		IndexNameFunctype(12,	Identical,								bool(const void* A, const void* B, uint32 PortFlags, bool& bOutResult)) \
+		IndexNameFunctype(13,	ExportTextItem,							bool(FString& ValueStr, const void* PropertyValue, const void* DefaultValue, UObject* Parent, int32 PortFlags, UObject* ExportRootScope)) \
+		IndexNameFunctype(14,	ImportTextItem,							bool(const TCHAR*& Buffer, void* Data, int32 PortFlags, UObject* OwnerObject, FOutputDevice* ErrorText)) \
+		IndexNameFunctype(15,	FindInnerPropertyInstance,				bool(FName PropertyName, const void* Data, const FProperty*& OutProp, const void*& OutData)) \
+		IndexNameFunctype(16,	AddStructReferencedObjects,				void(*())(void* A, FReferenceCollector& Collector)) \
+		IndexNameFunctype(17,	SerializeFromMismatchedTag,				bool(FPropertyTag const& Tag, FArchive& Ar, void* Data)) \
+		IndexNameFunctype(18,	StructuredSerializeFromMismatchedTag,	bool(FPropertyTag const& Tag, FStructuredArchive::FSlot Slot, void* Data)) \
+		IndexNameFunctype(19,	GetStructTypeHash,						uint32(const void* Src)) \
+		IndexNameFunctype(20,	InitializeIntrusiveUnsetOptionalValue,	void(void* Data)) \
+		IndexNameFunctype(21,	IsIntrusiveOptionalValueSet,			bool(const void* Data)) \
+		IndexNameFunctype(22,	ClearIntrusiveOptionalValue,			void(void* Data)) \
+		IndexNameFunctype(23,	Visit,									EPropertyVisitorControlFlow(FPropertyVisitorContext& Context, const TFunctionRef<EPropertyVisitorControlFlow(const FPropertyVisitorContext& Context)> InFunc)) \
+		IndexNameFunctype(24,	ResolveVisitedPathInfo,					void*(void* Data, const FPropertyVisitorInfo& Info)) \
 		UE_FOREACH_FAKE_VTABLE_FUNC_EDITOR_ONLY( \
-		IndexNameFunctype(24,	CanEditChange,							bool(const FEditPropertyChain& PropertyChain, const void* Data)) \
-		IndexNameFunctype(25,	PostLoad,								void(void* Data)) \
+		IndexNameFunctype(25,	CanEditChange,							bool(const FEditPropertyChain& PropertyChain, const void* Data)) \
+		IndexNameFunctype(26,	PostLoad,								void(void* Data)) \
 		)
 
 	// These represent one function per optional function in the fake vtable.
@@ -1972,6 +1994,21 @@ public:
 			return HasCopy() && FakeVPtr->GetHandlingFunc<UE::CoreUObject::Private::EStructOpsFakeVTableFlags::Copy>()(Dest, Src, ArrayDim);
 		}
 
+		/** return true if this struct can be move-assigned **/
+		bool HasMoveAssign() const
+		{
+			return GetCapabilities().HasMoveAssign;
+		}
+
+		/**
+		 * Move-assign this structure. Source is left in a valid but unspecified (moved-from) state.
+		 * @return true if the move was handled, otherwise caller should fall back to CopyScriptStruct
+		 */
+		bool MoveAssign(void* Dest, void* Src, int32 ArrayDim)
+		{
+			return HasMoveAssign() && FakeVPtr->GetHandlingFunc<UE::CoreUObject::Private::EStructOpsFakeVTableFlags::MoveAssign>()(Dest, Src, ArrayDim);
+		}
+
 		/** return true if this struct can compare **/
 		bool HasIdentical() const
 		{
@@ -2465,6 +2502,18 @@ public:
 	 * @param	Stride		Stride of the array, If this default (0), then we will pull the size from the struct
 	 */
 	COREUOBJECT_API void CopyScriptStruct(void* Dest, void const* Src, int32 ArrayDim = 1) const;
+
+	/**
+	 * Move-assign a struct over an existing struct. Uses C++ move assignment operator (operator=(T&&))
+	 * if the struct has opted in via WithMoveAssign trait. Falls back to CopyScriptStruct for types
+	 * without the trait, or for POD types (where move == copy).
+	 * Source is left in a valid but unspecified (moved-from) state after a successful move.
+	 *
+	 * @param	Dest		Pointer to initialized destination memory
+	 * @param	Src			Pointer to source memory (will be moved from)
+	 * @param	ArrayDim	Number of elements in the array
+	 */
+	COREUOBJECT_API void MoveAssignScriptStruct(void* Dest, void* Src, int32 ArrayDim = 1) const;
 	
 	/**
 	 * Reinitialize a struct in memory. This may be done by calling the native destructor and then the constructor or individually reinitializing properties
diff --git a/Engine/Source/Runtime/CoreUObject/Public/UObject/StructOpsTypeTraits.h b/Engine/Source/Runtime/CoreUObject/Public/UObject/StructOpsTypeTraits.h
--- a/Engine/Source/Runtime/CoreUObject/Public/UObject/StructOpsTypeTraits.h
+++ b/Engine/Source/Runtime/CoreUObject/Public/UObject/StructOpsTypeTraits.h
@@ -15,6 +15,7 @@ struct TStructOpsTypeTraitsBase2
 		WithNoInitConstructor          = false,                         // struct has a constructor which takes an EForceInit parameter which will force the constructor to perform initialization, where the default constructor performs 'uninitialization'.
 		WithNoDestructor               = false,                         // struct will not have its destructor called when it is destroyed.
 		WithCopy                       = !TIsPODType<CPPSTRUCT>::Value, // struct can be copied via its copy assignment operator.
+		WithMoveAssign                 = false,                         // struct can be move-assigned via its move assignment operator. Opt-in; enables MoveAssignScriptStruct to use operator=(T&&) instead of copy.
 		WithIdenticalViaEquality       = false,                         // struct can be compared via its operator==.  This should be mutually exclusive with WithIdentical.
 		WithIdentical                  = false,                         // struct can be compared via an Identical(const T* Other, uint32 PortFlags) function.  This should be mutually exclusive with WithIdenticalViaEquality.
 		WithExportTextItem             = false,                         // struct has an ExportTextItem function used to serialize its state into a string.
