Potential fix for FObjectProperty failing to return proper error stat info on ImportText.

#Synced-CL 51548522

diff --git a/Engine/Source/Runtime/CoreUObject/Private/UObject/PropertyBaseObject.cpp b/Engine/Source/Runtime/CoreUObject/Private/UObject/PropertyBaseObject.cpp
--- a/Engine/Source/Runtime/CoreUObject/Private/UObject/PropertyBaseObject.cpp
+++ b/Engine/Source/Runtime/CoreUObject/Private/UObject/PropertyBaseObject.cpp
@@ -417,47 +417,56 @@ const TCHAR* FObjectPropertyBase::ImportText_Internal(const TCHAR* InBuffer, TNo
 
 	bool bOk = ParseObjectPropertyValue(this, Parent, PropertyClass, PortFlags, Buffer, Result, FUObjectThreadContext::Get().GetSerializeContext());
 
-	bool bNeedsInstancing = false;
-	if (Result && ((PortFlags & PPF_InstanceSubobjects) != 0))
+	if(bOk)
 	{
-		if (HasAnyPropertyFlags(CPF_InstancedReference))
+		bool bNeedsInstancing = false;
+		if (Result && ((PortFlags & PPF_InstanceSubobjects) != 0))
 		{
-			// Always duplicate the imported reference for explicitly-instanced properties.
-			bNeedsInstancing = true;
-		}
-		else if (Parent && Parent->GetClass()->ShouldUseDynamicSubobjectInstancing())
-		{
-			// Only duplicate the imported reference when scoped to the parent or its archetype.
-			if (Result->IsInOuter(Parent) || Result->IsInOuter(Parent->GetArchetype()))
+			if (HasAnyPropertyFlags(CPF_InstancedReference))
 			{
+				// Always duplicate the imported reference for explicitly-instanced properties.
 				bNeedsInstancing = true;
 			}
+			else if (Parent && Parent->GetClass()->ShouldUseDynamicSubobjectInstancing())
+			{
+				// Only duplicate the imported reference when scoped to the parent or its archetype.
+				if (Result->IsInOuter(Parent) || Result->IsInOuter(Parent->GetArchetype()))
+				{
+					bNeedsInstancing = true;
+				}
+			}
 		}
-	}
-
-	if (bNeedsInstancing)
-	{
-		FName DesiredName = Result->GetFName();
 
-		// If an object currently exists with the same name as the imported object that is to be instanced
-		// 
-		if (UObject* ExistingObject = static_cast<UObject*>(FindObjectWithOuter(Parent, nullptr, DesiredName)))
+		if (bNeedsInstancing)
 		{
-			ExistingObject->Rename(nullptr, nullptr, REN_DontCreateRedirectors | REN_AllowPackageLinkerMismatch);
-		}
+			FName DesiredName = Result->GetFName();
+
+			// If an object currently exists with the same name as the imported object that is to be instanced
+			// 
+			if (UObject* ExistingObject = static_cast<UObject*>(FindObjectWithOuter(Parent, nullptr, DesiredName)))
+			{
+				ExistingObject->Rename(nullptr, nullptr, REN_DontCreateRedirectors | REN_AllowPackageLinkerMismatch);
+			}
 		
-		const EObjectFlags MaskedOuterFlags = Parent ? Parent->GetMaskedFlags(RF_PropagateToSubObjects) : RF_AllFlags;
-		FObjectDuplicationParameters ObjectDuplicationParams = InitStaticDuplicateObjectParams(Result, Parent, DesiredName, MaskedOuterFlags);
-		EnumRemoveFlags(ObjectDuplicationParams.FlagMask, RF_ArchetypeObject);
-		if (Parent && Parent->IsTemplate())
-		{
-			EnumAddFlags(ObjectDuplicationParams.ApplyFlags, RF_ArchetypeObject);
-		}
-		else
-		{
-			EnumRemoveFlags(ObjectDuplicationParams.ApplyFlags, RF_ArchetypeObject);
+			const EObjectFlags MaskedOuterFlags = Parent ? Parent->GetMaskedFlags(RF_PropagateToSubObjects) : RF_AllFlags;
+			FObjectDuplicationParameters ObjectDuplicationParams = InitStaticDuplicateObjectParams(Result, Parent, DesiredName, MaskedOuterFlags);
+			EnumRemoveFlags(ObjectDuplicationParams.FlagMask, RF_ArchetypeObject);
+			if (Parent && Parent->IsTemplate())
+			{
+				EnumAddFlags(ObjectDuplicationParams.ApplyFlags, RF_ArchetypeObject);
+			}
+			else
+			{
+				EnumRemoveFlags(ObjectDuplicationParams.ApplyFlags, RF_ArchetypeObject);
+			}
+			Result = StaticDuplicateObjectEx(ObjectDuplicationParams);
 		}
-		Result = StaticDuplicateObjectEx(ObjectDuplicationParams);
+	}
+	else
+	{
+		//Ensure we set nullptr and return null on failure.
+		Buffer = nullptr;
+		Result = nullptr;
 	}
 
 	if (PropertyPointerType == EPropertyPointerType::Container && HasSetter())
@@ -468,6 +477,7 @@ const TCHAR* FObjectPropertyBase::ImportText_Internal(const TCHAR* InBuffer, TNo
 	{
 		SetObjectPtrPropertyValue(PointerToValuePtr(ContainerOrPropertyPtr, PropertyPointerType), Result);
 	}
+
 	return Buffer;
 }
 
