[Cooker/MTSAVE]
- Remove report redirect that isn't used since 2018 [at] 4088724

#Synced-CL 51820749

diff --git a/Engine/Source/Runtime/CoreUObject/Private/UObject/Obj.cpp b/Engine/Source/Runtime/CoreUObject/Private/UObject/Obj.cpp
--- a/Engine/Source/Runtime/CoreUObject/Private/UObject/Obj.cpp
+++ b/Engine/Source/Runtime/CoreUObject/Private/UObject/Obj.cpp
@@ -1708,35 +1708,13 @@ void UObject::GetPreloadDependencies(TArray<UObject*>& OutDeps)
 	}
 }
 
-// This is a terrible hack to allow the checking of redirected
-// soft object paths in CDOs at cook time.  Redirects in CDOs
-// cause non-determinism issues and need to be reported.
-//
-// This global is extern'd and handled in SoftObjectPath.cpp.
-bool* GReportSoftObjectPathRedirects = nullptr;
-
 IMPLEMENT_FARCHIVE_SERIALIZER(UObject)
 
 void UObject::Serialize(FStructuredArchive::FRecord Record)
 {
 	SCOPED_LOADTIMER(UObject_Serialize);
 
-#if WITH_EDITOR
-	bool bReportSoftObjectPathRedirects = false;
-
-	{
-		// TOptionalGuardValue will not overwrite the value if it remains the same.
-		// This is important for TSAN as we only want warnings if this unprotected value is changing.
-		TOptionalGuardValue<bool*> GuardValue(
-			GReportSoftObjectPathRedirects,
-			  GReportSoftObjectPathRedirects
-			? GReportSoftObjectPathRedirects
-			: (GIsCookerLoadingPackage && HasAnyFlags(RF_ClassDefaultObject | RF_ArchetypeObject))
-			? &bReportSoftObjectPathRedirects
-			: nullptr
-		);
-#endif
-
+	{
 		FArchive& UnderlyingArchive = Record.GetUnderlyingArchive();
 
 		// This item is very special from a serialization standpoint. It is never actually serialized.
@@ -1911,14 +1889,7 @@ void UObject::Serialize(FStructuredArchive::FRecord Record)
 		// Memory counting (with proper alignment to match C++)
 		SIZE_T Size = GetClass()->GetStructureSize();
 		UnderlyingArchive.CountBytes(Size, Size);
-#if WITH_EDITOR
 	}
-
-	if (bReportSoftObjectPathRedirects && !GReportSoftObjectPathRedirects)
-	{
-		UE_ASSET_LOG(LogCore, Warning, this, TEXT("Soft object paths were redirected during cook of '%s' - package should be resaved."), *GetName());
-	}
-#endif
 }
 
 #if WITH_EDITORONLY_DATA
diff --git a/Engine/Source/Runtime/CoreUObject/Private/UObject/SoftObjectPath.cpp b/Engine/Source/Runtime/CoreUObject/Private/UObject/SoftObjectPath.cpp
--- a/Engine/Source/Runtime/CoreUObject/Private/UObject/SoftObjectPath.cpp
+++ b/Engine/Source/Runtime/CoreUObject/Private/UObject/SoftObjectPath.cpp
@@ -313,10 +313,6 @@ void FSoftObjectPath::SetPath(const UObject* InObject)
 	SubPathString = SubPathStringBuilder.ToString();
 }
 
-#if WITH_EDITOR
-	extern bool* GReportSoftObjectPathRedirects;
-#endif
-
 bool FSoftObjectPath::PreSavePath(bool* bReportSoftObjectPathRedirects)
 {
 #if WITH_EDITOR
@@ -377,7 +373,7 @@ void FSoftObjectPath::SerializePath(FArchive& Ar)
 #if WITH_EDITOR
 	if (Ar.IsSaving() && !(Ar.IsModifyingWeakAndStrongReferences() && Ar.IsObjectReferenceCollector()))
 	{
-		PreSavePath(/*false ? GReportSoftObjectPathRedirects :*/ nullptr);
+		PreSavePath(nullptr);
 	}
 
 	// Only read serialization options in editor as it is a bit slow
