[Insights]
* ObjectTrace: Added extra spec event for UFunction.
* ObjectTrace: Added a toggle for tracing total resource sizes (inclusive memory estimates) as it is too slow.

#rnx

#Synced-CL 51759468

diff --git a/Engine/Source/Runtime/CoreUObject/Private/UObject/ObjectTrace.cpp b/Engine/Source/Runtime/CoreUObject/Private/UObject/ObjectTrace.cpp
--- a/Engine/Source/Runtime/CoreUObject/Private/UObject/ObjectTrace.cpp
+++ b/Engine/Source/Runtime/CoreUObject/Private/UObject/ObjectTrace.cpp
@@ -35,6 +35,7 @@ UE_TRACE_MINIMAL_EVENT_BEGIN(CoreUObject, BeginSnapshot, NoSync)
 	UE_TRACE_MINIMAL_EVENT_FIELD(uint32, SizeOfUField)
 	UE_TRACE_MINIMAL_EVENT_FIELD(uint32, SizeOfUStruct)
 	UE_TRACE_MINIMAL_EVENT_FIELD(uint32, SizeOfUClass)
+	UE_TRACE_MINIMAL_EVENT_FIELD(uint32, SizeOfUFunction)
 	UE_TRACE_MINIMAL_EVENT_FIELD(uint32, SizeOfUPackage)
 UE_TRACE_MINIMAL_EVENT_END()
 
@@ -79,6 +80,13 @@ UE_TRACE_MINIMAL_EVENT_BEGIN(CoreUObject, ClassExtraSpec, NoSync)
 	UE_TRACE_MINIMAL_EVENT_FIELD(uint32, Id)
 UE_TRACE_MINIMAL_EVENT_END()
 
+UE_TRACE_MINIMAL_EVENT_BEGIN(CoreUObject, FunctionExtraSpec, NoSync)
+	UE_TRACE_MINIMAL_EVENT_FIELD(uint32, Id)
+	UE_TRACE_MINIMAL_EVENT_FIELD(uint32, Flags)
+	UE_TRACE_MINIMAL_EVENT_FIELD(uint8, NumParms)
+	UE_TRACE_MINIMAL_EVENT_FIELD(uint16, ParmsSize)
+UE_TRACE_MINIMAL_EVENT_END()
+
 UE_TRACE_MINIMAL_EVENT_BEGIN(CoreUObject, PackageExtraSpec, NoSync)
 	UE_TRACE_MINIMAL_EVENT_FIELD(uint32, Id)
 	UE_TRACE_MINIMAL_EVENT_FIELD(uint64, PackageId)
@@ -168,6 +176,7 @@ void FUObjectTrace::OutputBeginSnapshot()
 		<< BeginSnapshot.SizeOfUField(static_cast<uint32>(sizeof(UField)))
 		<< BeginSnapshot.SizeOfUStruct(static_cast<uint32>(sizeof(UStruct)))
 		<< BeginSnapshot.SizeOfUClass(static_cast<uint32>(sizeof(UClass)))
+		<< BeginSnapshot.SizeOfUFunction(static_cast<uint32>(sizeof(UFunction)))
 		<< BeginSnapshot.SizeOfUPackage(static_cast<uint32>(sizeof(UPackage)))
 	;
 }
@@ -237,18 +246,28 @@ bool FUObjectTrace::OutputObjectSpec(UObjectBase* Object)
 	}
 
 	constexpr bool bTraceResourceSize = true;
+	constexpr bool bTraceTotalResourceSize = false; // very slow (~20s / 1M objects) !!!
 	if (bTraceResourceSize)
 	{
+		// Exclusive Estimated (System/Video) Memory Size
 		FResourceSizeEx ResourceSize(EResourceSizeMode::Exclusive);
-		FResourceSizeEx TotalResourceSize(EResourceSizeMode::EstimatedTotal);
 		Obj->GetResourceSizeEx(ResourceSize);
-		Obj->GetResourceSizeEx(TotalResourceSize);
-		SIZE_T SystemMemoryBytes = ResourceSize.GetDedicatedSystemMemoryBytes(); // exclusive
-		SIZE_T VideoMemoryBytes = ResourceSize.GetDedicatedVideoMemoryBytes(); // exclusive
-		SIZE_T TotalSystemMemoryBytes = TotalResourceSize.GetDedicatedSystemMemoryBytes(); // inclusive
-		SIZE_T TotalVideoMemoryBytes = TotalResourceSize.GetDedicatedVideoMemoryBytes(); // inclusive
-		if (SystemMemoryBytes || VideoMemoryBytes || TotalSystemMemoryBytes || TotalVideoMemoryBytes)
+		SIZE_T SystemMemoryBytes = ResourceSize.GetDedicatedSystemMemoryBytes();
+		SIZE_T VideoMemoryBytes = ResourceSize.GetDedicatedVideoMemoryBytes();
+
+		if (SystemMemoryBytes || VideoMemoryBytes)
 		{
+			// Inclusive Estimated (System/Video) Memory Size
+			SIZE_T TotalSystemMemoryBytes = 0;
+			SIZE_T TotalVideoMemoryBytes = 0;
+			if (bTraceTotalResourceSize)
+			{
+				FResourceSizeEx TotalResourceSize(EResourceSizeMode::EstimatedTotal);
+				Obj->GetResourceSizeEx(TotalResourceSize);
+				TotalSystemMemoryBytes = TotalResourceSize.GetDedicatedSystemMemoryBytes();
+				TotalVideoMemoryBytes = TotalResourceSize.GetDedicatedVideoMemoryBytes();
+			}
+
 			UE_TRACE_MINIMAL_LOG(CoreUObject, ResourceSizeExtraSpec, CoreUObjectChannel)
 				<< ResourceSizeExtraSpec.Id(Id)
 				<< ResourceSizeExtraSpec.SystemMemoryBytes(SystemMemoryBytes)
@@ -288,6 +307,16 @@ bool FUObjectTrace::OutputObjectSpec(UObjectBase* Object)
 				<< ClassExtraSpec.Id(Id);
 		}
 
+		// UFunction -> UStruct -> UField -> UObject
+		if (UFunction* ObjAsFunction = Cast<UFunction>(Object))
+		{
+			UE_TRACE_MINIMAL_LOG(CoreUObject, FunctionExtraSpec, CoreUObjectChannel)
+				<< FunctionExtraSpec.Id(Id)
+				<< FunctionExtraSpec.Flags(uint32(ObjAsFunction->FunctionFlags))
+				<< FunctionExtraSpec.NumParms(ObjAsFunction->NumParms)
+				<< FunctionExtraSpec.ParmsSize(ObjAsFunction->ParmsSize);
+		}
+
 		// UPackage -> UObject
 		if (UPackage* ObjAsPackage = Cast<UPackage>(Object))
 		{
