[Backout] - CL51823190
#rnx
Original CL Desc
-----------------------------------------------------------------
Potential fix for FObjectProperty failing to return proper error stat info on ImportText.

diff --git a/Engine/Source/Runtime/CoreUObject/Private/UObject/PropertyBaseObject.cpp b/Engine/Source/Runtime/CoreUObject/Private/UObject/PropertyBaseObject.cpp
--- a/Engine/Source/Runtime/CoreUObject/Private/UObject/PropertyBaseObject.cpp
+++ b/Engine/Source/Runtime/CoreUObject/Private/UObject/PropertyBaseObject.cpp
@@ -417,56 +417,47 @@ const TCHAR* FObjectPropertyBase::ImportText_Internal(const TCHAR* InBuffer, TNo
 
 	bool bOk = ParseObjectPropertyValue(this, Parent, PropertyClass, PortFlags, Buffer, Result, FUObjectThreadContext::Get().GetSerializeContext());
 
-	if(bOk)
+	bool bNeedsInstancing = false;
+	if (Result && ((PortFlags & PPF_InstanceSubobjects) != 0))
 	{
-		bool bNeedsInstancing = false;
-		if (Result && ((PortFlags & PPF_InstanceSubobjects) != 0))
+		if (HasAnyPropertyFlags(CPF_InstancedReference))
 		{
-			if (HasAnyPropertyFlags(CPF_InstancedReference))
+			// Always duplicate the imported reference for explicitly-instanced properties.
+			bNeedsInstancing = true;
+		}
+		else if (Parent && Parent->GetClass()->ShouldUseDynamicSubobjectInstancing())
+		{
+			// Only duplicate the imported reference when scoped to the parent or its archetype.
+			if (Result->IsInOuter(Parent) || Result->IsInOuter(Parent->GetArchetype()))
 			{
-				// Always duplicate the imported reference for explicitly-instanced properties.
 				bNeedsInstancing = true;
 			}
-			else if (Parent && Parent->GetClass()->ShouldUseDynamicSubobjectInstancing())
-			{
-				// Only duplicate the imported reference when scoped to the parent or its archetype.
-				if (Result->IsInOuter(Parent) || Result->IsInOuter(Parent->GetArchetype()))
-				{
-					bNeedsInstancing = true;
-				}
-			}
 		}
+	}
 
-		if (bNeedsInstancing)
-		{
-			FName DesiredName = Result->GetFName();
+	if (bNeedsInstancing)
+	{
+		FName DesiredName = Result->GetFName();
 
-			// If an object currently exists with the same name as the imported object that is to be instanced
-			// 
-			if (UObject* ExistingObject = static_cast<UObject*>(FindObjectWithOuter(Parent, nullptr, DesiredName)))
-			{
-				ExistingObject->Rename(nullptr, nullptr, REN_DontCreateRedirectors | REN_AllowPackageLinkerMismatch);
-			}
+		// If an object currently exists with the same name as the imported object that is to be instanced
+		// 
+		if (UObject* ExistingObject = static_cast<UObject*>(FindObjectWithOuter(Parent, nullptr, DesiredName)))
+		{
+			ExistingObject->Rename(nullptr, nullptr, REN_DontCreateRedirectors | REN_AllowPackageLinkerMismatch);
+		}
 		
-			const EObjectFlags MaskedOuterFlags = Parent ? Parent->GetMaskedFlags(RF_PropagateToSubObjects) : RF_AllFlags;
-			FObjectDuplicationParameters ObjectDuplicationParams = InitStaticDuplicateObjectParams(Result, Parent, DesiredName, MaskedOuterFlags);
-			EnumRemoveFlags(ObjectDuplicationParams.FlagMask, RF_ArchetypeObject);
-			if (Parent && Parent->IsTemplate())
-			{
-				EnumAddFlags(ObjectDuplicationParams.ApplyFlags, RF_ArchetypeObject);
-			}
-			else
-			{
-				EnumRemoveFlags(ObjectDuplicationParams.ApplyFlags, RF_ArchetypeObject);
-			}
-			Result = StaticDuplicateObjectEx(ObjectDuplicationParams);
+		const EObjectFlags MaskedOuterFlags = Parent ? Parent->GetMaskedFlags(RF_PropagateToSubObjects) : RF_AllFlags;
+		FObjectDuplicationParameters ObjectDuplicationParams = InitStaticDuplicateObjectParams(Result, Parent, DesiredName, MaskedOuterFlags);
+		EnumRemoveFlags(ObjectDuplicationParams.FlagMask, RF_ArchetypeObject);
+		if (Parent && Parent->IsTemplate())
+		{
+			EnumAddFlags(ObjectDuplicationParams.ApplyFlags, RF_ArchetypeObject);
 		}
-	}
-	else
-	{
-		//Ensure we set nullptr and return null on failure.
-		Buffer = nullptr;
-		Result = nullptr;
+		else
+		{
+			EnumRemoveFlags(ObjectDuplicationParams.ApplyFlags, RF_ArchetypeObject);
+		}
+		Result = StaticDuplicateObjectEx(ObjectDuplicationParams);
 	}
 
 	if (PropertyPointerType == EPropertyPointerType::Container && HasSetter())
@@ -477,7 +468,6 @@ const TCHAR* FObjectPropertyBase::ImportText_Internal(const TCHAR* InBuffer, TNo
 	{
 		SetObjectPtrPropertyValue(PointerToValuePtr(ContainerOrPropertyPtr, PropertyPointerType), Result);
 	}
-
 	return Buffer;
 }
 
