[PackageAccessTracking]
- Use the new InheritedContext extensions to automatically propagate package access along task chains
- Rename Inherited FImpl to FInterface since it is a public API and not strictly internals

#Synced-CL 51820749

diff --git a/Engine/Source/Runtime/Core/Public/Async/InheritedContext.h b/Engine/Source/Runtime/Core/Public/Async/InheritedContext.h
--- a/Engine/Source/Runtime/Core/Public/Async/InheritedContext.h
+++ b/Engine/Source/Runtime/Core/Public/Async/InheritedContext.h
@@ -64,7 +64,7 @@ namespace UE
 		using FRestoreFunc  = void (*)(const void* SaveSrc, void* UserData);
 		using FDestroyFunc  = void (*)(const void* DataSrc, void* UserData);
 
-		struct FImpl
+		struct FInterface
 		{
 			FCaptureFunc  Capture;    // Read current thread state into DataDst
 			FApplyFunc    Apply;      // Write DataSrc to thread state, save previous to SaveDst
@@ -75,7 +75,7 @@ namespace UE
 			void*         UserData;   // Typically pointer to the variable being propagated
 		};
 
-		TSharedPtr<FImpl> Impl;
+		TSharedPtr<FInterface> Impl;
 
 		bool IsValid() const { return Impl.IsValid(); }
 
@@ -105,7 +105,7 @@ namespace UE
 		FuncPtr FnPtr = GetVar;
 
 		FInheritedContextExtension Ext;
-		Ext.Impl = MakeShared<FInheritedContextExtension::FImpl>();
+		Ext.Impl = MakeShared<FInheritedContextExtension::FInterface>();
 		Ext.Impl->Capture = [](void* DataDst, void* UserData)
 		{
 			auto Fn = reinterpret_cast<T& (*)()>(UserData);
@@ -146,7 +146,7 @@ namespace UE
 		using T = std::decay_t<decltype(Getter())>;
 
 		FInheritedContextExtension Ext;
-		Ext.Impl = MakeShared<FInheritedContextExtension::FImpl>();
+		Ext.Impl = MakeShared<FInheritedContextExtension::FInterface>();
 		Ext.Impl->Capture = [](void* DataDst, void* /*UserData*/)
 		{
 			new (DataDst) T(Getter());
@@ -171,14 +171,14 @@ namespace UE
 	}
 
 	/**
-	 * Derived FImpl that co-allocates state of type T with the descriptor.
+	 * Derived FInterface that co-allocates state of type T with the descriptor.
 	 * Defined at namespace scope because local classes cannot have template members.
-	 * When the TSharedPtr<FImpl> ref count hits zero, the control block destructs
+	 * When the TSharedPtr<FInterface> ref count hits zero, the control block destructs
 	 * TStatefulContextExtensionImpl<T> which properly destroys State via the type-erased
 	 * deleter stored by TSharedPtr. No virtual destructor needed.
 	 */
 	template<typename T>
-	struct TStatefulContextExtensionImpl : FInheritedContextExtension::FImpl
+	struct TStatefulContextExtensionImpl : FInheritedContextExtension::FInterface
 	{
 		T State;
 
@@ -188,7 +188,7 @@ namespace UE
 
 	/**
 	 * Creates an FInheritedContextExtension with custom scope callbacks and co-allocated mutable state.
-	 * The state of type T is constructed in-place inside a derived FImpl and owned by the handle's
+	 * The state of type T is constructed in-place inside a derived FInterface and owned by the handle's
 	 * TSharedPtr, so it lives exactly as long as the last copy of the extension handle.
 	 * ApplyFn() is called when a task starts executing and returns saved state (preserved per-scope).
 	 * RestoreFn(SavedState, SharedState) is called when the task scope ends.
@@ -240,7 +240,7 @@ namespace UE
 		StatefulRef->UserData  = &StatefulRef->State;
 
 		FInheritedContextExtension Ext;
-		Ext.Impl = TSharedPtr<FInheritedContextExtension::FImpl>(StatefulRef);
+		Ext.Impl = TSharedPtr<FInheritedContextExtension::FInterface>(StatefulRef);
 		return Ext;
 	}
 
diff --git a/Engine/Source/Runtime/Core/Public/Async/InheritedGuardValue.h b/Engine/Source/Runtime/Core/Public/Async/InheritedGuardValue.h
--- a/Engine/Source/Runtime/Core/Public/Async/InheritedGuardValue.h
+++ b/Engine/Source/Runtime/Core/Public/Async/InheritedGuardValue.h
@@ -23,7 +23,7 @@ struct TInheritedGuardValue : private FNoncopyable
 	/**
 	 * Callable constructor — creates the extension internally.
 	 * GetVar must return RefType& (e.g. []() -> bool& { return GMyFlag; }).
-	 * Allocates a TSharedPtr<FImpl> internally, so prefer the pre-built variant for hot paths.
+	 * Allocates a TSharedPtr<FInterface> internally, so prefer the pre-built variant for hot paths.
 	 */
 	template<typename Func>
 	[[nodiscard]] TInheritedGuardValue(Func GetVar, const AssignedType& NewValue)
@@ -91,7 +91,7 @@ struct TInheritedGuardValueAccessors : private FNoncopyable
 
 	/**
 	 * Creates the extension internally.
-	 * Allocates a TSharedPtr<FImpl> internally, so prefer the pre-built variant for hot paths.
+	 * Allocates a TSharedPtr<FInterface> internally, so prefer the pre-built variant for hot paths.
 	 */
 	[[nodiscard]] explicit TInheritedGuardValueAccessors(const ValueType& NewValue)
 		: Extension(UE::MakeInheritedContextExtension<Getter, Setter>())
diff --git a/Engine/Source/Runtime/CoreUObject/Private/Misc/PackageAccessTracking.cpp b/Engine/Source/Runtime/CoreUObject/Private/Misc/PackageAccessTracking.cpp
--- a/Engine/Source/Runtime/CoreUObject/Private/Misc/PackageAccessTracking.cpp
+++ b/Engine/Source/Runtime/CoreUObject/Private/Misc/PackageAccessTracking.cpp
@@ -2,6 +2,7 @@
 
 #include "Misc/PackageAccessTracking.h"
 
+#include "Async/InheritedContext.h"
 #include "Cooker/CookDependency.h"
 #include "Misc/PackageAccessTrackingOps.h"
 #include "Misc/PackageName.h"
@@ -28,6 +29,20 @@ bool IsBuildOpName(FName OpName)
 		;
 }
 
+namespace
+{
+	struct FCapturedPackageTrackingData
+	{
+		FName PackageName;
+		FName OpName;
+		FName BuildOpName;
+		FName CookResultProjection;
+		const UObject* Object;          // TODO: may need FWeakObjectPtr if GC becomes problematic
+		const ITargetPlatform* TargetPlatform;
+		bool bHasData;
+	};
+}
+
 FTrackedData::FTrackedData(FName InPackageName, FName InOpName, FName InCookResultProjection,
 	const ITargetPlatform* InTargetPlatform, const UObject* InObject)
 	: PackageName(InPackageName)
@@ -39,6 +54,17 @@ FTrackedData::FTrackedData(FName InPackageName, FName InOpName, FName InCookResu
 {
 }
 
+FTrackedData::FTrackedData(FName InPackageName, FName InOpName, FName InBuildOpName,
+	FName InCookResultProjection, const ITargetPlatform* InTargetPlatform, const UObject* InObject)
+	: PackageName(InPackageName)
+	, OpName(InOpName)
+	, BuildOpName(InBuildOpName)
+	, CookResultProjection(InCookResultProjection)
+	, Object(InObject)
+	, TargetPlatform(InTargetPlatform)
+{
+}
+
 FTrackedData::FTrackedData(FTrackedData& DirectData, FTrackedData* OuterAccumulatedData)
 {
 	if (OuterAccumulatedData)
@@ -87,6 +113,72 @@ FTrackedData::FTrackedData(FTrackedData& DirectData, FTrackedData* OuterAccumula
 	}
 }
 
+UE::FInheritedContextExtension& FPackageAccessRefScope::GetInheritedContextExtension()
+{
+	static UE::FInheritedContextExtension Extension = []()
+	{
+		UE::FInheritedContextExtension Ext;
+		Ext.Impl = MakeShared<UE::FInheritedContextExtension::FInterface>();
+		Ext.Impl->Capture = [](void* DataDst, void* /*UserData*/)
+		{
+			auto* Dst = static_cast<FCapturedPackageTrackingData*>(DataDst);
+			FTrackedData* Accumulated = GetCurrentThreadAccumulatedData();
+			if (Accumulated)
+			{
+				Dst->PackageName = Accumulated->PackageName;
+				Dst->OpName = Accumulated->OpName;
+				Dst->BuildOpName = Accumulated->BuildOpName;
+				Dst->CookResultProjection = Accumulated->CookResultProjection;
+				Dst->Object = Accumulated->Object;
+				Dst->TargetPlatform = Accumulated->TargetPlatform;
+				Dst->bHasData = true;
+			}
+			else
+			{
+				Dst->bHasData = false;
+			}
+		};
+		Ext.Impl->Apply = [](const void* DataSrc, void* SaveDst, void* /*UserData*/)
+		{
+			auto* Src = static_cast<const FCapturedPackageTrackingData*>(DataSrc);
+			FPackageAccessRefScope** SaveSlot = static_cast<FPackageAccessRefScope**>(SaveDst);
+			if (Src->bHasData)
+			{
+				*SaveSlot = new FPackageAccessRefScope(FInheritedContextTag{},
+					Src->PackageName, Src->OpName, Src->BuildOpName,
+					Src->CookResultProjection, Src->Object, Src->TargetPlatform);
+			}
+			else
+			{
+				*SaveSlot = nullptr;
+			}
+		};
+		Ext.Impl->Restore = [](const void* SaveSrc, void* /*UserData*/)
+		{
+			FPackageAccessRefScope* const* Slot = static_cast<FPackageAccessRefScope* const*>(SaveSrc);
+			delete *Slot; // this may delete nullptr, which is fine
+		};
+		static_assert(std::is_trivially_destructible_v<FCapturedPackageTrackingData>);
+		Ext.Impl->Destroy = nullptr;
+		Ext.Impl->DataSize = static_cast<uint32>(FMath::Max(sizeof(FCapturedPackageTrackingData), sizeof(FPackageAccessRefScope*)));
+		Ext.Impl->DataAlign = static_cast<uint32>(FMath::Max(alignof(FCapturedPackageTrackingData), alignof(FPackageAccessRefScope*)));
+		Ext.Impl->UserData = nullptr;
+		return Ext;
+	}();
+	return Extension;
+}
+
+FPackageAccessRefScope::FPackageAccessRefScope(FInheritedContextTag,
+	FName InPackageName, FName InOpName, FName InBuildOpName,
+	FName InCookResultProjection, const UObject* InObject,
+	const ITargetPlatform* InTargetPlatform)
+	: DirectData(InPackageName, InOpName, InBuildOpName, InCookResultProjection, InTargetPlatform, InObject)
+	, AccumulatedData(DirectData, GetCurrentThreadAccumulatedData())
+	, Outer(CurrentThreadScope)
+{
+	CurrentThreadScope = this;
+}
+
 FPackageAccessRefScope::FPackageAccessRefScope(FName InPackageName, FName InOpName,
 	FName InCookResultProjection, const ITargetPlatform* InTargetPlatform, const UObject* InObject)
 	: DirectData(InPackageName, InOpName, InCookResultProjection, InTargetPlatform, InObject)
@@ -94,6 +186,11 @@ FPackageAccessRefScope::FPackageAccessRefScope(FName InPackageName, FName InOpNa
 	, Outer(CurrentThreadScope)
 {
 	CurrentThreadScope = this;
+
+	if (!Outer)
+	{
+		InheritedExtensionScope = MakeUnique<UE::FInheritedContextExtensionScope>(GetInheritedContextExtension());
+	}
 }
 
 FPackageAccessRefScope::FPackageAccessRefScope(FName InPackageName, FName InOpName)
diff --git a/Engine/Source/Runtime/CoreUObject/Public/Misc/PackageAccessTracking.h b/Engine/Source/Runtime/CoreUObject/Public/Misc/PackageAccessTracking.h
--- a/Engine/Source/Runtime/CoreUObject/Public/Misc/PackageAccessTracking.h
+++ b/Engine/Source/Runtime/CoreUObject/Public/Misc/PackageAccessTracking.h
@@ -4,6 +4,7 @@
 
 #include "CoreTypes.h"
 #include "Misc/Optional.h"
+#include "Templates/UniquePtr.h"
 #include "UObject/NameTypes.h"
 #include "UObject/ObjectHandleDefines.h"
 
@@ -58,6 +59,8 @@ class ITargetPlatform;
 ///////////////////////////////////////////////////////////////////////////////////////////////////////////////////////
 
 #if UE_WITH_PACKAGE_ACCESS_TRACKING
+namespace UE { struct FInheritedContextExtension; class FInheritedContextExtensionScope; }
+
 namespace PackageAccessTracking_Private
 {
 	struct FTrackedData
@@ -66,6 +69,10 @@ namespace PackageAccessTracking_Private
 		COREUOBJECT_API FTrackedData(FName PackageNameIn, FName InOpName, FName InCookResultProjection,
 			const ITargetPlatform* InTargetPlatform, const UObject* InObject);
 
+		/** Direct-init constructor; copies all fields as-is without recomputing BuildOpName via IsBuildOpName. */
+		COREUOBJECT_API FTrackedData(FName InPackageName, FName InOpName, FName InBuildOpName,
+			FName InCookResultProjection, const ITargetPlatform* InTargetPlatform, const UObject* InObject);
+
 		/** Accumulating constructor; sets new AccumulatedData by combining DirectData with Outer's AccumulatedData. */
 		COREUOBJECT_API FTrackedData(FTrackedData& DirectData, FTrackedData* OuterAccumulatedData);
 
@@ -105,9 +112,21 @@ namespace PackageAccessTracking_Private
 		COREUOBJECT_API FPackageAccessRefScope(FName InPackageName, FName InOpName, FName InCookResultProjection,
 			const ITargetPlatform* InTargetPlatform, const UObject* InObject);
 
+		// Tag type + constructor for restoring from inherited context
+		struct FInheritedContextTag { explicit FInheritedContextTag() = default; };
+		COREUOBJECT_API explicit FPackageAccessRefScope(FInheritedContextTag,
+			FName InPackageName, FName InOpName, FName InBuildOpName,
+			FName InCookResultProjection, const UObject* InObject,
+			const ITargetPlatform* InTargetPlatform);
+
+		// Extension factory (lambdas inside have private access)
+		static UE::FInheritedContextExtension& GetInheritedContextExtension();
+
 		FTrackedData DirectData;
 		FTrackedData AccumulatedData;
 		FPackageAccessRefScope* Outer = nullptr;
+		// Extension scope, active only for outermost scope per thread
+		TUniquePtr<UE::FInheritedContextExtensionScope> InheritedExtensionScope;
 		static thread_local FPackageAccessRefScope* CurrentThreadScope;
 	};
 };
