[UObjectGlobals]
- Make a delegate thread-safe

#Synced-CL 51820749

diff --git a/Engine/Source/Runtime/CoreUObject/Public/UObject/UObjectGlobals.h b/Engine/Source/Runtime/CoreUObject/Public/UObject/UObjectGlobals.h
--- a/Engine/Source/Runtime/CoreUObject/Public/UObject/UObjectGlobals.h
+++ b/Engine/Source/Runtime/CoreUObject/Public/UObject/UObjectGlobals.h
@@ -3382,7 +3382,7 @@ struct FCoreUObjectDelegates
 #endif // WITH_EDITOR
 
 	/** Called when new sparse class data has been created (and the base data initialized) for the given class */
-	DECLARE_MULTICAST_DELEGATE_ThreeParams(FOnPostInitSparseClassData, UClass*, UScriptStruct*, void*);
+	DECLARE_TS_MULTICAST_DELEGATE_ThreeParams(FOnPostInitSparseClassData, UClass*, UScriptStruct*, void*);
 	static COREUOBJECT_API FOnPostInitSparseClassData OnPostInitSparseClassData;
 
 	/** Called by ReloadPackage during package reloading. It will be called several times for different phases of fix-up to allow custom code to handle updating objects as needed */
