Fixed OverridableManager stepping into non-owned subobjects in edge case.

diff --git a/Engine/Source/Runtime/CoreUObject/Private/UObject/OverridableManager.cpp b/Engine/Source/Runtime/CoreUObject/Private/UObject/OverridableManager.cpp
--- a/Engine/Source/Runtime/CoreUObject/Private/UObject/OverridableManager.cpp
+++ b/Engine/Source/Runtime/CoreUObject/Private/UObject/OverridableManager.cpp
@@ -286,11 +286,12 @@ EOverriddenState FOverridableManager::GetOverriddenState(TNotNull<UObject*> Obje
 					}
 				}
 
-				// Only recurse on instanced subobjects
-				if (SubObject)
+				if (ObjectProperty)
 				{
-					switch(GetOverriddenState(SubObject, ConsiderPropertyForOverriddenState, {}, nullptr, false, ConsiderSubObjectForOverriddenState))
+					if (SubObject) // Only recurse on object properties if they point to an instanced subobject.
 					{
+						switch(GetOverriddenState(SubObject, ConsiderPropertyForOverriddenState, {}, nullptr, false, ConsiderSubObjectForOverriddenState))
+						{
 						case EOverriddenState::NoOverrides:
 							bAllOverridden = false;
 							break;
@@ -303,8 +304,9 @@ EOverriddenState FOverridableManager::GetOverriddenState(TNotNull<UObject*> Obje
 							break;
 						default:
 							break;
+						}
+						return EPropertyVisitorControlFlow::StepOver;
 					}
-					return EPropertyVisitorControlFlow::StepOver;
 				}
 				else
 				{
