[Tests] Modifying Placeholder tests to recompile the blueprint and force a reinstance pass
[Tests] Disabling some placeholder tests that can't be supported until we version bump
[BugFix] Fixed crash when serializing a loose array property
[BugFix] Fixed lost remove overrides when serializing a loose struct property with an overridable map

#Synced-CL 51689558

diff --git a/Engine/Source/Runtime/CoreUObject/Private/UObject/PropertyArray.cpp b/Engine/Source/Runtime/CoreUObject/Private/UObject/PropertyArray.cpp
--- a/Engine/Source/Runtime/CoreUObject/Private/UObject/PropertyArray.cpp
+++ b/Engine/Source/Runtime/CoreUObject/Private/UObject/PropertyArray.cpp
@@ -329,14 +329,20 @@ void FArrayProperty::SerializeItem(FStructuredArchive::FSlot Slot, TNotNull<void
 				checkf(InnerObjectProperty->HasAnyPropertyFlags(CPF_PersistentInstance), TEXT("Only supported code path here is the instanced subobjects"));
 
 				FOverriddenPropertySet* OverriddenProperties = FOverridableSerializationLogic::GetOverriddenProperties();
-
-				checkf(Defaults, TEXT("Expecting overridable serialization to have defaults to compare to"));
+				
+#if WITH_EDITORONLY_DATA
+				FUObjectSerializeContext* SerializeContext = FUObjectThreadContext::Get().GetSerializeContext();
+				const bool bIsLoadingToPropertyBagObject = SerializeContext->bIsLoadingToPropertyBagObject;
+#else
+				const bool bIsLoadingToPropertyBagObject = false;
+#endif
+				checkf(Defaults || bIsLoadingToPropertyBagObject, TEXT("Expecting overridable serialization to have defaults to compare to"));
 
 				FScriptArrayHelper DefaultsArrayHelper(this, Defaults);
 
 				// If the arrays are not the same size, the code will not be able to remove items correctly
 				// as it searches for the items to remove in the default array but removes them from this array.
-				checkf (bUsingExternalObjects || (DefaultsArrayHelper.Num() == ArrayHelper.Num()), TEXT("The array \'%s\' is expected to be the same size as it default for removal logic."), *GetName());
+				checkf (bUsingExternalObjects || bIsLoadingToPropertyBagObject || (DefaultsArrayHelper.Num() == ArrayHelper.Num()), TEXT("The array \'%s\' is expected to be the same size as it default for removal logic."), *GetName());
 
 				auto FindObject = [InnerObjectProperty](UObject* Object, UObject* Object2, FScriptArrayHelper& ArrayHelper) -> int32
 				{
@@ -382,14 +388,19 @@ void FArrayProperty::SerializeItem(FStructuredArchive::FSlot Slot, TNotNull<void
 
 						if (UObject* RemovedSubObject = InnerObjectProperty->GetObjectPropertyValue(TempValueStorage))
 						{
-							int32 Index = FindObject(RemovedSubObject, nullptr, DefaultsArrayHelper);
-							if (Index != INDEX_NONE)
+							// when loading into property bags we don't have a default. luckily we don't need accurate overrides for that so we can skip this step.
+							// we'll reseirialize this into the IDO later and that will get the correct override state
+							if (!bIsLoadingToPropertyBagObject)
 							{
-								IndicesToRemove.Add(Index);
-							}
-							else
-							{
-								UE_LOG(LogOverridableObject, VeryVerbose, TEXT("Unable to load removed item %s(0x%p)"), *GetNameSafe(RemovedSubObject), RemovedSubObject);
+								int32 Index = FindObject(RemovedSubObject, nullptr, DefaultsArrayHelper);
+								if (Index != INDEX_NONE)
+								{
+									IndicesToRemove.Add(Index);
+								}
+								else
+								{
+									UE_LOG(LogOverridableObject, VeryVerbose, TEXT("Unable to load removed item %s(0x%p)"), *GetNameSafe(RemovedSubObject), RemovedSubObject);
+								}
 							}
 
 							ContainerOverrides.RemovedArchetypes.Add(RemovedSubObject);
diff --git a/Engine/Source/Runtime/CoreUObject/Private/UObject/PropertyMap.cpp b/Engine/Source/Runtime/CoreUObject/Private/UObject/PropertyMap.cpp
--- a/Engine/Source/Runtime/CoreUObject/Private/UObject/PropertyMap.cpp
+++ b/Engine/Source/Runtime/CoreUObject/Private/UObject/PropertyMap.cpp
@@ -12,6 +12,7 @@
 #include "UObject/UnrealTypePrivate.h"
 #include "UObject/UObjectThreadContext.h"
 #include "UObject/CommonFNames.h"
+#include "UObject/InstanceDataObjectUtils.h"
 
 namespace UEMapProperty_Private
 {
@@ -374,7 +375,7 @@ void FMapProperty::SerializeItem(FStructuredArchive::FSlot Slot, TNotNull<void*>
 							// the lack of CDO constructor is causing an element to be missing that otherwise would've been present
 							if (!bPersistOverride)
 							{
-								bPersistOverride |= this->GetBoolMetaData(UE::NAME_IsLooseMetadata);
+								bPersistOverride |= UE::IsSerializingUnknownData(UnderlyingArchive, this);
 							}
 #endif
 							if (bPersistOverride)
