Changed some of the type/property checks performed in ConstructDefaultObjectValueIfNecessary to use the class that will be passed to object construction.

#rnx
#Synced-CL 51894126

diff --git a/Engine/Source/Runtime/CoreUObject/Private/UObject/PropertyBaseObject.cpp b/Engine/Source/Runtime/CoreUObject/Private/UObject/PropertyBaseObject.cpp
--- a/Engine/Source/Runtime/CoreUObject/Private/UObject/PropertyBaseObject.cpp
+++ b/Engine/Source/Runtime/CoreUObject/Private/UObject/PropertyBaseObject.cpp
@@ -830,11 +830,11 @@ UObject* FObjectPropertyBase::ConstructDefaultObjectValueIfNecessary(UObject* Ex
 				*OutFailureReason = TEXT("malformed property: PropertyClass=null");
 			}
 		}
-		else if (PropertyClass->HasAnyClassFlags(CLASS_Abstract))
+		else if (NewClass->HasAnyClassFlags(CLASS_Abstract))
 		{
 			if (OutFailureReason)
 			{
-				*OutFailureReason = FString::Printf(TEXT("class '%s' is abstract"), *PropertyClass->GetName());
+				*OutFailureReason = FString::Printf(TEXT("class '%s' is abstract"), *NewClass->GetName());
 			}
 		}
 #if WITH_METADATA
@@ -853,7 +853,7 @@ UObject* FObjectPropertyBase::ConstructDefaultObjectValueIfNecessary(UObject* Ex
 			// We can use the existing name if: 
 			// a) there is no existing object with that name in the outer, or
 			// b) we are replacing an object of the same type.
-			bool bCanReplaceObjectWithExistingName = !ObjectWithExistingName || ObjectWithExistingName->GetClass()->IsChildOf(PropertyClass);
+			bool bCanReplaceObjectWithExistingName = !ObjectWithExistingName || ObjectWithExistingName->GetClass()->IsChildOf(NewClass);
 
 			FName NewName = (bCanReplaceObjectWithExistingName) ? ExistingName : NAME_None;
 
@@ -874,7 +874,7 @@ UObject* FObjectPropertyBase::ConstructDefaultObjectValueIfNecessary(UObject* Ex
 				Template = GetObjectPropertyValue(Defaults);
 			}
 			
-			if (Template && !Template->IsA(PropertyClass))
+			if (Template && !Template->IsA(NewClass))
 			{
 				Template = nullptr;
 			}
