[LinkerLoad]
- Do not send event to detach linker when there was no previous linker

#Synced-CL 51934851

diff --git a/Engine/Source/Runtime/CoreUObject/Private/UObject/UObjectLinker.cpp b/Engine/Source/Runtime/CoreUObject/Private/UObject/UObjectLinker.cpp
--- a/Engine/Source/Runtime/CoreUObject/Private/UObject/UObjectLinker.cpp
+++ b/Engine/Source/Runtime/CoreUObject/Private/UObject/UObjectLinker.cpp
@@ -136,6 +136,10 @@ void UObject::SetLinker( FLinkerLoad* LinkerLoad, int32 LinkerIndex, bool bShoul
 		ExportObject.ResetObject();
 		ExportObject.bExportLoadFailed = bIsInvalid;
 	}
+	else
+	{
+		bShouldDetachExisting = false; // no previous linker, we don't need to detach anything
+	}
 
 	if (Existing.Linker == LinkerLoad)
 	{
