Performance improvements for async package loading:
- Cache hash indices to avoid redundant binary searches when handling export collisions.

#Synced-CL 51944883

diff --git a/Engine/Source/Runtime/CoreUObject/Private/Serialization/AsyncLoading2.cpp b/Engine/Source/Runtime/CoreUObject/Private/Serialization/AsyncLoading2.cpp
--- a/Engine/Source/Runtime/CoreUObject/Private/Serialization/AsyncLoading2.cpp
+++ b/Engine/Source/Runtime/CoreUObject/Private/Serialization/AsyncLoading2.cpp
@@ -7759,7 +7759,7 @@ EEventLoadNodeExecutionResult FAsyncPackage2::ProcessLinkerLoadPackageSummary(FA
 			LinkerLoadState->bContainsClasses |= true;
 		}
 
-		//if ((ObjectExport.ObjectFlags & RF_Public) > 0) // We need hashes for all objects, external actors are breaking assumptions about public exports
+		if ((ObjectExport.ObjectFlags & (RF_Public | RF_GeneratePublicHash)) > 0 || ObjectExport.bGeneratePublicHash || ObjectExport.OuterIndex.IsImport())
 		{
 			TArray<int32, TInlineAllocator<128>> FullPath;
 			int32 PathIndex = ExportIndex;
