More debug info when assigning migrated object data to a stub that already has migrated data assigned: added comparison if both data structures are identical

#prefight
#Synced-CL 51881181

diff --git a/Engine/Source/Runtime/CoreUObject/Private/UObject/RemoteObjectTransfer.cpp b/Engine/Source/Runtime/CoreUObject/Private/UObject/RemoteObjectTransfer.cpp
--- a/Engine/Source/Runtime/CoreUObject/Private/UObject/RemoteObjectTransfer.cpp
+++ b/Engine/Source/Runtime/CoreUObject/Private/UObject/RemoteObjectTransfer.cpp
@@ -19,6 +19,48 @@
 #include "Templates/Atomic.h"
 
 
+bool FRemoteObjectData::Compare(const FRemoteObjectData& InOtherData) const
+{
+	if (!Tables.Compare(InOtherData.Tables))
+	{
+		return false;
+	}
+	if (PathNames.Num() != InOtherData.PathNames.Num())
+	{
+		return false;
+	}
+	for (int32 Index = 0; Index < PathNames.Num(); ++Index)
+	{
+		if (!PathNames[Index].Compare(InOtherData.PathNames[Index]))
+		{
+			return false;
+		}
+	}
+	if (SerializedObjectHeaders.Num() != InOtherData.SerializedObjectHeaders.Num())
+	{
+		return false;
+	}
+	for (int32 Index = 0; Index < SerializedObjectHeaders.Num(); ++Index)
+	{
+		if (!SerializedObjectHeaders[Index].Compare(InOtherData.SerializedObjectHeaders[Index]))
+		{
+			return false;
+		}
+	}
+	if (Bytes.Num() != InOtherData.Bytes.Num())
+	{
+		return false;
+	}
+	for (int32 Index = 0; Index < Bytes.Num(); ++Index)
+	{
+		if (!Bytes[Index].Compare(InOtherData.Bytes[Index]))
+		{
+			return false;
+		}
+	}
+	return true;
+}
+
 namespace UE::RemoteObject::Transfer
 {
 const FRemoteServerId DatabaseId(ERemoteServerIdConstants::Database);
@@ -972,8 +1014,8 @@ TRefCountPtr<FTransactionalMigrationData> StoreMigratedObjectData(FRemoteObjectD
 		FRemoteObjectId SerializedId = It.GetId();
 		FRemoteObjectStub* Stub = FindOrAddRemoteObjectStub(SerializedId, FRemoteServerId::GetLocalServerId());
 
-		checkf(!Stub->MigratedData, TEXT("Received remote object %s (%s) from server %s bundled with %s (%s) but the latter object's data already exists locally (received from server %s and bundled with %s %s, state: %s)"), 
-			*MigratedData->Context.ObjectId.ToString(), 
+		checkf(!Stub->MigratedData, TEXT("Received remote object %s (%s) from server %s bundled with %s (%s) but the latter object's data already exists locally (received from server %s and bundled with %s %s, state: %s, data %s identical)"),
+			*MigratedData->Context.ObjectId.ToString(),
 			*MigratedData->Data.GetName(MigratedData->Context.ObjectId).ToString(),
 			*MigratedData->Context.RemoteServerId.ToString(),
 			*SerializedId.ToString(),
@@ -981,7 +1023,9 @@ TRefCountPtr<FTransactionalMigrationData> StoreMigratedObjectData(FRemoteObjectD
 			*Stub->MigratedData->Context.RemoteServerId.ToString(),
 			*Stub->MigratedData->Context.ObjectId.ToString(),
 			*Stub->MigratedData->Data.GetName(Stub->MigratedData->Context.ObjectId).ToString(),
-			EnumToString(Stub->MigratedData->State));
+			EnumToString(Stub->MigratedData->State),
+			Stub->MigratedData->Data.Compare(MigratedData->Data) ? TEXT("is") : TEXT("is NOT")
+		);
 
 		Stub->OwningServerId = MigratedData->Context.OwnerServerId;
 		Stub->MigratedData = MigratedData;
diff --git a/Engine/Source/Runtime/CoreUObject/Public/UObject/RemoteObjectPathName.h b/Engine/Source/Runtime/CoreUObject/Public/UObject/RemoteObjectPathName.h
--- a/Engine/Source/Runtime/CoreUObject/Public/UObject/RemoteObjectPathName.h
+++ b/Engine/Source/Runtime/CoreUObject/Public/UObject/RemoteObjectPathName.h
@@ -14,6 +14,14 @@ struct FRemoteObjectTables
 	TArray<FName> Names;
 	TArray<FRemoteObjectId> RemoteIds;
 
+	/**
+     * Compare function instead of operator== so that it's not used accidentally for performance reasons
+	 */
+	bool Compare(const FRemoteObjectTables& InOtherTables) const
+	{
+		return Names == InOtherTables.Names && RemoteIds == InOtherTables.RemoteIds;
+	}
+
 	UE_FORCEINLINE_HINT int32 Num() const
 	{
 		return Names.Num();
@@ -55,6 +63,14 @@ struct FPackedRemoteObjectPathName
 		return Names.Num();
 	}
 
+	/**
+     * Compare function instead of operator== so that it's not used accidentally for performance reasons
+	 */
+	bool Compare(const FPackedRemoteObjectPathName& InOtherPathName) const
+	{
+		return RemoteIds == InOtherPathName.RemoteIds && Names == InOtherPathName.Names;
+	}
+
 	UE_FORCEINLINE_HINT FName GetSegmentName(int32 InSegmentIndex, const FRemoteObjectTables& InTables) const
 	{
 		return InTables.Names[Names[InSegmentIndex]];
diff --git a/Engine/Source/Runtime/CoreUObject/Public/UObject/RemoteObjectTransfer.h b/Engine/Source/Runtime/CoreUObject/Public/UObject/RemoteObjectTransfer.h
--- a/Engine/Source/Runtime/CoreUObject/Public/UObject/RemoteObjectTransfer.h
+++ b/Engine/Source/Runtime/CoreUObject/Public/UObject/RemoteObjectTransfer.h
@@ -23,6 +23,11 @@ struct FRemoteObjectData;
 struct FRemoteObjectBytes
 {
 	TArray<uint8> Bytes;
+
+	bool Compare(const FRemoteObjectBytes& InOtherData) const
+	{
+		return Bytes == InOtherData.Bytes;
+	}
 };
 
 /**
@@ -33,6 +38,11 @@ struct FPackedRemoteObjectHeader
 	FRemoteObjectTables::FNameIndexType NameIndex = 0;
 	FRemoteObjectTables::FNameIndexType IdIndex = 0;
 	FRemoteObjectTables::FNameIndexType ClassIndex = 0;
+
+	bool Compare(const FPackedRemoteObjectHeader& InOtherHeader) const
+	{
+		return NameIndex == InOtherHeader.NameIndex && IdIndex == InOtherHeader.IdIndex && ClassIndex == InOtherHeader.ClassIndex;
+	}
 };
 
 /**
@@ -89,6 +99,11 @@ struct FRemoteObjectData
 	/** Serialized object headers and data (properties) */
 	TArray<FRemoteObjectBytes> Bytes;
 
+	/**
+	* Returns true if both structures contain identical data
+	*/
+	COREUOBJECT_API bool Compare(const FRemoteObjectData& InOtherData) const;
+
 	/**
 	* Returns the size of serialized data
 	*/
