Don't allow remote object migrations in PreGC delegate (we already don't allow them during the entire GC process)

#Synced-CL 51881181

diff --git a/Engine/Source/Runtime/CoreUObject/Private/UObject/GarbageCollection.cpp b/Engine/Source/Runtime/CoreUObject/Private/UObject/GarbageCollection.cpp
--- a/Engine/Source/Runtime/CoreUObject/Private/UObject/GarbageCollection.cpp
+++ b/Engine/Source/Runtime/CoreUObject/Private/UObject/GarbageCollection.cpp
@@ -5728,6 +5728,8 @@ AUTORTFM_DISABLE void PreCollectGarbageImpl(EObjectFlags KeepFlags)
 			GGCStats.bFlushedAsyncLoading = true;
 		}
 		{
+			UE::RemoteObject::Private::FUnsafeToMigrateScope UnsafeToMigrate;
+
 			TRACE_CPUPROFILER_EVENT_SCOPE(BroadcastPreGarbageCollect);
 			FCoreUObjectDelegates::GetPreGarbageCollectDelegate().Broadcast();
 		}
