Adjust CheckObjectValidBasedOnItsFlags to be consistent with other validity check such as TWeakObjectPtr::Get

#Synced-CL 51808504

diff --git a/Engine/Source/Runtime/CoreUObject/Public/UObject/Object.h b/Engine/Source/Runtime/CoreUObject/Public/UObject/Object.h
--- a/Engine/Source/Runtime/CoreUObject/Public/UObject/Object.h
+++ b/Engine/Source/Runtime/CoreUObject/Public/UObject/Object.h
@@ -1855,9 +1855,10 @@ struct FInternalUObjectBaseUtilityIsValidFlagsChecker
 {
 	inline static bool CheckObjectValidBasedOnItsFlags(const UObject* Test)
 	{
+		FUObjectItem* ObjectItem = GUObjectArray.IndexToObject(Test->InternalIndex);
 		// Here we don't really check if the flags match but if the end result is the same
-		checkSlow(GUObjectArray.IndexToObject(Test->InternalIndex)->HasAnyFlags(EInternalObjectFlags::Garbage) == Test->HasAnyFlags(RF_MirroredGarbage));
-		return !Test->HasAnyFlags(RF_MirroredGarbage);
+		checkSlow(ObjectItem->HasAnyFlags(EInternalObjectFlags::Garbage) == Test->HasAnyFlags(RF_MirroredGarbage));
+		return !ObjectItem->HasAnyFlags(EInternalObjectFlags::Garbage | EInternalObjectFlags::Unreachable);
 	}
 };
 
@@ -1880,7 +1881,7 @@ private:
 * Test validity of object
 *
 * @param	Test			The object to test
-* @return	Return true if the object is usable: non-null and not pending kill or garbage
+* @return	Return true if the object is usable: non-null, reachable and not pending kill or garbage
 */
 UE_FORCEINLINE_HINT bool IsValid(const UObject *Test)
 {
@@ -1891,7 +1892,7 @@ UE_FORCEINLINE_HINT bool IsValid(const UObject *Test)
 * Test validity of object similar to IsValid(Test) however the null pointer test is skipped
 *
 * @param	Test			The object to test
-* @return	Return true if the object is usable: not pending kill or garbage
+* @return	Return true if the object is usable: reachable and not pending kill or garbage
 */
 inline bool IsValidChecked(const UObject* Test)
 {
