[ZenLoader]
- Fix race between IncrementalPurgeGarbage and FlushDeferredLoadedObjectRenameQueue that could happen if we didn't remove unreachable objects as the first operation after resuming ALT after a GC.

#Synced-CL 51934851

diff --git a/Engine/Source/Runtime/CoreUObject/Private/Serialization/AsyncLoading2.cpp b/Engine/Source/Runtime/CoreUObject/Private/Serialization/AsyncLoading2.cpp
--- a/Engine/Source/Runtime/CoreUObject/Private/Serialization/AsyncLoading2.cpp
+++ b/Engine/Source/Runtime/CoreUObject/Private/Serialization/AsyncLoading2.cpp
@@ -11355,13 +11355,15 @@ uint32 FAsyncLoadingThread2::Run()
 					};
 #endif
 
-					FlushDeferredLoadedObjectRenameQueue();
-
+					// The first thing we need to do is cleanup the global import table from unreachable objects
+					// to avoid racing with the IncrementalPurgeGarbage that happens outside of the GC guard.
 					{
 						FScopeLock UnreachableObjectsLock(&UnreachableObjectsCritical);
 						RemoveUnreachableObjects(UnreachableObjects);
 					}
 
+					FlushDeferredLoadedObjectRenameQueue();
+
 					if (bShouldWaitForExternalReads)
 					{
 						TRACE_CPUPROFILER_EVENT_SCOPE(WaitingForExternalReads);
