Correctly handling interface instantiation checks in ConstructDefaultObjectValueIfNecessary.

#rnx

#Synced-CL 52058114

diff --git a/Engine/Source/Runtime/CoreUObject/Private/UObject/PropertyBaseObject.cpp b/Engine/Source/Runtime/CoreUObject/Private/UObject/PropertyBaseObject.cpp
--- a/Engine/Source/Runtime/CoreUObject/Private/UObject/PropertyBaseObject.cpp
+++ b/Engine/Source/Runtime/CoreUObject/Private/UObject/PropertyBaseObject.cpp
@@ -790,6 +790,7 @@ UObject* FObjectPropertyBase::ConstructDefaultObjectValueIfNecessary(UObject* Ex
 	UObject* Template = nullptr;
 	UClass* NewClass = PropertyClass;
 	UObject* NewDefaultObjectValue = nullptr;
+	bool bCanPropertyHoldExistingValue = false;
 
 	FUObjectSerializeContext* SerializeContext = FUObjectThreadContext::Get().GetSerializeContext();
 	UObject* Outer = SerializeContext ? SerializeContext->SerializedObject : nullptr;
@@ -817,12 +818,18 @@ UObject* FObjectPropertyBase::ConstructDefaultObjectValueIfNecessary(UObject* Ex
 				// Similar to the above condition but the property was not referencing an instanced value in which case it's ok to leave the CDO default here
 				NewDefaultObjectValue = ExistingValue;
 			}
+
+			bCanPropertyHoldExistingValue = true;
 		}
 	}
 
 	// We cannot know which object to create if the class is abstract or if there is a required interface
 	if (!NewDefaultObjectValue)
 	{
+#if WITH_METADATA
+		UClass* MetaDataClass = this->GetOwnerProperty()->GetClassMetaData(TEXT("ObjectMustImplement"));
+#endif
+
 		if (!ensureMsgf(PropertyClass, TEXT("Malformed %s: PropertyClass=null."), *GetFullName()))
 		{
 			if (OutFailureReason)
@@ -838,7 +845,9 @@ UObject* FObjectPropertyBase::ConstructDefaultObjectValueIfNecessary(UObject* Ex
 			}
 		}
 #if WITH_METADATA
-		else if (UClass* MetaDataClass = this->GetOwnerProperty()->GetClassMetaData(TEXT("ObjectMustImplement")))
+		// CanPropertyHoldObject has already determined that we can instantiate an object with NewClass so when bCanPropertyHoldExistingValue
+		// is true we bypass the check for 'ObjectMustImplement'. 
+		else if (!bCanPropertyHoldExistingValue && MetaDataClass)
 		{
 			if (OutFailureReason)
 			{
