[WorldPartition] Fix Detaching from existing linker

#tests Cooked offending island, preflight editor tests/cook
#Synced-CL 52011392

diff --git a/Engine/Source/Runtime/CoreUObject/Private/UObject/UObjectGlobals.cpp b/Engine/Source/Runtime/CoreUObject/Private/UObject/UObjectGlobals.cpp
--- a/Engine/Source/Runtime/CoreUObject/Private/UObject/UObjectGlobals.cpp
+++ b/Engine/Source/Runtime/CoreUObject/Private/UObject/UObjectGlobals.cpp
@@ -41,6 +41,7 @@
 #include "UObject/PropertyStateTracking.h"
 #include "UObject/UObjectIterator.h"
 #include "UObject/Package.h"
+#include "UObject/Linker.h"
 #include "Templates/Casts.h"
 #include "UObject/PropertyPortFlags.h"
 #include "Serialization/SerializedPropertyScope.h"
@@ -265,6 +266,14 @@ void TrashObject(TNotNull<UObject* const> Object)
 
 		const FName BaseName = bIsAPackage ? NAME_TrashedPackage : NAME_None;
 
+		if (bIsAPackage)
+		{
+			// Trashing package hazard:
+			// A later package load might contain imports to objects originally in this trashed package, but renamed to a different package. CreateExport/FindExistingExport will find them.
+			// These objects might still have the stale linker set if GC hasn't finished UPackage::FinishDestroy yet. SetLinker will try to replace the existing stale linker, and error out with "Detaching from existing linker...".
+			ResetLoaders(Object);
+		}
+
 		const FName NewName = MakeUniqueObjectName(Object->GetOuter(), Object->GetClass(), BaseName);
 		Object->Rename(*NewName.ToString(), nullptr, REN_DontCreateRedirectors | REN_NonTransactional | REN_DoNotDirty);
 
