[UE] Additional test cases/ground truth for transient object properties and instancing.

#rnx

#Synced-CL 52049047

diff --git a/Engine/Source/Runtime/CoreUObject/Tests/UObject/InstancedPropertiesTest.cpp b/Engine/Source/Runtime/CoreUObject/Tests/UObject/InstancedPropertiesTest.cpp
--- a/Engine/Source/Runtime/CoreUObject/Tests/UObject/InstancedPropertiesTest.cpp
+++ b/Engine/Source/Runtime/CoreUObject/Tests/UObject/InstancedPropertiesTest.cpp
@@ -79,6 +79,43 @@ TEST_CASE_NAMED(FInstancedPropertiesTest, "CoreUObject::InstancedProperties", "[
 #endif
 	}
 
+	SECTION("DuplicateTransient instanced object property")
+	{
+		const FName DuplicateTransientInstancedObjectPropertyName = GET_MEMBER_NAME_CHECKED(USubobjectInstancingTestOuterObject, DuplicateTransientInnerObjectUsingNew);
+		const FProperty* DuplicateTransientInstancedObjectProperty = USubobjectInstancingTestOuterObject::StaticClass()->FindPropertyByName(DuplicateTransientInstancedObjectPropertyName);
+
+		REQUIRE(DuplicateTransientInstancedObjectProperty != nullptr);
+		CHECK(DuplicateTransientInstancedObjectProperty->IsA<FObjectProperty>());
+
+		CHECK(DuplicateTransientInstancedObjectProperty->ContainsInstancedObjectProperty());
+		CHECK_FALSE(DuplicateTransientInstancedObjectProperty->HasAllPropertyFlags(CPF_Transient));
+		CHECK(DuplicateTransientInstancedObjectProperty->HasAllPropertyFlags(CPF_DuplicateTransient));
+		CHECK(DuplicateTransientInstancedObjectProperty->HasAllPropertyFlags(CPF_InstancedReference));
+		CHECK(DuplicateTransientInstancedObjectProperty->HasAllPropertyFlags(CPF_PersistentInstance));
+#if WITH_METADATA
+		CHECK(DuplicateTransientInstancedObjectProperty->GetBoolMetaData(TEXT("EditInline")));
+#endif
+	}
+
+	SECTION("NonPIEDuplicateTransient instanced object property")
+	{
+		const FName NonPIEDuplicateTransientInstancedObjectPropertyName = GET_MEMBER_NAME_CHECKED(USubobjectInstancingTestOuterObject, NonPIEDuplicateTransientInnerObjectUsingNew);
+		const FProperty* NonPIEDuplicateTransientInstancedObjectProperty = USubobjectInstancingTestOuterObject::StaticClass()->FindPropertyByName(NonPIEDuplicateTransientInstancedObjectPropertyName);
+
+		REQUIRE(NonPIEDuplicateTransientInstancedObjectProperty != nullptr);
+		CHECK(NonPIEDuplicateTransientInstancedObjectProperty->IsA<FObjectProperty>());
+
+		CHECK(NonPIEDuplicateTransientInstancedObjectProperty->ContainsInstancedObjectProperty());
+		CHECK_FALSE(NonPIEDuplicateTransientInstancedObjectProperty->HasAllPropertyFlags(CPF_Transient));
+		CHECK_FALSE(NonPIEDuplicateTransientInstancedObjectProperty->HasAllPropertyFlags(CPF_DuplicateTransient));
+		CHECK(NonPIEDuplicateTransientInstancedObjectProperty->HasAllPropertyFlags(CPF_NonPIEDuplicateTransient));
+		CHECK(NonPIEDuplicateTransientInstancedObjectProperty->HasAllPropertyFlags(CPF_InstancedReference));
+		CHECK(NonPIEDuplicateTransientInstancedObjectProperty->HasAllPropertyFlags(CPF_PersistentInstance));
+#if WITH_METADATA
+		CHECK(NonPIEDuplicateTransientInstancedObjectProperty->GetBoolMetaData(TEXT("EditInline")));
+#endif
+	}
+
 	SECTION("Transient type instanced object property")
 	{
 		const FName TransientTypeInstancedObjectPropertyName = GET_MEMBER_NAME_CHECKED(USubobjectInstancingTestOuterObject, TransientInnerObjectFromType);
@@ -158,6 +195,34 @@ TEST_CASE_NAMED(FInstancedPropertiesTest, "CoreUObject::InstancedProperties", "[
 #endif
 	}
 
+	SECTION("Transient instanced object array property")
+	{
+		const FName TransientInstancedObjectArrayPropertyName = GET_MEMBER_NAME_CHECKED(USubobjectInstancingTestOuterObject, TransientInnerObjectArray);
+		const FArrayProperty* TransientInstancedObjectArrayProperty = CastField<FArrayProperty>(USubobjectInstancingTestOuterObject::StaticClass()->FindPropertyByName(TransientInstancedObjectArrayPropertyName));
+
+		REQUIRE(TransientInstancedObjectArrayProperty != nullptr);
+		CHECK(TransientInstancedObjectArrayProperty->ContainsInstancedObjectProperty());
+		CHECK(TransientInstancedObjectArrayProperty->HasAllPropertyFlags(CPF_ContainsInstancedReference));
+		CHECK(TransientInstancedObjectArrayProperty->HasAllPropertyFlags(CPF_Transient));
+		CHECK_FALSE(TransientInstancedObjectArrayProperty->HasAllPropertyFlags(CPF_InstancedReference));
+		CHECK_FALSE(TransientInstancedObjectArrayProperty->HasAllPropertyFlags(CPF_PersistentInstance));
+#if WITH_METADATA
+		CHECK(TransientInstancedObjectArrayProperty->GetBoolMetaData(TEXT("EditInline")));
+#endif
+
+		REQUIRE(TransientInstancedObjectArrayProperty->Inner != nullptr);
+		CHECK(TransientInstancedObjectArrayProperty->Inner->IsA<FObjectProperty>());
+
+		CHECK(TransientInstancedObjectArrayProperty->Inner->ContainsInstancedObjectProperty());
+		CHECK_FALSE(TransientInstancedObjectArrayProperty->Inner->HasAllPropertyFlags(CPF_ContainsInstancedReference));
+		CHECK_FALSE(TransientInstancedObjectArrayProperty->Inner->HasAllPropertyFlags(CPF_Transient));
+		CHECK(TransientInstancedObjectArrayProperty->Inner->HasAllPropertyFlags(CPF_InstancedReference));
+		CHECK(TransientInstancedObjectArrayProperty->Inner->HasAllPropertyFlags(CPF_PersistentInstance));
+#if WITH_METADATA
+		CHECK(TransientInstancedObjectArrayProperty->Inner->GetBoolMetaData(TEXT("EditInline")));
+#endif
+	}
+
 	SECTION("Instanced object set container property")
 	{
 		const FName InstancedObjectSetPropertyName = GET_MEMBER_NAME_CHECKED(USubobjectInstancingTestOuterObject, InnerObjectSet);
@@ -212,6 +277,34 @@ TEST_CASE_NAMED(FInstancedPropertiesTest, "CoreUObject::InstancedProperties", "[
 #endif
 	}
 
+	SECTION("Transient instanced object set container property")
+	{
+		const FName TransientInstancedObjectSetPropertyName = GET_MEMBER_NAME_CHECKED(USubobjectInstancingTestOuterObject, TransientInnerObjectSet);
+		const FSetProperty* TransientInstancedObjectSetProperty = CastField<FSetProperty>(USubobjectInstancingTestOuterObject::StaticClass()->FindPropertyByName(TransientInstancedObjectSetPropertyName));
+
+		REQUIRE(TransientInstancedObjectSetProperty != nullptr);
+		CHECK(TransientInstancedObjectSetProperty->ContainsInstancedObjectProperty());
+		CHECK(TransientInstancedObjectSetProperty->HasAllPropertyFlags(CPF_ContainsInstancedReference));
+		CHECK(TransientInstancedObjectSetProperty->HasAllPropertyFlags(CPF_Transient));
+		CHECK_FALSE(TransientInstancedObjectSetProperty->HasAllPropertyFlags(CPF_InstancedReference));
+		CHECK_FALSE(TransientInstancedObjectSetProperty->HasAllPropertyFlags(CPF_PersistentInstance));
+#if WITH_METADATA
+		CHECK(TransientInstancedObjectSetProperty->GetBoolMetaData(TEXT("EditInline")));
+#endif
+
+		REQUIRE(TransientInstancedObjectSetProperty->ElementProp != nullptr);
+		CHECK(TransientInstancedObjectSetProperty->ElementProp->IsA<FObjectProperty>());
+
+		CHECK(TransientInstancedObjectSetProperty->ElementProp->ContainsInstancedObjectProperty());
+		CHECK_FALSE(TransientInstancedObjectSetProperty->ElementProp->HasAllPropertyFlags(CPF_ContainsInstancedReference));
+		CHECK_FALSE(TransientInstancedObjectSetProperty->ElementProp->HasAllPropertyFlags(CPF_Transient));
+		CHECK(TransientInstancedObjectSetProperty->ElementProp->HasAllPropertyFlags(CPF_InstancedReference));
+		CHECK(TransientInstancedObjectSetProperty->ElementProp->HasAllPropertyFlags(CPF_PersistentInstance));
+#if WITH_METADATA
+		CHECK(TransientInstancedObjectSetProperty->ElementProp->GetBoolMetaData(TEXT("EditInline")));
+#endif
+	}
+
 	SECTION("Instanced object map container property")
 	{
 		const FName InstancedObjectMapPropertyName = GET_MEMBER_NAME_CHECKED(USubobjectInstancingTestOuterObject, InnerObjectMap);
@@ -291,6 +384,48 @@ TEST_CASE_NAMED(FInstancedPropertiesTest, "CoreUObject::InstancedProperties", "[
 #endif
 	}
 
+	SECTION("Transient instanced object map container property")
+	{
+		const FName TransientInstancedObjectMapPropertyName = GET_MEMBER_NAME_CHECKED(USubobjectInstancingTestOuterObject, TransientInnerObjectMap);
+		const FMapProperty* TransientInstancedObjectMapProperty = CastField<FMapProperty>(USubobjectInstancingTestOuterObject::StaticClass()->FindPropertyByName(TransientInstancedObjectMapPropertyName));
+
+		REQUIRE(TransientInstancedObjectMapProperty != nullptr);
+		CHECK(TransientInstancedObjectMapProperty->ContainsInstancedObjectProperty());
+		CHECK(TransientInstancedObjectMapProperty->HasAllPropertyFlags(CPF_ContainsInstancedReference));
+		CHECK(TransientInstancedObjectMapProperty->HasAllPropertyFlags(CPF_Transient));
+		CHECK_FALSE(TransientInstancedObjectMapProperty->HasAllPropertyFlags(CPF_InstancedReference));
+		CHECK_FALSE(TransientInstancedObjectMapProperty->HasAllPropertyFlags(CPF_PersistentInstance));
+#if WITH_METADATA
+		// Note: 'EditInline' is set via UHT.
+		CHECK(TransientInstancedObjectMapProperty->GetBoolMetaData(TEXT("EditInline")));
+#endif
+
+		REQUIRE(TransientInstancedObjectMapProperty->KeyProp != nullptr);
+		CHECK(TransientInstancedObjectMapProperty->KeyProp->IsA<FObjectProperty>());
+
+		CHECK(TransientInstancedObjectMapProperty->KeyProp->ContainsInstancedObjectProperty());
+		CHECK_FALSE(TransientInstancedObjectMapProperty->KeyProp->HasAllPropertyFlags(CPF_ContainsInstancedReference));
+		CHECK_FALSE(TransientInstancedObjectMapProperty->KeyProp->HasAllPropertyFlags(CPF_Transient));
+		CHECK(TransientInstancedObjectMapProperty->KeyProp->HasAllPropertyFlags(CPF_InstancedReference));
+		CHECK(TransientInstancedObjectMapProperty->KeyProp->HasAllPropertyFlags(CPF_PersistentInstance));
+#if WITH_METADATA
+		// Note: Currently the KeyProp is not also flagged as 'EditInline' when the map is declared as 'Instanced'.
+		CHECK_FALSE(TransientInstancedObjectMapProperty->KeyProp->GetBoolMetaData(TEXT("EditInline")));
+#endif
+
+		REQUIRE(TransientInstancedObjectMapProperty->ValueProp != nullptr);
+		CHECK(TransientInstancedObjectMapProperty->ValueProp->IsA<FObjectProperty>());
+
+		CHECK(TransientInstancedObjectMapProperty->ValueProp->ContainsInstancedObjectProperty());
+		CHECK_FALSE(TransientInstancedObjectMapProperty->ValueProp->HasAllPropertyFlags(CPF_ContainsInstancedReference));
+		CHECK_FALSE(TransientInstancedObjectMapProperty->ValueProp->HasAllPropertyFlags(CPF_Transient));
+		CHECK(TransientInstancedObjectMapProperty->ValueProp->HasAllPropertyFlags(CPF_InstancedReference));
+		CHECK(TransientInstancedObjectMapProperty->ValueProp->HasAllPropertyFlags(CPF_PersistentInstance));
+#if WITH_METADATA
+		CHECK(TransientInstancedObjectMapProperty->ValueProp->GetBoolMetaData(TEXT("EditInline")));
+#endif
+	}
+
 	SECTION("Struct with instanced object property")
 	{
 		const FName StructWithInstancedObjectPropertyName = GET_MEMBER_NAME_CHECKED(USubobjectInstancingTestOuterObject, StructWithInnerObjects);
@@ -446,6 +581,35 @@ TEST_CASE_NAMED(FInstancedPropertiesTest, "CoreUObject::InstancedProperties", "[
 #endif
 	}
 
+	SECTION("Optional transient instanced object property")
+	{
+		const FName OptionalTransientInstancedObjectPropertyName = GET_MEMBER_NAME_CHECKED(USubobjectInstancingTestOuterObject, OptionalTransientInnerObject);
+		const FOptionalProperty* OptionalTransientInstancedObjectProperty = CastField<FOptionalProperty>(USubobjectInstancingTestOuterObject::StaticClass()->FindPropertyByName(OptionalTransientInstancedObjectPropertyName));
+
+		REQUIRE(OptionalTransientInstancedObjectProperty != nullptr);
+		CHECK(OptionalTransientInstancedObjectProperty->ContainsInstancedObjectProperty());
+		CHECK(OptionalTransientInstancedObjectProperty->HasAllPropertyFlags(CPF_ContainsInstancedReference));
+		CHECK(OptionalTransientInstancedObjectProperty->HasAllPropertyFlags(CPF_Transient));
+		CHECK_FALSE(OptionalTransientInstancedObjectProperty->HasAllPropertyFlags(CPF_InstancedReference));
+		CHECK_FALSE(OptionalTransientInstancedObjectProperty->HasAllPropertyFlags(CPF_PersistentInstance));
+#if WITH_METADATA
+		CHECK(OptionalTransientInstancedObjectProperty->GetBoolMetaData(TEXT("EditInline")));
+#endif
+
+		const FProperty* OptionalTransientInstancedObjectValueProperty = OptionalTransientInstancedObjectProperty->GetValueProperty();
+		REQUIRE(OptionalTransientInstancedObjectValueProperty != nullptr);
+		CHECK(OptionalTransientInstancedObjectValueProperty->IsA<FObjectProperty>());
+
+		CHECK(OptionalTransientInstancedObjectValueProperty->ContainsInstancedObjectProperty());
+		CHECK_FALSE(OptionalTransientInstancedObjectValueProperty->HasAllPropertyFlags(CPF_ContainsInstancedReference));
+		CHECK_FALSE(OptionalTransientInstancedObjectValueProperty->HasAllPropertyFlags(CPF_Transient));
+		CHECK(OptionalTransientInstancedObjectValueProperty->HasAllPropertyFlags(CPF_InstancedReference));
+		CHECK(OptionalTransientInstancedObjectValueProperty->HasAllPropertyFlags(CPF_PersistentInstance));
+#if WITH_METADATA
+		CHECK(OptionalTransientInstancedObjectValueProperty->GetBoolMetaData(TEXT("EditInline")));
+#endif
+	}
+
 	SECTION("Optional instanced object array property")
 	{
 		const FName OptionalInstancedObjectArrayPropertyName = GET_MEMBER_NAME_CHECKED(USubobjectInstancingTestOuterObject, OptionalInnerObjectArray);
diff --git a/Engine/Source/Runtime/CoreUObject/Tests/UObject/SubobjectInstancingTest.cpp b/Engine/Source/Runtime/CoreUObject/Tests/UObject/SubobjectInstancingTest.cpp
--- a/Engine/Source/Runtime/CoreUObject/Tests/UObject/SubobjectInstancingTest.cpp
+++ b/Engine/Source/Runtime/CoreUObject/Tests/UObject/SubobjectInstancingTest.cpp
@@ -66,10 +66,17 @@ USubobjectInstancingTestOuterObject::USubobjectInstancingTestOuterObject(const F
 	// Instanced subobject constructed via NewObject() that should not be serialized during duplication, but is non-transient during save
 	DuplicateTransientInnerObjectUsingNew = NewObject<USubobjectInstancingTestObject>(this, TEXT("DuplicateTransientInnerObjectUsingNew"));
 
+	// Instanced subobject constructed via NewObject() that should not be serialized during duplication outside of PIE, but is non-transient during save
+	NonPIEDuplicateTransientInnerObjectUsingNew = NewObject<USubobjectInstancingTestObject>(this, TEXT("NonPIEDuplicateTransientInnerObjectUsingNew"));
+
 	// Array container of references to instanced subobjects
 	InnerObjectArray.Add(CreateDefaultSubobject<USubobjectInstancingTestObject>(TEXT("InnerObjectArrayElem_0")));
 	InnerObjectArray.Add(CreateDefaultSubobject<USubobjectInstancingTestObject>(TEXT("InnerObjectArrayElem_1")));
 
+	// Transient array container of references to instanced subobjects (CPF_Transient on the property; elements are not RF_Transient)
+	TransientInnerObjectArray.Add(CreateDefaultSubobject<USubobjectInstancingTestObject>(TEXT("TransientInnerObjectArrayElem_0")));
+	TransientInnerObjectArray.Add(CreateDefaultSubobject<USubobjectInstancingTestObject>(TEXT("TransientInnerObjectArrayElem_1")));
+
 	// Array container of references to DefaultToInstanced-typed subobjects
 	InnerObjectFromTypeArray.Add(CreateDefaultSubobject<USubobjectInstancingDefaultToInstancedTestObject>(TEXT("InnerObjectFromTypeArrayElem_0")));
 	InnerObjectFromTypeArray.Add(CreateDefaultSubobject<USubobjectInstancingDefaultToInstancedTestObject>(TEXT("InnerObjectFromTypeArrayElem_1")));
@@ -78,6 +85,10 @@ USubobjectInstancingTestOuterObject::USubobjectInstancingTestOuterObject(const F
 	InnerObjectSet.Add(CreateDefaultSubobject<USubobjectInstancingTestObject>(TEXT("InnerObjectSetElem_0")));
 	InnerObjectSet.Add(CreateDefaultSubobject<USubobjectInstancingTestObject>(TEXT("InnerObjectSetElem_1")));
 
+	// Transient set container of references to instanced subobjects (CPF_Transient on the property; elements are not RF_Transient)
+	TransientInnerObjectSet.Add(CreateDefaultSubobject<USubobjectInstancingTestObject>(TEXT("TransientInnerObjectSetElem_0")));
+	TransientInnerObjectSet.Add(CreateDefaultSubobject<USubobjectInstancingTestObject>(TEXT("TransientInnerObjectSetElem_1")));
+
 	// Set container of references to DefaultToInstanced-typed subobjects
 	InnerObjectFromTypeSet.Add(CreateDefaultSubobject<USubobjectInstancingDefaultToInstancedTestObject>(TEXT("InnerObjectFromTypeSetElem_0")));
 	InnerObjectFromTypeSet.Add(CreateDefaultSubobject<USubobjectInstancingDefaultToInstancedTestObject>(TEXT("InnerObjectFromTypeSetElem_1")));
@@ -90,6 +101,10 @@ USubobjectInstancingTestOuterObject::USubobjectInstancingTestOuterObject(const F
 	InnerObjectFromTypeMap.Add(CreateDefaultSubobject<USubobjectInstancingDefaultToInstancedTestObject>(TEXT("InnerObjectFromTypeMapKey_0")), CreateDefaultSubobject<USubobjectInstancingDefaultToInstancedTestObject>(TEXT("InnerObjectFromTypeMapVal_0")));
 	InnerObjectFromTypeMap.Add(CreateDefaultSubobject<USubobjectInstancingDefaultToInstancedTestObject>(TEXT("InnerObjectFromTypeMapKey_1")), CreateDefaultSubobject<USubobjectInstancingDefaultToInstancedTestObject>(TEXT("InnerObjectFromTypeMapVal_1")));
 
+	// Transient map container of pairs of references to instanced subobjects (CPF_Transient on the property; elements are not RF_Transient)
+	TransientInnerObjectMap.Add(CreateDefaultSubobject<USubobjectInstancingTestObject>(TEXT("TransientInnerObjectMapKey_0")), CreateDefaultSubobject<USubobjectInstancingTestObject>(TEXT("TransientInnerObjectMapVal_0")));
+	TransientInnerObjectMap.Add(CreateDefaultSubobject<USubobjectInstancingTestObject>(TEXT("TransientInnerObjectMapKey_1")), CreateDefaultSubobject<USubobjectInstancingTestObject>(TEXT("TransientInnerObjectMapVal_1")));
+
 	// Struct container of references to instanced subobjects
 	StructWithInnerObjects.InnerObject = CreateDefaultSubobject<USubobjectInstancingTestObject>(TEXT("InnerObjectForStruct"));
 	StructWithInnerObjects.InnerObjectFromType = CreateDefaultSubobject<USubobjectInstancingDefaultToInstancedTestObject>(TEXT("InnerObjectFromTypeForStruct"));
@@ -100,6 +115,9 @@ USubobjectInstancingTestOuterObject::USubobjectInstancingTestOuterObject(const F
 	// Note that this internally sets the 'bIsRequired' parameter to false, which could also be done here instead
 	OptionalInnerObject = CreateOptionalDefaultSubobject<USubobjectInstancingTestObject>(TEXT("OptionalInnerObject"));
 
+	// Optional transient reference to instanced subobject (CPF_Transient on the property; element is not RF_Transient)
+	OptionalTransientInnerObject = CreateOptionalDefaultSubobject<USubobjectInstancingTestObject>(TEXT("OptionalTransientInnerObject"));
+
 	// Optional reference to DefaultToInstanced subobject (with instanced property semantics inferred from ptr type)
 	OptionalInnerObjectFromType = CreateDefaultSubobject<USubobjectInstancingDefaultToInstancedTestObject>(TEXT("OptionalInnerObjectFromType"));
 
@@ -134,6 +152,9 @@ USubobjectInstancingTestDirectlyNestedObject::USubobjectInstancingTestDirectlyNe
 
 	// Nested instanced subobject that is expected to be a unique ptr value after construction
 	InnerObject = CreateDefaultSubobject<USubobjectInstancingTestObject>(TEXT("InnerObject"));
+
+	// Nested subobject that may be instanced at edit time rather than at construction time; this is NOT considered a default subobject
+	EditTimeInnerObject = nullptr;
 }
 
 USubobjectInstancingTestIndirectlyNestedObject::USubobjectInstancingTestIndirectlyNestedObject()
@@ -168,7 +189,9 @@ namespace SubobjectInstancingTest::Private
 	static const FName NonNativeInnerObjectPropertyName("NonNativeInnerObject");
 	static const FName NonNativeOuterObjectPropertyName("NonNativeOwnerObject");
 	static const FName NonNativeEditTimeInnerObjectPropertyName("NonNativeEditTimeInnerObject");
+	static const FName NonNativeDirectlyNestedInnerObjectPropertyName("NonNativeDirectlyNestedInnerObject");
 	static const FName NonNativeTransientInnerObjectPropertyName("NonNativeTransientInnerObject");
+	static const FName NonNativeTransientEditTimeInnerObjectPropertyName("NonNativeTransientEditTimeInnerObject");
 	static const FName NonNullableInnerObjectPropertyName("NonNullableInnerObject");
 	static const FName NonNullableTransientInnerObjectPropertyName("NonNullableTransientInnerObject");
 	static const FName NonNullableOptionalInnerObjectPropertyName("OptionalNonNullableInnerObject");
@@ -217,6 +240,12 @@ namespace SubobjectInstancingTest::Private
 		NonNativeInnerObjectProperty->SetPropertyClass(USubobjectInstancingTestObject::StaticClass());
 		PropertyListBuilder.AppendNoTerminate(*NonNativeInnerObjectProperty);
 
+		// Non-native object reference property to store a non-native instantiation of a directly-nested subobject - TObjectPtr<USubobjectInstancingTestDirectlyNestedObject>
+		FObjectProperty* NonNativeDirectlyNestedInnerObjectProperty = new FObjectProperty(TestClass, NonNativeDirectlyNestedInnerObjectPropertyName);
+		NonNativeDirectlyNestedInnerObjectProperty->SetPropertyFlags(NonNativeObjectPropertyFlags);
+		NonNativeDirectlyNestedInnerObjectProperty->SetPropertyClass(USubobjectInstancingTestDirectlyNestedObject::StaticClass());
+		PropertyListBuilder.AppendNoTerminate(*NonNativeDirectlyNestedInnerObjectProperty);
+
 		// Non-native object property to store a non-native reference to an outer object - TObjectPtr<UObject>
 		FObjectProperty* NonNativeOuterObjectProperty = new FObjectProperty(TestClass, NonNativeOuterObjectPropertyName);
 		NonNativeOuterObjectProperty->SetPropertyFlags(NonNativeObjectPropertyFlags);
@@ -235,6 +264,12 @@ namespace SubobjectInstancingTest::Private
 		NonNativeTransientInnerObjectProperty->SetPropertyClass(USubobjectInstancingTestObject::StaticClass());
 		PropertyListBuilder.AppendNoTerminate(*NonNativeTransientInnerObjectProperty);
 
+		// Non-native transient object reference property to simulate deferred edit-time instantiation - TObjectPtr<USubobjectInstancingTestObject>
+		FObjectProperty* NonNativeTransientEditTimeInnerObjectProperty = new FObjectProperty(TestClass, NonNativeTransientEditTimeInnerObjectPropertyName);
+		NonNativeTransientEditTimeInnerObjectProperty->SetPropertyFlags(NonNativeObjectPropertyFlags | CPF_Transient);
+		NonNativeTransientEditTimeInnerObjectProperty->SetPropertyClass(USubobjectInstancingTestObject::StaticClass());
+		PropertyListBuilder.AppendNoTerminate(*NonNativeTransientEditTimeInnerObjectProperty);
+
 		// Non-nullable object reference property to store a non-native instantiation by ctor - TNonNullPtr<USubobjectInstancingTestObject>
 		EPropertyFlags NonNullableObjectPropertyFlags = (NonNativeObjectPropertyFlags & ~CPF_TObjectPtrWrapper) | CPF_NonNullable;
 		FObjectProperty* NonNullableInnerObjectProperty = new FObjectProperty(TestClass, NonNullableInnerObjectPropertyName);
@@ -275,6 +310,11 @@ namespace SubobjectInstancingTest::Private
 		UObject* NonNativeInnerDefaultSubobject = NewObject<USubobjectInstancingTestObject>(TestCDO, NonNativeInnerObjectName, RF_ArchetypeObject | TestCDO->GetMaskedFlags(RF_PropagateToSubObjects));
 		NonNativeInnerObjectProperty->SetObjectPropertyValue_InContainer(TestCDO, NonNativeInnerDefaultSubobject);
 
+		// Same as above, but typed to USubobjectInstancingTestDirectlyNestedObject to allow access to its nested instanced properties
+		static const FName NonNativeDirectlyNestedInnerObjectName("NonNativeDirectlyNestedInnerObject");
+		USubobjectInstancingTestDirectlyNestedObject* NonNativeDirectlyNestedInnerDefaultSubobject = NewObject<USubobjectInstancingTestDirectlyNestedObject>(TestCDO, NonNativeDirectlyNestedInnerObjectName, RF_ArchetypeObject | TestCDO->GetMaskedFlags(RF_PropagateToSubObjects));
+		NonNativeDirectlyNestedInnerObjectProperty->SetObjectPropertyValue_InContainer(TestCDO, NonNativeDirectlyNestedInnerDefaultSubobject);
+
 		// Same as above, but assigned to a non-native transient reference property
 		static const FName NonNativeTransientInnerObjectName("NonNativeTransientInnerObject");
 		UObject* NonNativeTransientInnerDefaultSubobject = NewObject<USubobjectInstancingTestObject>(TestCDO, NonNativeTransientInnerObjectName, RF_ArchetypeObject | TestCDO->GetMaskedFlags(RF_PropagateToSubObjects));
@@ -283,10 +323,13 @@ namespace SubobjectInstancingTest::Private
 		// Simulate a non-native class constructor initialization of the outer object property's default value
 		NonNativeOuterObjectProperty->SetObjectPropertyValue_InContainer(TestCDO, nullptr);
 
-		// Simulate a non-native class constructor initialization of the edit-time reference property's default value
+		// Simulate a non-native class constructor initialization of the edit-time reference property's default value (non-transient case)
 		NonNativeEditTimeInnerObjectProperty->SetObjectPropertyValue_InContainer(TestCDO, nullptr);
 
-		// Simulate a non-native clsas constructor initialization of a non-nullable subobject reference, also "captured" by the CDO and used to initialize all new objects of this type
+		// Simulate a non-native class constructor initialization of the transient edit-time reference property's default value (same as above)
+		NonNativeTransientEditTimeInnerObjectProperty->SetObjectPropertyValue_InContainer(TestCDO, nullptr);
+
+		// Simulate a non-native class constructor initialization of a non-nullable subobject reference, also "captured" by the CDO and used to initialize all new objects of this type
 		static const FName NonNullableInnerObjectName("NonNullableInnerObject");
 		UObject* NonNullableInnerDefaultSubobject = NewObject<USubobjectInstancingTestObject>(TestCDO, NonNullableInnerObjectName, RF_ArchetypeObject | TestCDO->GetMaskedFlags(RF_PropagateToSubObjects));
 		NonNullableInnerObjectProperty->SetObjectPropertyValue_InContainer(TestCDO, NonNullableInnerDefaultSubobject);
@@ -680,6 +723,53 @@ namespace SubobjectInstancingTest::Private
 			CHECK(OuterObject->TransientInnerObjectUsingNew != NewOuterObject->TransientInnerObjectUsingNew);
 		}
 
+		SECTION("DuplicateTransient native property initialized from CDO using NewObject()")
+		{
+			USubobjectInstancingTestOuterObject* OuterObject = NewObject<USubobjectInstancingTestOuterObject>(GetTransientPackage(), InstancingTestClass);
+
+			// CPF_DuplicateTransient suppresses default value initialization in FObjectInitializer::InitProperties()
+			// in the same way as CPF_Transient. The ctor-created instance is used directly and instance data does not propagate from the CDO.
+			REQUIRE(OuterObject->DuplicateTransientInnerObjectUsingNew != nullptr);
+			CHECK(OuterObject->DuplicateTransientInnerObjectUsingNew->IsInOuter(OuterObject));
+
+			CHECK(OuterObject->DuplicateTransientInnerObjectUsingNew->TestValue == 100);
+
+			CHECK_FALSE(OuterObject->DuplicateTransientInnerObjectUsingNew->HasAllFlags(RF_DefaultSubObject));
+			CHECK(OuterObject->DuplicateTransientInnerObjectUsingNew->HasAllFlags(OuterObject->GetMaskedFlags(RF_PropagateToSubObjects)));
+
+			CHECK_FALSE(OuterObject->DuplicateTransientInnerObjectUsingNew->IsTemplate());
+			CHECK(OuterObject->DuplicateTransientInnerObjectUsingNew->IsDefaultSubobject());
+			CHECK_FALSE(OuterObject->DuplicateTransientInnerObjectUsingNew->IsTemplateForSubobjects());
+
+			USubobjectInstancingTestOuterObject* NewOuterObject = NewObject<USubobjectInstancingTestOuterObject>(GetTransientPackage(), InstancingTestClass);
+			CHECK(OuterObject->DuplicateTransientInnerObjectUsingNew != NewOuterObject->DuplicateTransientInnerObjectUsingNew);
+		}
+
+		SECTION("DuplicateTransient native property initialized from template using NewObject()")
+		{
+			USubobjectInstancingTestOuterObject* Template = NewObject<USubobjectInstancingTestOuterObject>(GetTransientPackage(), InstancingTestClass, NAME_None, RF_ArchetypeObject);
+			Template->DuplicateTransientInnerObjectUsingNew->TestValue = 200;
+
+			USubobjectInstancingTestOuterObject* OuterObject = NewObject<USubobjectInstancingTestOuterObject>(GetTransientPackage(), InstancingTestClass, NAME_None, RF_NoFlags, Template);
+
+			REQUIRE(OuterObject->DuplicateTransientInnerObjectUsingNew != nullptr);
+			CHECK(OuterObject->DuplicateTransientInnerObjectUsingNew->IsInOuter(OuterObject));
+
+			// CPF_DuplicateTransient suppresses the InitProperties template copy path in the same way as CPF_Transient, so instance
+			// data does not propagate from the template regardless of whether the type is native or non-native.
+			CHECK(OuterObject->DuplicateTransientInnerObjectUsingNew->TestValue == 100);
+
+			CHECK_FALSE(OuterObject->DuplicateTransientInnerObjectUsingNew->HasAllFlags(RF_DefaultSubObject));
+			CHECK(OuterObject->DuplicateTransientInnerObjectUsingNew->HasAllFlags(OuterObject->GetMaskedFlags(RF_PropagateToSubObjects)));
+
+			CHECK_FALSE(OuterObject->DuplicateTransientInnerObjectUsingNew->IsTemplate());
+			CHECK(OuterObject->DuplicateTransientInnerObjectUsingNew->IsDefaultSubobject());
+			CHECK_FALSE(OuterObject->DuplicateTransientInnerObjectUsingNew->IsTemplateForSubobjects());
+
+			USubobjectInstancingTestOuterObject* NewOuterObject = NewObject<USubobjectInstancingTestOuterObject>(GetTransientPackage(), InstancingTestClass, NAME_None, RF_NoFlags, Template);
+			CHECK(OuterObject->DuplicateTransientInnerObjectUsingNew != NewOuterObject->DuplicateTransientInnerObjectUsingNew);
+		}
+
 		SECTION("Native property initialized to a transient type instantiation from CDO")
 		{
 			USubobjectInstancingTestOuterObject* OuterObject = NewObject<USubobjectInstancingTestOuterObject>(GetTransientPackage(), InstancingTestClass);
@@ -829,6 +919,62 @@ namespace SubobjectInstancingTest::Private
 				CHECK_FALSE(InnerObjectFromType->IsTemplateForSubobjects());
 			}
 		}
+
+		SECTION("Transient array of default subobjects initialized from CDO")
+		{
+			USubobjectInstancingTestOuterObject* OuterObject = NewObject<USubobjectInstancingTestOuterObject>(GetTransientPackage(), InstancingTestClass);
+
+			// CPF_Transient on the array property excludes it from InitProperties default value copy. Elements are
+			// re-constructed from the native constructor on each new instance, not inherited from the CDO.
+			CHECK_FALSE(OuterObject->TransientInnerObjectArray.IsEmpty());
+			for (USubobjectInstancingTestObject* InnerObject : OuterObject->TransientInnerObjectArray)
+			{
+				REQUIRE(InnerObject != nullptr);
+				CHECK(InnerObject->IsInOuter(OuterObject));
+
+				CHECK(InnerObject->TestValue == 100);
+
+				CHECK(InnerObject->HasAllFlags(RF_DefaultSubObject));
+				CHECK(InnerObject->HasAllFlags(OuterObject->GetMaskedFlags(RF_PropagateToSubObjects)));
+				// CPF_Transient is on the array property, not the element instances; RF_Transient is not imposed.
+				CHECK_FALSE(InnerObject->HasAllFlags(RF_Transient));
+
+				CHECK_FALSE(InnerObject->IsTemplate());
+				CHECK(InnerObject->IsDefaultSubobject());
+				CHECK_FALSE(InnerObject->IsTemplateForSubobjects());
+			}
+		}
+
+		SECTION("Transient array of default subobjects initialized from template")
+		{
+			USubobjectInstancingTestOuterObject* Template = NewObject<USubobjectInstancingTestOuterObject>(GetTransientPackage(), InstancingTestClass, NAME_None, RF_ArchetypeObject);
+			for (USubobjectInstancingTestObject* InnerObject : Template->TransientInnerObjectArray)
+			{
+				InnerObject->TestValue = 200;
+			}
+
+			USubobjectInstancingTestOuterObject* OuterObject = NewObject<USubobjectInstancingTestOuterObject>(GetTransientPackage(), InstancingTestClass, NAME_None, RF_NoFlags, Template);
+
+			// CPF_Transient on the array property excludes it from InitProperties template copy. The container
+			// elements are re-constructed by the native constructor via CreateDefaultSubobject, which initializes
+			// each new element from its named counterpart in the template — so TestValue is still inherited.
+			CHECK_FALSE(OuterObject->TransientInnerObjectArray.IsEmpty());
+			for (USubobjectInstancingTestObject* InnerObject : OuterObject->TransientInnerObjectArray)
+			{
+				REQUIRE(InnerObject != nullptr);
+				CHECK(InnerObject->IsInOuter(OuterObject));
+
+				CHECK(InnerObject->TestValue == 200);
+
+				CHECK(InnerObject->HasAllFlags(RF_DefaultSubObject));
+				CHECK(InnerObject->HasAllFlags(OuterObject->GetMaskedFlags(RF_PropagateToSubObjects)));
+				CHECK_FALSE(InnerObject->HasAllFlags(RF_Transient));
+
+				CHECK_FALSE(InnerObject->IsTemplate());
+				CHECK(InnerObject->IsDefaultSubobject());
+				CHECK_FALSE(InnerObject->IsTemplateForSubobjects());
+			}
+		}
 	}
 
 	void RunSubobjectInstancingTests_Set(UClass* InstancingTestClass)
@@ -928,6 +1074,62 @@ namespace SubobjectInstancingTest::Private
 				CHECK_FALSE(InnerObjectFromType->IsTemplateForSubobjects());
 			}
 		}
+
+		SECTION("Transient set of default subobjects initialized from CDO")
+		{
+			USubobjectInstancingTestOuterObject* OuterObject = NewObject<USubobjectInstancingTestOuterObject>(GetTransientPackage(), InstancingTestClass);
+
+			// CPF_Transient on the set property excludes it from InitProperties default value copy. Elements are
+			// re-constructed from the native constructor on each new instance, not inherited from the CDO.
+			CHECK_FALSE(OuterObject->TransientInnerObjectSet.IsEmpty());
+			for (USubobjectInstancingTestObject* InnerObject : OuterObject->TransientInnerObjectSet)
+			{
+				REQUIRE(InnerObject != nullptr);
+				CHECK(InnerObject->IsInOuter(OuterObject));
+
+				CHECK(InnerObject->TestValue == 100);
+
+				CHECK(InnerObject->HasAllFlags(RF_DefaultSubObject));
+				CHECK(InnerObject->HasAllFlags(OuterObject->GetMaskedFlags(RF_PropagateToSubObjects)));
+				// CPF_Transient is on the set property, not the element instances; RF_Transient is not imposed.
+				CHECK_FALSE(InnerObject->HasAllFlags(RF_Transient));
+
+				CHECK_FALSE(InnerObject->IsTemplate());
+				CHECK(InnerObject->IsDefaultSubobject());
+				CHECK_FALSE(InnerObject->IsTemplateForSubobjects());
+			}
+		}
+
+		SECTION("Transient set of default subobjects initialized from template")
+		{
+			USubobjectInstancingTestOuterObject* Template = NewObject<USubobjectInstancingTestOuterObject>(GetTransientPackage(), InstancingTestClass, NAME_None, RF_ArchetypeObject);
+			for (USubobjectInstancingTestObject* InnerObject : Template->TransientInnerObjectSet)
+			{
+				InnerObject->TestValue = 200;
+			}
+
+			USubobjectInstancingTestOuterObject* OuterObject = NewObject<USubobjectInstancingTestOuterObject>(GetTransientPackage(), InstancingTestClass, NAME_None, RF_NoFlags, Template);
+
+			// CPF_Transient on the set property excludes it from InitProperties template copy. The container
+			// elements are re-constructed by the native constructor via CreateDefaultSubobject, which initializes
+			// each new element from its named counterpart in the template — so TestValue is still inherited.
+			CHECK_FALSE(OuterObject->TransientInnerObjectSet.IsEmpty());
+			for (USubobjectInstancingTestObject* InnerObject : OuterObject->TransientInnerObjectSet)
+			{
+				REQUIRE(InnerObject != nullptr);
+				CHECK(InnerObject->IsInOuter(OuterObject));
+
+				CHECK(InnerObject->TestValue == 200);
+
+				CHECK(InnerObject->HasAllFlags(RF_DefaultSubObject));
+				CHECK(InnerObject->HasAllFlags(OuterObject->GetMaskedFlags(RF_PropagateToSubObjects)));
+				CHECK_FALSE(InnerObject->HasAllFlags(RF_Transient));
+
+				CHECK_FALSE(InnerObject->IsTemplate());
+				CHECK(InnerObject->IsDefaultSubobject());
+				CHECK_FALSE(InnerObject->IsTemplateForSubobjects());
+			}
+		}
 	}
 
 	void RunSubobjectInstancingTests_Map(UClass* InstancingTestClass)
@@ -1077,6 +1279,89 @@ namespace SubobjectInstancingTest::Private
 				CHECK_FALSE(InnerObjectFromTypePair.Value->IsTemplateForSubobjects());
 			}
 		}
+
+		SECTION("Transient map of default subobjects initialized from CDO")
+		{
+			USubobjectInstancingTestOuterObject* OuterObject = NewObject<USubobjectInstancingTestOuterObject>(GetTransientPackage(), InstancingTestClass);
+
+			// CPF_Transient on the map property excludes it from InitProperties default value copy. Elements are
+			// re-constructed from the native constructor on each new instance, not inherited from the CDO.
+			CHECK_FALSE(OuterObject->TransientInnerObjectMap.IsEmpty());
+			for (TPair<USubobjectInstancingTestObject*, USubobjectInstancingTestObject*> InnerObjectPair : OuterObject->TransientInnerObjectMap)
+			{
+				REQUIRE(InnerObjectPair.Key != nullptr);
+				CHECK(InnerObjectPair.Key->IsInOuter(OuterObject));
+
+				CHECK(InnerObjectPair.Key->TestValue == 100);
+
+				CHECK(InnerObjectPair.Key->HasAllFlags(RF_DefaultSubObject));
+				CHECK(InnerObjectPair.Key->HasAllFlags(OuterObject->GetMaskedFlags(RF_PropagateToSubObjects)));
+				// CPF_Transient is on the map property, not the element instances; RF_Transient is not imposed.
+				CHECK_FALSE(InnerObjectPair.Key->HasAllFlags(RF_Transient));
+
+				CHECK_FALSE(InnerObjectPair.Key->IsTemplate());
+				CHECK(InnerObjectPair.Key->IsDefaultSubobject());
+				CHECK_FALSE(InnerObjectPair.Key->IsTemplateForSubobjects());
+
+				REQUIRE(InnerObjectPair.Value != nullptr);
+				CHECK(InnerObjectPair.Value->IsInOuter(OuterObject));
+
+				CHECK(InnerObjectPair.Value->TestValue == 100);
+
+				CHECK(InnerObjectPair.Value->HasAllFlags(RF_DefaultSubObject));
+				CHECK(InnerObjectPair.Value->HasAllFlags(OuterObject->GetMaskedFlags(RF_PropagateToSubObjects)));
+				CHECK_FALSE(InnerObjectPair.Value->HasAllFlags(RF_Transient));
+
+				CHECK_FALSE(InnerObjectPair.Value->IsTemplate());
+				CHECK(InnerObjectPair.Value->IsDefaultSubobject());
+				CHECK_FALSE(InnerObjectPair.Value->IsTemplateForSubobjects());
+			}
+		}
+
+		SECTION("Transient map of default subobjects initialized from template")
+		{
+			USubobjectInstancingTestOuterObject* Template = NewObject<USubobjectInstancingTestOuterObject>(GetTransientPackage(), InstancingTestClass, NAME_None, RF_ArchetypeObject);
+			for (TPair<USubobjectInstancingTestObject*, USubobjectInstancingTestObject*> InnerObjectPair : Template->TransientInnerObjectMap)
+			{
+				InnerObjectPair.Key->TestValue = 200;
+				InnerObjectPair.Value->TestValue = 300;
+			}
+
+			USubobjectInstancingTestOuterObject* OuterObject = NewObject<USubobjectInstancingTestOuterObject>(GetTransientPackage(), InstancingTestClass, NAME_None, RF_NoFlags, Template);
+
+			// CPF_Transient on the map property excludes it from InitProperties template copy. The container
+			// elements are re-constructed by the native constructor via CreateDefaultSubobject, which initializes
+			// each new element from its named counterpart in the template — so TestValue is still inherited.
+			CHECK_FALSE(OuterObject->TransientInnerObjectMap.IsEmpty());
+			for (TPair<USubobjectInstancingTestObject*, USubobjectInstancingTestObject*> InnerObjectPair : OuterObject->TransientInnerObjectMap)
+			{
+				REQUIRE(InnerObjectPair.Key != nullptr);
+				CHECK(InnerObjectPair.Key->IsInOuter(OuterObject));
+
+				CHECK(InnerObjectPair.Key->TestValue == 200);
+
+				CHECK(InnerObjectPair.Key->HasAllFlags(RF_DefaultSubObject));
+				CHECK(InnerObjectPair.Key->HasAllFlags(OuterObject->GetMaskedFlags(RF_PropagateToSubObjects)));
+				CHECK_FALSE(InnerObjectPair.Key->HasAllFlags(RF_Transient));
+
+				CHECK_FALSE(InnerObjectPair.Key->IsTemplate());
+				CHECK(InnerObjectPair.Key->IsDefaultSubobject());
+				CHECK_FALSE(InnerObjectPair.Key->IsTemplateForSubobjects());
+
+				REQUIRE(InnerObjectPair.Value != nullptr);
+				CHECK(InnerObjectPair.Value->IsInOuter(OuterObject));
+
+				CHECK(InnerObjectPair.Value->TestValue == 300);
+
+				CHECK(InnerObjectPair.Value->HasAllFlags(RF_DefaultSubObject));
+				CHECK(InnerObjectPair.Value->HasAllFlags(OuterObject->GetMaskedFlags(RF_PropagateToSubObjects)));
+				CHECK_FALSE(InnerObjectPair.Value->HasAllFlags(RF_Transient));
+
+				CHECK_FALSE(InnerObjectPair.Value->IsTemplate());
+				CHECK(InnerObjectPair.Value->IsDefaultSubobject());
+				CHECK_FALSE(InnerObjectPair.Value->IsTemplateForSubobjects());
+			}
+		}
 	}
 
 	void RunSubobjectInstancingTests_Struct(UClass* InstancingTestClass)
@@ -1396,6 +1681,63 @@ namespace SubobjectInstancingTest::Private
 			CHECK(OuterObject->OptionalNullableInnerObjectUnset.IsSet());
 			CHECK(OuterObject->OptionalNullableInnerObjectUnset.GetValue().Get() == nullptr);
 		}
+
+		SECTION("Optional transient default subobject initialized from CDO")
+		{
+			USubobjectInstancingTestOuterObject* OuterObject = NewObject<USubobjectInstancingTestOuterObject>(GetTransientPackage(), InstancingTestClass);
+
+			// CPF_Transient on the optional property excludes it from InitProperties default value copy. The element
+			// is re-constructed from the native constructor on each new instance, not inherited from the CDO.
+			CHECK(OuterObject->OptionalTransientInnerObject.IsSet());
+
+			USubobjectInstancingTestObject* InnerObject = OuterObject->OptionalTransientInnerObject.GetValue();
+			REQUIRE(InnerObject != nullptr);
+			CHECK(InnerObject->IsInOuter(OuterObject));
+
+			CHECK(InnerObject->TestValue == 100);
+
+			CHECK(InnerObject->HasAllFlags(RF_DefaultSubObject));
+			CHECK(InnerObject->HasAllFlags(OuterObject->GetMaskedFlags(RF_PropagateToSubObjects)));
+			// CPF_Transient is on the optional property, not the element instance; RF_Transient is not imposed.
+			CHECK_FALSE(InnerObject->HasAllFlags(RF_Transient));
+
+			CHECK_FALSE(InnerObject->IsTemplate());
+			CHECK(InnerObject->IsDefaultSubobject());
+			CHECK_FALSE(InnerObject->IsTemplateForSubobjects());
+
+			USubobjectInstancingTestOuterObject* NewOuterObject = NewObject<USubobjectInstancingTestOuterObject>(GetTransientPackage(), InstancingTestClass);
+			CHECK(InnerObject != NewOuterObject->OptionalTransientInnerObject.GetValue());
+		}
+
+		SECTION("Optional transient default subobject initialized from template")
+		{
+			USubobjectInstancingTestOuterObject* Template = NewObject<USubobjectInstancingTestOuterObject>(GetTransientPackage(), InstancingTestClass, NAME_None, RF_ArchetypeObject);
+			Template->OptionalTransientInnerObject.GetValue()->TestValue = 200;
+
+			USubobjectInstancingTestOuterObject* OuterObject = NewObject<USubobjectInstancingTestOuterObject>(GetTransientPackage(), InstancingTestClass, NAME_None, RF_NoFlags, Template);
+
+			// CPF_Transient on the optional property excludes it from InitProperties template copy. The element is
+			// re-constructed by the native constructor via CreateDefaultSubobject, which initializes it from its
+			// named counterpart in the template — so TestValue is still inherited.
+			CHECK(OuterObject->OptionalTransientInnerObject.IsSet());
+
+			USubobjectInstancingTestObject* InnerObject = OuterObject->OptionalTransientInnerObject.GetValue();
+			REQUIRE(InnerObject != nullptr);
+			CHECK(InnerObject->IsInOuter(OuterObject));
+
+			CHECK(InnerObject->TestValue == 200);
+
+			CHECK(InnerObject->HasAllFlags(RF_DefaultSubObject));
+			CHECK(InnerObject->HasAllFlags(OuterObject->GetMaskedFlags(RF_PropagateToSubObjects)));
+			CHECK_FALSE(InnerObject->HasAllFlags(RF_Transient));
+
+			CHECK_FALSE(InnerObject->IsTemplate());
+			CHECK(InnerObject->IsDefaultSubobject());
+			CHECK_FALSE(InnerObject->IsTemplateForSubobjects());
+
+			USubobjectInstancingTestOuterObject* NewOuterObject = NewObject<USubobjectInstancingTestOuterObject>(GetTransientPackage(), InstancingTestClass);
+			CHECK(InnerObject != NewOuterObject->OptionalTransientInnerObject.GetValue());
+		}
 	}
 
 	void RunSubobjectInstancingTests_Nested(UClass* InstancingTestClass)
@@ -1540,6 +1882,59 @@ namespace SubobjectInstancingTest::Private
 			CHECK(OuterObject->InnerObjectWithIndirectlyNestedObject->InnerObject->InnerObject != NewOuterObject->InnerObjectWithIndirectlyNestedObject->InnerObject->InnerObject);
 		}
 
+		SECTION("Nested instanced subobject preserves transient flag")
+		{
+			if (InstancingTestClass->IsNative())
+			{
+				// Skip: This test only applies to non-native outer types.
+			}
+			else
+			{
+				FObjectProperty* NonNativeDirectlyNestedInnerObjectProperty = CastFieldChecked<FObjectProperty>(InstancingTestClass->FindPropertyByName(NonNativeDirectlyNestedInnerObjectPropertyName));
+
+				USubobjectInstancingTestOuterObject* CDO = Cast<USubobjectInstancingTestOuterObject>(InstancingTestClass->GetDefaultObject());
+				REQUIRE(CDO != nullptr);
+
+				// Default-constructed archetype
+				USubobjectInstancingTestDirectlyNestedObject* NonNativeDirectlyNestedInnerObjectTemplate = CastChecked<USubobjectInstancingTestDirectlyNestedObject>(NonNativeDirectlyNestedInnerObjectProperty->GetObjectPropertyValue_InContainer(CDO));
+				
+				REQUIRE(NonNativeDirectlyNestedInnerObjectTemplate != nullptr);
+				CHECK(NonNativeDirectlyNestedInnerObjectTemplate->IsTemplate());
+				CHECK(NonNativeDirectlyNestedInnerObjectTemplate->IsDefaultSubobject());
+				CHECK(NonNativeDirectlyNestedInnerObjectTemplate->IsTemplateForSubobjects());
+
+				// Default value of the nested archetype before it's edited 
+				CHECK(NonNativeDirectlyNestedInnerObjectTemplate->EditTimeInnerObject == nullptr);
+
+				// Simulate instantiating a nested transient subobject archetype at edit time
+				NonNativeDirectlyNestedInnerObjectTemplate->EditTimeInnerObject = NewObject<USubobjectInstancingTestObject>(NonNativeDirectlyNestedInnerObjectTemplate, TEXT("EditTimeInnerObject"), NonNativeDirectlyNestedInnerObjectTemplate->GetMaskedFlags(RF_PropagateToSubObjects) | RF_Transient);
+
+				REQUIRE(NonNativeDirectlyNestedInnerObjectTemplate->EditTimeInnerObject != nullptr);
+				CHECK(NonNativeDirectlyNestedInnerObjectTemplate->EditTimeInnerObject->IsTemplate());
+				CHECK_FALSE(NonNativeDirectlyNestedInnerObjectTemplate->EditTimeInnerObject->IsDefaultSubobject());
+				CHECK(NonNativeDirectlyNestedInnerObjectTemplate->EditTimeInnerObject->IsTemplateForSubobjects());
+
+				CHECK(NonNativeDirectlyNestedInnerObjectTemplate->EditTimeInnerObject->HasAllFlags(RF_Transient));
+
+				// Construct a new instance and verify RF_Transient is preserved on the nested instanced copy.
+				USubobjectInstancingTestOuterObject* OuterInstance = NewObject<USubobjectInstancingTestOuterObject>(GetTransientPackage(), InstancingTestClass);
+				REQUIRE(OuterInstance != nullptr);
+
+				USubobjectInstancingTestDirectlyNestedObject* NonNativeDirectlyNestedInnerObject = CastChecked<USubobjectInstancingTestDirectlyNestedObject>(NonNativeDirectlyNestedInnerObjectProperty->GetObjectPropertyValue_InContainer(OuterInstance));
+
+				REQUIRE(NonNativeDirectlyNestedInnerObject != nullptr);
+				CHECK(NonNativeDirectlyNestedInnerObject != NonNativeDirectlyNestedInnerObjectTemplate);
+				CHECK(NonNativeDirectlyNestedInnerObject->IsInOuter(OuterInstance));
+
+				// Instanced value of the nested subobject
+				REQUIRE(NonNativeDirectlyNestedInnerObject->EditTimeInnerObject != nullptr);
+				CHECK(NonNativeDirectlyNestedInnerObject->EditTimeInnerObject != NonNativeDirectlyNestedInnerObjectTemplate->EditTimeInnerObject);
+				CHECK(NonNativeDirectlyNestedInnerObject->EditTimeInnerObject->IsInOuter(NonNativeDirectlyNestedInnerObject));
+
+				CHECK(NonNativeDirectlyNestedInnerObject->EditTimeInnerObject->HasAllFlags(RF_Transient));
+			}
+		}
+
 		SECTION("Default subobject with nested dynamically-instanced subobject")
 		{
 			if (InstancingTestClass->IsNative())
@@ -1564,6 +1959,14 @@ namespace SubobjectInstancingTest::Private
 				NonNativeEditTimeInnerObjectDefaultValue->TestValue = 200;
 				CastFieldChecked<FObjectProperty>(TestClass->FindPropertyByName(Private::NonNativeOuterObjectPropertyName))->SetObjectPropertyValue_InContainer(NonNativeEditTimeInnerObjectDefaultValue, Template);
 				CastFieldChecked<FObjectProperty>(TestClass->FindPropertyByName(Private::NonNativeSelfReferencePropertyName))->SetObjectPropertyValue_InContainer(NonNativeEditTimeInnerObjectDefaultValue, NonNativeEditTimeInnerObjectDefaultValue);
+				
+				// Modify the default child object's state to simulate an override value on its nested subobject. This value should be copied to the child object of the nested subobject instance.
+				USubobjectInstancingTestObject* NestedNonNativeInnerObjectDefaultValue = CastChecked<USubobjectInstancingTestObject>(CastFieldChecked<FObjectProperty>(TestClass->FindPropertyByName(Private::NonNativeInnerObjectPropertyName))->GetObjectPropertyValue_InContainer(NonNativeEditTimeInnerObjectDefaultValue));
+				NestedNonNativeInnerObjectDefaultValue->TestValue = 200;
+
+				// Also modify the default child object's state to simulate an override value on its nested non-native transient-referenced subobject. This value should also be copied after instancing.
+				USubobjectInstancingTestObject* NestedNonNativeTransientInnerObjectDefaultValue = CastChecked<USubobjectInstancingTestObject>(CastFieldChecked<FObjectProperty>(TestClass->FindPropertyByName(Private::NonNativeTransientInnerObjectPropertyName))->GetObjectPropertyValue_InContainer(NonNativeEditTimeInnerObjectDefaultValue));
+				NestedNonNativeTransientInnerObjectDefaultValue->TestValue = 200;
 
 				// Construct a new instance of the outer object type using the template for initialization.
 				USubobjectInstancingTestOuterObject* OuterObject = NewObject<USubobjectInstancingTestOuterObject>(GetTransientPackage(), InstancingTestClass, NAME_None, RF_NoFlags, Template);
@@ -1578,14 +1981,32 @@ namespace SubobjectInstancingTest::Private
 				{
 					REQUIRE(NonNativeEditTimeInnerObject_NonNativeInnerObject != nullptr);
 					CHECK(NonNativeEditTimeInnerObject_NonNativeInnerObject->IsInOuter(NonNativeEditTimeInnerObject));
+
+					USubobjectInstancingTestObject* NestedNonNativeInnerObjectValue = Cast<USubobjectInstancingTestObject>(NonNativeEditTimeInnerObject_NonNativeInnerObject);
+					{
+						REQUIRE(NestedNonNativeInnerObjectValue != nullptr);
+
+						CHECK(NestedNonNativeInnerObjectValue != NestedNonNativeInnerObjectDefaultValue);
+						CHECK(NestedNonNativeInnerObjectValue->IsInOuter(NonNativeEditTimeInnerObject));
+						CHECK(NestedNonNativeInnerObjectValue->TestValue == NestedNonNativeInnerObjectDefaultValue->TestValue);
+					}
 				}
 
 				// Deeply-nested transient reference to a dynamically-instanced subobject, expected to be duplicated as part of constructing the outermost root object initialized using the template.
-				// Note: Unlike the case above, this would not be expected to be serialized as an export on save.
+				// Note: Unlike the case above, this would not be expected to be serialized as an export on save. It should, however, be initialized the same as the non-transient-referenced subobject.
 				UObject* NonNativeEditTimeInnerObject_NonNativeTransientInnerObject = CastFieldChecked<FObjectProperty>(TestClass->FindPropertyByName(Private::NonNativeTransientInnerObjectPropertyName))->GetObjectPropertyValue_InContainer(NonNativeEditTimeInnerObject);
 				{
 					REQUIRE(NonNativeEditTimeInnerObject_NonNativeTransientInnerObject != nullptr);
 					CHECK(NonNativeEditTimeInnerObject_NonNativeTransientInnerObject->IsInOuter(NonNativeEditTimeInnerObject));
+					
+					USubobjectInstancingTestObject* NestedNonNativeTransientInnerObjectValue = Cast<USubobjectInstancingTestObject>(NonNativeEditTimeInnerObject_NonNativeTransientInnerObject);
+					{
+						REQUIRE(NestedNonNativeTransientInnerObjectValue != nullptr);
+						
+						CHECK(NestedNonNativeTransientInnerObjectValue != NestedNonNativeTransientInnerObjectDefaultValue);
+						CHECK(NestedNonNativeTransientInnerObjectValue->IsInOuter(NonNativeEditTimeInnerObject));
+						CHECK(NestedNonNativeTransientInnerObjectValue->TestValue == NestedNonNativeTransientInnerObjectDefaultValue->TestValue);
+					}
 				}
 
 				// Nested reference to the outermost instance; ensures that we can self-reference nodes that exist above the immediate parent node in the object graph for nested subobject types that use dynamic instancing.
@@ -1845,6 +2266,68 @@ namespace SubobjectInstancingTest::Private
 			}
 		}
 
+		SECTION("Simulate native DuplicateTransient property instanced at edit time on non-native CDO")
+		{
+			if (InstancingTestClass->IsNative())
+			{
+				// Skip: This test only applies to non-native types.
+			}
+			else
+			{
+				USubobjectInstancingTestOuterObject* CDO = Cast<USubobjectInstancingTestOuterObject>(InstancingTestClass->GetDefaultObject());
+
+				// Default value before it's edited (as will be initialized by the native super class ctor).
+				CHECK(CDO->DuplicateTransientEditTimeInnerObject == nullptr);
+
+				// Simulate a deferred edit-time instancing event along with a value override on the non-native CDO.
+				CDO->DuplicateTransientEditTimeInnerObject = NewObject<USubobjectInstancingTestObject>(CDO, GET_MEMBER_NAME_CHECKED(USubobjectInstancingTestOuterObject, DuplicateTransientEditTimeInnerObject), CDO->GetMaskedFlags(RF_PropagateToSubObjects));
+				CDO->DuplicateTransientEditTimeInnerObject->TestValue = 200;
+
+				CHECK(CDO->DuplicateTransientEditTimeInnerObject->IsTemplate());
+				CHECK(CDO->DuplicateTransientEditTimeInnerObject->IsDefaultSubobject());
+				CHECK(CDO->DuplicateTransientEditTimeInnerObject->IsTemplateForSubobjects());
+
+				USubobjectInstancingTestOuterObject* OuterObject = NewObject<USubobjectInstancingTestOuterObject>(GetTransientPackage(), InstancingTestClass);
+
+				// Since the subobject archetype was not instanced at construction time, it's expected to remain set to NULL. When the CDO is used as the default object for initialization, as an optimization,
+				// the native object initializer will not include properties inherited from the native super class hierarchy, because it expects they have already been initialized by the native ctor.
+				// Additionally, CPF_DuplicateTransient suppresses the InitProperties copy path in the same way as CPF_Transient.
+				CHECK(OuterObject->DuplicateTransientEditTimeInnerObject == nullptr);
+
+				// Clean up the subobject on the CDO so it doesn't persist through remaining tests.
+				CDO->DuplicateTransientEditTimeInnerObject->Rename(nullptr, GetTransientPackage(), REN_DontCreateRedirectors);
+				CDO->DuplicateTransientEditTimeInnerObject = nullptr;
+			}
+		}
+
+		SECTION("Simulate native DuplicateTransient property instanced at edit time on non-native template")
+		{
+			if (InstancingTestClass->IsNative())
+			{
+				// Skip: This test only applies to non-native types.
+			}
+			else
+			{
+				USubobjectInstancingTestOuterObject* Template = NewObject<USubobjectInstancingTestOuterObject>(GetTransientPackage(), InstancingTestClass, NAME_None, RF_ArchetypeObject);
+				CHECK(Template->DuplicateTransientEditTimeInnerObject == nullptr);
+
+				// Simulate a deferred edit-time instancing event along with a value override on the non-native template.
+				Template->DuplicateTransientEditTimeInnerObject = NewObject<USubobjectInstancingTestObject>(Template, GET_MEMBER_NAME_CHECKED(USubobjectInstancingTestOuterObject, DuplicateTransientEditTimeInnerObject), Template->GetMaskedFlags(RF_PropagateToSubObjects));
+				Template->DuplicateTransientEditTimeInnerObject->TestValue = 200;
+
+				CHECK(Template->DuplicateTransientEditTimeInnerObject->IsTemplate());
+				CHECK_FALSE(Template->DuplicateTransientEditTimeInnerObject->IsDefaultSubobject());
+				CHECK_FALSE(Template->DuplicateTransientEditTimeInnerObject->IsTemplateForSubobjects());
+				CHECK(Template->DuplicateTransientEditTimeInnerObject->IsTemplateForSubobjects(RF_ArchetypeObject));
+
+				USubobjectInstancingTestOuterObject* OuterObject = NewObject<USubobjectInstancingTestOuterObject>(GetTransientPackage(), InstancingTestClass, NAME_None, RF_NoFlags, Template);
+
+				// CPF_DuplicateTransient suppresses the InitProperties template copy path in the same way as CPF_Transient, so the
+				// template value is not propagated.
+				CHECK(OuterObject->DuplicateTransientEditTimeInnerObject == nullptr);
+			}
+		}
+
 		SECTION("Simulate non-native property instanced at edit time on non-native CDO")
 		{
 			if (InstancingTestClass->IsNative())
@@ -1944,6 +2427,234 @@ namespace SubobjectInstancingTest::Private
 			}
 		}
 
+		SECTION("Simulate non-native transient property instanced at edit time on non-native CDO")
+		{
+			if (InstancingTestClass->IsNative())
+			{
+				// Skip: This test only applies to non-native iterations - the native class won't include the non-native property.
+			}
+			else
+			{
+				USubobjectInstancingTestOuterObject* CDO = Cast<USubobjectInstancingTestOuterObject>(InstancingTestClass->GetDefaultObject());
+
+				FObjectProperty* ObjectProperty = CastFieldChecked<FObjectProperty>(InstancingTestClass->FindPropertyByName(Private::NonNativeTransientEditTimeInnerObjectPropertyName));
+				USubobjectInstancingTestObject* NonNativeTransientEditTimeInnerObjectDefaultValue = Cast<USubobjectInstancingTestObject>(ObjectProperty->GetObjectPropertyValue_InContainer(CDO));
+
+				// Default value before it's edited (it will have been zero-initialized).
+				CHECK(NonNativeTransientEditTimeInnerObjectDefaultValue == nullptr);
+
+				// Simulate a deferred edit-time instancing event along with a value override on the non-native CDO.
+				NonNativeTransientEditTimeInnerObjectDefaultValue = NewObject<USubobjectInstancingTestObject>(CDO, TEXT("NonNativeTransientEditTimeInnerObject"), CDO->GetMaskedFlags(RF_PropagateToSubObjects));
+				NonNativeTransientEditTimeInnerObjectDefaultValue->TestValue = 200;
+				ObjectProperty->SetObjectPropertyValue_InContainer(CDO, NonNativeTransientEditTimeInnerObjectDefaultValue);
+
+				CHECK(NonNativeTransientEditTimeInnerObjectDefaultValue->IsTemplate());
+				CHECK(NonNativeTransientEditTimeInnerObjectDefaultValue->IsDefaultSubobject());
+				CHECK(NonNativeTransientEditTimeInnerObjectDefaultValue->IsTemplateForSubobjects());
+
+				USubobjectInstancingTestOuterObject* OuterObject = NewObject<USubobjectInstancingTestOuterObject>(GetTransientPackage(), InstancingTestClass);
+				USubobjectInstancingTestObject* NonNativeTransientEditTimeInnerObject = Cast<USubobjectInstancingTestObject>(ObjectProperty->GetObjectPropertyValue_InContainer(OuterObject));
+				{
+					// Note: Only non-native types with dynamic referencing enabled can reflect a transient property's value through the CDO to new instances of the class, in order for
+					// them to be dynamically instanced as part of UObject construction. This copies the deferred override from the CDO which effectively acts as the non-native ctor.
+					// The reason it works in this case is because the property does not require 'Instanced', which decouples it from having to first be initialized by a native super.
+					if (InstancingTestClass->ShouldUseDynamicSubobjectInstancing())
+					{
+						REQUIRE(NonNativeTransientEditTimeInnerObject != nullptr);
+						CHECK(NonNativeTransientEditTimeInnerObject->IsInOuter(OuterObject));
+
+						CHECK(NonNativeTransientEditTimeInnerObject->TestValue == 200);
+
+						CHECK_FALSE(NonNativeTransientEditTimeInnerObject->HasAllFlags(RF_DefaultSubObject));
+						CHECK(NonNativeTransientEditTimeInnerObject->HasAllFlags(OuterObject->GetMaskedFlags(RF_PropagateToSubObjects)));
+
+						CHECK_FALSE(NonNativeTransientEditTimeInnerObject->IsTemplate());
+						CHECK(NonNativeTransientEditTimeInnerObject->IsDefaultSubobject());
+						CHECK_FALSE(NonNativeTransientEditTimeInnerObject->IsTemplateForSubobjects());
+
+						USubobjectInstancingTestOuterObject* NewOuterObject = NewObject<USubobjectInstancingTestOuterObject>(GetTransientPackage(), InstancingTestClass);
+						CHECK(NonNativeTransientEditTimeInnerObject != ObjectProperty->GetObjectPropertyValue_InContainer(NewOuterObject));
+					}
+					else
+					{
+						// As with the native case (see TransientEditTimeInnerObject), the non-native property is also tagged as 'Instanced', which means it is expected to have first been
+						// default-initialized by the native class constructor. Since the native side cannot access this property except via reflection, it is generally not default-
+						// initialized, and is left memset() to zero (NULL). While users can set the value on the CDO at edit time, the initializer will not propagate it to new instances
+						// of the class, and it cannot be instanced as a result. For this reason, overriding non-native transient references at edit time is not supported in this case.
+						CHECK(NonNativeTransientEditTimeInnerObject == nullptr);
+					}
+				}
+
+				// Clean up the subobject on the CDO so it doesn't persist through remaining tests.
+				ObjectProperty->SetObjectPropertyValue_InContainer(CDO, nullptr);
+				NonNativeTransientEditTimeInnerObjectDefaultValue->Rename(nullptr, GetTransientPackage(), REN_DontCreateRedirectors);
+			}
+		}
+
+		SECTION("Simulate non-native transient property instanced at edit time on non-native template")
+		{
+			if (InstancingTestClass->IsNative())
+			{
+				// Skip: This test only applies to non-native iterations - the native class won't include the non-native property.
+			}
+			else
+			{
+				USubobjectInstancingTestOuterObject* Template = NewObject<USubobjectInstancingTestOuterObject>(GetTransientPackage(), InstancingTestClass, NAME_None, RF_ArchetypeObject);
+
+				FObjectProperty* ObjectProperty = CastFieldChecked<FObjectProperty>(InstancingTestClass->FindPropertyByName(Private::NonNativeTransientEditTimeInnerObjectPropertyName));
+				USubobjectInstancingTestObject* NonNativeTransientEditTimeInnerObjectDefaultValue = Cast<USubobjectInstancingTestObject>(ObjectProperty->GetObjectPropertyValue_InContainer(Template));
+
+				// Default value before it's edited (it will have been zero-initialized).
+				CHECK(NonNativeTransientEditTimeInnerObjectDefaultValue == nullptr);
+
+				// Simulate a deferred edit-time instancing event along with a value override on the non-native template.
+				NonNativeTransientEditTimeInnerObjectDefaultValue = NewObject<USubobjectInstancingTestObject>(Template, TEXT("NonNativeTransientEditTimeInnerObject"), Template->GetMaskedFlags(RF_PropagateToSubObjects));
+				NonNativeTransientEditTimeInnerObjectDefaultValue->TestValue = 200;
+				ObjectProperty->SetObjectPropertyValue_InContainer(Template, NonNativeTransientEditTimeInnerObjectDefaultValue);
+
+				CHECK(NonNativeTransientEditTimeInnerObjectDefaultValue->IsTemplate());
+				// As with the non-transient case, since we instanced the subobject on the template and not the CDO, this is not considered to be a default subobject.
+				CHECK_FALSE(NonNativeTransientEditTimeInnerObjectDefaultValue->IsDefaultSubobject());
+				// And, like the non-transient case, it is also not considered to be a template for a CDO or default subobject, but rather a standalone archetype subobject.
+				CHECK_FALSE(NonNativeTransientEditTimeInnerObjectDefaultValue->IsTemplateForSubobjects());
+				CHECK(NonNativeTransientEditTimeInnerObjectDefaultValue->IsTemplateForSubobjects(RF_ArchetypeObject));
+
+				USubobjectInstancingTestOuterObject* OuterObject = NewObject<USubobjectInstancingTestOuterObject>(GetTransientPackage(), InstancingTestClass, NAME_None, RF_NoFlags, Template);
+				USubobjectInstancingTestObject* NonNativeTransientEditTimeInnerObject = Cast<USubobjectInstancingTestObject>(ObjectProperty->GetObjectPropertyValue_InContainer(OuterObject));
+				{
+					// As with the CDO case, transient references that are not default-constructed in native code can only be instanced from the archetype's value in the dynamic referencing case.
+					if (InstancingTestClass->ShouldUseDynamicSubobjectInstancing())
+					{
+						REQUIRE(NonNativeTransientEditTimeInnerObject != nullptr);
+						CHECK(NonNativeTransientEditTimeInnerObject->IsInOuter(OuterObject));
+
+						CHECK(NonNativeTransientEditTimeInnerObject->TestValue == 200);
+
+						CHECK_FALSE(NonNativeTransientEditTimeInnerObject->HasAllFlags(RF_DefaultSubObject));
+						CHECK(NonNativeTransientEditTimeInnerObject->HasAllFlags(OuterObject->GetMaskedFlags(RF_PropagateToSubObjects)));
+
+						CHECK_FALSE(NonNativeTransientEditTimeInnerObject->IsTemplate());
+						CHECK_FALSE(NonNativeTransientEditTimeInnerObject->IsDefaultSubobject());
+						CHECK_FALSE(NonNativeTransientEditTimeInnerObject->IsTemplateForSubobjects());
+
+						USubobjectInstancingTestOuterObject* NewOuterObject = NewObject<USubobjectInstancingTestOuterObject>(GetTransientPackage(), InstancingTestClass);
+						CHECK(NonNativeTransientEditTimeInnerObject != ObjectProperty->GetObjectPropertyValue_InContainer(NewOuterObject));
+					}
+					else
+					{
+						// As with the CDO test, the property is tagged as 'Instanced' in this case, and transient 'Instanced' properties must be constructed in native code.
+						// Since this is a non-native property, the native super type will not default-construct the archetype for this value, and it is left NULL as a result.
+						CHECK(NonNativeTransientEditTimeInnerObject == nullptr);
+					}
+				}
+			}
+		}
+
+		SECTION("Non-native transient subobject initialized from CDO")
+		{
+			if (InstancingTestClass->IsNative())
+			{
+				// Skip: This test only applies to non-native iterations - the native class won't include the non-native property.
+			}
+			else
+			{
+				FObjectProperty* ObjectProperty = CastFieldChecked<FObjectProperty>(InstancingTestClass->FindPropertyByName(Private::NonNativeTransientInnerObjectPropertyName));
+				USubobjectInstancingTestOuterObject* CDO = Cast<USubobjectInstancingTestOuterObject>(InstancingTestClass->GetDefaultObject());
+
+				// The CDO's value was assigned at class creation time as an RF_ArchetypeObject subobject, simulating a non-native constructor assignment.
+				USubobjectInstancingTestObject* NonNativeTransientInnerObjectDefaultValue = Cast<USubobjectInstancingTestObject>(ObjectProperty->GetObjectPropertyValue_InContainer(CDO));
+				REQUIRE(NonNativeTransientInnerObjectDefaultValue != nullptr);
+				CHECK(NonNativeTransientInnerObjectDefaultValue->IsTemplate());
+				CHECK(NonNativeTransientInnerObjectDefaultValue->IsDefaultSubobject());
+				CHECK(NonNativeTransientInnerObjectDefaultValue->IsTemplateForSubobjects());
+
+				NonNativeTransientInnerObjectDefaultValue->TestValue = 200;
+
+				USubobjectInstancingTestOuterObject* OuterObject = NewObject<USubobjectInstancingTestOuterObject>(GetTransientPackage(), InstancingTestClass);
+				USubobjectInstancingTestObject* NonNativeTransientInnerObject = Cast<USubobjectInstancingTestObject>(ObjectProperty->GetObjectPropertyValue_InContainer(OuterObject));
+				{
+					if (InstancingTestClass->ShouldUseDynamicSubobjectInstancing())
+					{
+						// Dynamic instancing propagates the CDO's transient non-native value to new instances.
+						REQUIRE(NonNativeTransientInnerObject != nullptr);
+						CHECK(NonNativeTransientInnerObject->IsInOuter(OuterObject));
+
+						CHECK(NonNativeTransientInnerObject->TestValue == 200);
+
+						CHECK_FALSE(NonNativeTransientInnerObject->HasAllFlags(RF_DefaultSubObject));
+						CHECK(NonNativeTransientInnerObject->HasAllFlags(OuterObject->GetMaskedFlags(RF_PropagateToSubObjects)));
+
+						CHECK_FALSE(NonNativeTransientInnerObject->IsTemplate());
+						CHECK(NonNativeTransientInnerObject->IsDefaultSubobject());
+						CHECK_FALSE(NonNativeTransientInnerObject->IsTemplateForSubobjects());
+
+						USubobjectInstancingTestOuterObject* NewOuterObject = NewObject<USubobjectInstancingTestOuterObject>(GetTransientPackage(), InstancingTestClass);
+						CHECK(NonNativeTransientInnerObject != ObjectProperty->GetObjectPropertyValue_InContainer(NewOuterObject));
+					}
+					else
+					{
+						// As with the native transient edit-time case, the property requires 'Instanced' (CPF_InstancedReference), which means it must first be initialized by native code.
+						// Since the native super type cannot access this non-native property, it is never default-constructed, and the transient InitProperties copy path is suppressed.
+						CHECK(NonNativeTransientInnerObject == nullptr);
+					}
+				}
+
+				// Restore CDO subobject state so it doesn't affect subsequent tests.
+				NonNativeTransientInnerObjectDefaultValue->TestValue = 100;
+			}
+		}
+
+		SECTION("Non-native transient subobject initialized from template")
+		{
+			if (InstancingTestClass->IsNative())
+			{
+				// Skip: This test only applies to non-native iterations - the native class won't include the non-native property.
+			}
+			else
+			{
+				USubobjectInstancingTestOuterObject* Template = NewObject<USubobjectInstancingTestOuterObject>(GetTransientPackage(), InstancingTestClass, NAME_None, RF_ArchetypeObject);
+				FObjectProperty* ObjectProperty = CastFieldChecked<FObjectProperty>(InstancingTestClass->FindPropertyByName(Private::NonNativeTransientInnerObjectPropertyName));
+				USubobjectInstancingTestObject* NonNativeTransientInnerObjectTemplateValue = Cast<USubobjectInstancingTestObject>(ObjectProperty->GetObjectPropertyValue_InContainer(Template));
+
+				if (InstancingTestClass->ShouldUseDynamicSubobjectInstancing())
+				{
+					// When the template was constructed from the CDO, dynamic instancing created a copy of the CDO's transient value on the template.
+					REQUIRE(NonNativeTransientInnerObjectTemplateValue != nullptr);
+					CHECK(NonNativeTransientInnerObjectTemplateValue->IsDefaultSubobject());
+					CHECK(NonNativeTransientInnerObjectTemplateValue->IsTemplateForSubobjects(RF_ArchetypeObject));
+
+					NonNativeTransientInnerObjectTemplateValue->TestValue = 200;
+
+					USubobjectInstancingTestOuterObject* OuterObject = NewObject<USubobjectInstancingTestOuterObject>(GetTransientPackage(), InstancingTestClass, NAME_None, RF_NoFlags, Template);
+					USubobjectInstancingTestObject* NonNativeTransientInnerObject = Cast<USubobjectInstancingTestObject>(ObjectProperty->GetObjectPropertyValue_InContainer(OuterObject));
+
+					REQUIRE(NonNativeTransientInnerObject != nullptr);
+					CHECK(NonNativeTransientInnerObject->IsInOuter(OuterObject));
+
+					CHECK(NonNativeTransientInnerObject->TestValue == 200);
+
+					CHECK_FALSE(NonNativeTransientInnerObject->HasAllFlags(RF_DefaultSubObject));
+					CHECK(NonNativeTransientInnerObject->HasAllFlags(OuterObject->GetMaskedFlags(RF_PropagateToSubObjects)));
+
+					CHECK_FALSE(NonNativeTransientInnerObject->IsTemplate());
+					CHECK(NonNativeTransientInnerObject->IsDefaultSubobject());
+					CHECK_FALSE(NonNativeTransientInnerObject->IsTemplateForSubobjects());
+
+					USubobjectInstancingTestOuterObject* NewOuterObject = NewObject<USubobjectInstancingTestOuterObject>(GetTransientPackage(), InstancingTestClass, NAME_None, RF_NoFlags, Template);
+					CHECK(NonNativeTransientInnerObject != ObjectProperty->GetObjectPropertyValue_InContainer(NewOuterObject));
+				}
+				else
+				{
+					// CPF_Transient suppresses the InitProperties template copy for instanced references, so the template value is nullptr
+					// and no instancing occurs.
+					CHECK(NonNativeTransientInnerObjectTemplateValue == nullptr);
+
+					USubobjectInstancingTestOuterObject* OuterObject = NewObject<USubobjectInstancingTestOuterObject>(GetTransientPackage(), InstancingTestClass, NAME_None, RF_NoFlags, Template);
+					CHECK(Cast<USubobjectInstancingTestObject>(ObjectProperty->GetObjectPropertyValue_InContainer(OuterObject)) == nullptr);
+				}
+			}
+		}
+
 		SECTION("Native property initialized from template with deferred construction using NewObject()")
 		{
 			USubobjectInstancingTestOuterObject* Template = NewObject<USubobjectInstancingTestOuterObject>(GetTransientPackage(), InstancingTestClass, NAME_None, RF_ArchetypeObject);
@@ -1967,6 +2678,28 @@ namespace SubobjectInstancingTest::Private
 			CHECK(OuterObject->InnerObjectPostInit != NewOuterObject->InnerObjectPostInit);
 		}
 
+		SECTION("Native transient property initialized from CDO with deferred construction using NewObject()")
+		{
+			USubobjectInstancingTestOuterObject* OuterObject = NewObject<USubobjectInstancingTestOuterObject>(GetTransientPackage(), InstancingTestClass);
+
+			// Transient instanced properties are excluded from default value initialization - see FObjectInitializer::InitProperties().
+			// Default value initialization is deferred to the native PostInitProperties() call, so should be non-NULL post-construction.
+			REQUIRE(OuterObject->TransientInnerObjectPostInit != nullptr);
+			CHECK(OuterObject->TransientInnerObjectPostInit->IsInOuter(OuterObject));
+
+			CHECK(OuterObject->TransientInnerObjectPostInit->TestValue == 100);
+
+			CHECK_FALSE(OuterObject->TransientInnerObjectPostInit->HasAllFlags(RF_DefaultSubObject));
+			CHECK(OuterObject->TransientInnerObjectPostInit->HasAllFlags(OuterObject->GetMaskedFlags(RF_PropagateToSubObjects)));
+
+			CHECK_FALSE(OuterObject->TransientInnerObjectPostInit->IsTemplate());
+			CHECK(OuterObject->TransientInnerObjectPostInit->IsDefaultSubobject());
+			CHECK_FALSE(OuterObject->TransientInnerObjectPostInit->IsTemplateForSubobjects());
+
+			USubobjectInstancingTestOuterObject* NewOuterObject = NewObject<USubobjectInstancingTestOuterObject>(GetTransientPackage(), InstancingTestClass);
+			CHECK(OuterObject->TransientInnerObjectPostInit != NewOuterObject->TransientInnerObjectPostInit);
+		}
+
 		SECTION("Native transient property initialized from template with deferred construction using NewObject()")
 		{
 			USubobjectInstancingTestOuterObject* Template = NewObject<USubobjectInstancingTestOuterObject>(GetTransientPackage(), InstancingTestClass, NAME_None, RF_ArchetypeObject);
@@ -2255,6 +2988,11 @@ namespace SubobjectInstancingTest::Private
 			CHECK(DstObject->InnerObjectUsingNew != SrcObject->InnerObjectUsingNew);
 			CHECK(DstObject->InnerObjectUsingNew->TestValue == SrcObject->InnerObjectUsingNew->TestValue);
 
+			CHECK(DstObject->TransientInnerObjectFromType != nullptr);
+			CHECK(DstObject->TransientInnerObjectFromType->IsInOuter(DstObject));
+			CHECK(DstObject->TransientInnerObjectFromType != SrcObject->TransientInnerObjectFromType);
+			CHECK(DstObject->TransientInnerObjectFromType->TestValue == SrcObject->TransientInnerObjectFromType->TestValue);
+
 			CHECK(DstObject->TransientInnerObjectUsingNew != nullptr);
 			CHECK(DstObject->TransientInnerObjectUsingNew->IsInOuter(DstObject));
 			CHECK(DstObject->TransientInnerObjectUsingNew != SrcObject->TransientInnerObjectUsingNew);
diff --git a/Engine/Source/Runtime/CoreUObject/Tests/UObject/SubobjectInstancingTest.h b/Engine/Source/Runtime/CoreUObject/Tests/UObject/SubobjectInstancingTest.h
--- a/Engine/Source/Runtime/CoreUObject/Tests/UObject/SubobjectInstancingTest.h
+++ b/Engine/Source/Runtime/CoreUObject/Tests/UObject/SubobjectInstancingTest.h
@@ -56,6 +56,9 @@ public:
 
 	UPROPERTY(Instanced)
 	TObjectPtr<USubobjectInstancingTestObject> InnerObject;
+
+	UPROPERTY(Instanced)
+	TObjectPtr<USubobjectInstancingTestObject> EditTimeInnerObject;
 };
 
 UCLASS()
@@ -155,6 +158,9 @@ public:
 	UPROPERTY(DuplicateTransient, Instanced)
 	TObjectPtr<USubobjectInstancingTestObject> DuplicateTransientInnerObjectUsingNew;
 
+	UPROPERTY(NonPIEDuplicateTransient, Instanced)
+	TObjectPtr<USubobjectInstancingTestObject> NonPIEDuplicateTransientInnerObjectUsingNew;
+
 	UPROPERTY(Instanced)
 	TObjectPtr<USubobjectInstancingTestObject> InnerObjectPostInit;
 
@@ -167,18 +173,27 @@ public:
 	UPROPERTY()
 	TArray<TObjectPtr<USubobjectInstancingDefaultToInstancedTestObject>> InnerObjectFromTypeArray;
 
+	UPROPERTY(Instanced, Transient)
+	TArray<TObjectPtr<USubobjectInstancingTestObject>> TransientInnerObjectArray;
+
 	UPROPERTY(Instanced)
 	TSet<TObjectPtr<USubobjectInstancingTestObject>> InnerObjectSet;
 
 	UPROPERTY()
 	TSet<TObjectPtr<USubobjectInstancingDefaultToInstancedTestObject>> InnerObjectFromTypeSet;
 
+	UPROPERTY(Instanced, Transient)
+	TSet<TObjectPtr<USubobjectInstancingTestObject>> TransientInnerObjectSet;
+
 	UPROPERTY(Instanced)
 	TMap<TObjectPtr<USubobjectInstancingTestObject>, TObjectPtr<USubobjectInstancingTestObject>> InnerObjectMap;
 
 	UPROPERTY()
 	TMap<TObjectPtr<USubobjectInstancingDefaultToInstancedTestObject>, TObjectPtr<USubobjectInstancingDefaultToInstancedTestObject>> InnerObjectFromTypeMap;
 
+	UPROPERTY(Instanced, Transient)
+	TMap<TObjectPtr<USubobjectInstancingTestObject>, TObjectPtr<USubobjectInstancingTestObject>> TransientInnerObjectMap;
+
 	UPROPERTY()
 	FSubobjectInstancingTestStructType StructWithInnerObjects;
 
@@ -188,6 +203,9 @@ public:
 	UPROPERTY()
 	TOptional<TObjectPtr<USubobjectInstancingDefaultToInstancedTestObject>> OptionalInnerObjectFromType;
 
+	UPROPERTY(Instanced, Transient)
+	TOptional<TObjectPtr<USubobjectInstancingTestObject>> OptionalTransientInnerObject;
+
 	UPROPERTY(Instanced)
 	TOptional<TArray<TObjectPtr<USubobjectInstancingTestObject>>>  OptionalInnerObjectArray;
 
