RemoteObjects:  More cases where GetNameSafe should not ever cause object migrations.

#rnx
#Synced-CL 52055930

diff --git a/Engine/Source/Runtime/CoreUObject/Private/UObject/ObjectPtr.cpp b/Engine/Source/Runtime/CoreUObject/Private/UObject/ObjectPtr.cpp
--- a/Engine/Source/Runtime/CoreUObject/Private/UObject/ObjectPtr.cpp
+++ b/Engine/Source/Runtime/CoreUObject/Private/UObject/ObjectPtr.cpp
@@ -133,7 +133,8 @@ FName FObjectPtr::GetFName() const
 		}
 		else if (const FRemoteObjectStub* Stub = LocalHandle.ToStub())
 		{
-			if (!Stub->Name.IsNone())
+			const bool bCannotMigrate = !AutoRTFM::IsClosed();
+			if (bCannotMigrate || !Stub->Name.IsNone())
 			{
 				return Stub->Name;
 			}
diff --git a/Engine/Source/Runtime/CoreUObject/Public/UObject/ObjectPtr.h b/Engine/Source/Runtime/CoreUObject/Public/UObject/ObjectPtr.h
--- a/Engine/Source/Runtime/CoreUObject/Public/UObject/ObjectPtr.h
+++ b/Engine/Source/Runtime/CoreUObject/Public/UObject/ObjectPtr.h
@@ -220,7 +220,24 @@ public:
 	// Gets the string name of the object without resolving the object reference.
 	[[nodiscard]] FORCEINLINE FString GetName() const
 	{
+#if UE_WITH_REMOTE_OBJECT_HANDLE
+		if (IsObjectHandleNull(Handle))
+		{
+			return TEXT("None");
+		}
+
+		// We enforce GetFName() does not cause migrations when run in the Open and instead will return NAME_None if non-local
+		const FName MyName = GetFName();
+		if (MyName != NAME_None)
+		{
+			return MyName.ToString();
+		}
+
+		// We know the handle was non-null, so we are remote
+		return FString::Printf(TEXT("%s (%s)"), *GetRemoteId().ToString(), EnumToString(GetResidence()));
+#else
 		return GetFName().ToString();
+#endif
 	}
 
 	// Gets the Outer UObject for this Object without resolving the object reference.
