Moving TransactionCommon from Engine to CoreUObject.

#rnx

#Synced-CL 52057580

diff --git a/Engine/Plugins/Chooser/Source/ChooserEditor/Private/MultiEnumColumnEditor.cpp b/Engine/Plugins/Chooser/Source/ChooserEditor/Private/MultiEnumColumnEditor.cpp
--- a/Engine/Plugins/Chooser/Source/ChooserEditor/Private/MultiEnumColumnEditor.cpp
+++ b/Engine/Plugins/Chooser/Source/ChooserEditor/Private/MultiEnumColumnEditor.cpp
@@ -9,7 +9,7 @@
 #include "ChooserTableEditor.h"
 #include "GraphEditorSettings.h"
 #include "SEnumCombo.h"
-#include "TransactionCommon.h"
+#include "Misc/TransactionCommon.h"
 #include "Widgets/Input/SButton.h"
 #include "ScopedTransaction.h"
 
diff --git a/Engine/Plugins/Chooser/Source/ChooserEditor/Private/NameColumnEditor.cpp b/Engine/Plugins/Chooser/Source/ChooserEditor/Private/NameColumnEditor.cpp
--- a/Engine/Plugins/Chooser/Source/ChooserEditor/Private/NameColumnEditor.cpp
+++ b/Engine/Plugins/Chooser/Source/ChooserEditor/Private/NameColumnEditor.cpp
@@ -9,7 +9,7 @@
 #include "ObjectChooserWidgetFactories.h"
 #include "NameColumn.h"
 #include "PropertyCustomizationHelpers.h"
-#include "TransactionCommon.h"
+#include "Misc/TransactionCommon.h"
 #include "Widgets/Images/SImage.h"
 #include "Widgets/Input/SEditableTextBox.h"
 #include "Widgets/Input/SButton.h"
diff --git a/Engine/Plugins/Chooser/Source/ChooserEditor/Private/ObjectClassColumnEditor.cpp b/Engine/Plugins/Chooser/Source/ChooserEditor/Private/ObjectClassColumnEditor.cpp
--- a/Engine/Plugins/Chooser/Source/ChooserEditor/Private/ObjectClassColumnEditor.cpp
+++ b/Engine/Plugins/Chooser/Source/ChooserEditor/Private/ObjectClassColumnEditor.cpp
@@ -10,7 +10,7 @@
 #include "ObjectChooserWidgetFactories.h"
 #include "ObjectClassColumn.h"
 #include "PropertyCustomizationHelpers.h"
-#include "TransactionCommon.h"
+#include "Misc/TransactionCommon.h"
 #include "Widgets/Images/SImage.h"
 #include "Widgets/Input/SButton.h"
 #include "Widgets/Layout/SBox.h"
diff --git a/Engine/Plugins/Chooser/Source/ChooserEditor/Private/ObjectColumnEditor.cpp b/Engine/Plugins/Chooser/Source/ChooserEditor/Private/ObjectColumnEditor.cpp
--- a/Engine/Plugins/Chooser/Source/ChooserEditor/Private/ObjectColumnEditor.cpp
+++ b/Engine/Plugins/Chooser/Source/ChooserEditor/Private/ObjectColumnEditor.cpp
@@ -9,7 +9,7 @@
 #include "ObjectChooserWidgetFactories.h"
 #include "ObjectColumn.h"
 #include "PropertyCustomizationHelpers.h"
-#include "TransactionCommon.h"
+#include "Misc/TransactionCommon.h"
 #include "Widgets/Images/SImage.h"
 #include "Widgets/Input/SButton.h"
 #include "Widgets/Layout/SBox.h"
diff --git a/Engine/Plugins/Chooser/Source/ChooserEditor/Private/OutputEnumColumnEditor.cpp b/Engine/Plugins/Chooser/Source/ChooserEditor/Private/OutputEnumColumnEditor.cpp
--- a/Engine/Plugins/Chooser/Source/ChooserEditor/Private/OutputEnumColumnEditor.cpp
+++ b/Engine/Plugins/Chooser/Source/ChooserEditor/Private/OutputEnumColumnEditor.cpp
@@ -9,7 +9,7 @@
 #include "ChooserTableEditor.h"
 #include "GraphEditorSettings.h"
 #include "SEnumCombo.h"
-#include "TransactionCommon.h"
+#include "Misc/TransactionCommon.h"
 #include "Widgets/Input/SButton.h"
 #include "ScopedTransaction.h"
 
diff --git a/Engine/Plugins/Chooser/Source/ChooserEditor/Private/OutputNameColumnEditor.cpp b/Engine/Plugins/Chooser/Source/ChooserEditor/Private/OutputNameColumnEditor.cpp
--- a/Engine/Plugins/Chooser/Source/ChooserEditor/Private/OutputNameColumnEditor.cpp
+++ b/Engine/Plugins/Chooser/Source/ChooserEditor/Private/OutputNameColumnEditor.cpp
@@ -9,7 +9,7 @@
 #include "ObjectChooserWidgetFactories.h"
 #include "OutputNameColumn.h"
 #include "PropertyCustomizationHelpers.h"
-#include "TransactionCommon.h"
+#include "Misc/TransactionCommon.h"
 #include "Widgets/Images/SImage.h"
 #include "Widgets/Input/SEditableTextBox.h"
 #include "Widgets/Input/SButton.h"
diff --git a/Engine/Plugins/Chooser/Source/ChooserEditor/Private/OutputObjectColumnEditor.cpp b/Engine/Plugins/Chooser/Source/ChooserEditor/Private/OutputObjectColumnEditor.cpp
--- a/Engine/Plugins/Chooser/Source/ChooserEditor/Private/OutputObjectColumnEditor.cpp
+++ b/Engine/Plugins/Chooser/Source/ChooserEditor/Private/OutputObjectColumnEditor.cpp
@@ -9,7 +9,7 @@
 #include "ObjectChooserWidgetFactories.h"
 #include "OutputObjectColumn.h"
 #include "PropertyCustomizationHelpers.h"
-#include "TransactionCommon.h"
+#include "Misc/TransactionCommon.h"
 #include "Widgets/Images/SImage.h"
 #include "Widgets/Input/SButton.h"
 #include "Widgets/Layout/SBox.h"
diff --git a/Engine/Plugins/Developer/Concert/ConcertSync/ConcertSyncClient/Source/ConcertSyncClient/Private/ConcertClientTransactionBridge.cpp b/Engine/Plugins/Developer/Concert/ConcertSync/ConcertSyncClient/Source/ConcertSyncClient/Private/ConcertClientTransactionBridge.cpp
--- a/Engine/Plugins/Developer/Concert/ConcertSync/ConcertSyncClient/Source/ConcertSyncClient/Private/ConcertClientTransactionBridge.cpp
+++ b/Engine/Plugins/Developer/Concert/ConcertSync/ConcertSyncClient/Source/ConcertSyncClient/Private/ConcertClientTransactionBridge.cpp
@@ -13,7 +13,7 @@
 #include "Misc/ITransactionObjectAnnotation.h"
 #include "Misc/PackageName.h"
 #include "Misc/CoreDelegates.h"
-#include "TransactionCommon.h"
+#include "Misc/TransactionCommon.h"
 #include "UObject/Package.h"
 #include "LevelInstance/LevelInstanceActor.h"
 #include "LevelInstance/LevelInstanceComponent.h"
diff --git a/Engine/Plugins/Experimental/UAF/UAFTestSuites/Source/UAFTestSuite/Private/AnimNextVariableReferenceTest.cpp b/Engine/Plugins/Experimental/UAF/UAFTestSuites/Source/UAFTestSuite/Private/AnimNextVariableReferenceTest.cpp
--- a/Engine/Plugins/Experimental/UAF/UAFTestSuites/Source/UAFTestSuite/Private/AnimNextVariableReferenceTest.cpp
+++ b/Engine/Plugins/Experimental/UAF/UAFTestSuites/Source/UAFTestSuite/Private/AnimNextVariableReferenceTest.cpp
@@ -5,7 +5,7 @@
 #if WITH_DEV_AUTOMATION_TESTS && WITH_EDITOR
 
 #include "CoreMinimal.h"
-#include "TransactionCommon.h"
+#include "Misc/TransactionCommon.h"
 #include "UncookedOnlyUtils.h"
 #include "Misc/AutomationTest.h"
 #include "Serialization/BufferArchive.h"
diff --git a/Engine/Source/Editor/UnrealEd/Classes/Editor/Transactor.h b/Engine/Source/Editor/UnrealEd/Classes/Editor/Transactor.h
--- a/Engine/Source/Editor/UnrealEd/Classes/Editor/Transactor.h
+++ b/Engine/Source/Editor/UnrealEd/Classes/Editor/Transactor.h
@@ -15,7 +15,7 @@
 #include "Serialization/ArchiveSerializedPropertyChain.h"
 #include "Misc/ITransaction.h"
 #include "Misc/TransactionObjectEvent.h"
-#include "TransactionCommon.h"
+#include "Misc/TransactionCommon.h"
 #include "Algo/Find.h"
 #include <type_traits>
 #include "Transactor.generated.h"
diff --git a/Engine/Source/Runtime/CoreUObject/Private/Misc/TransactionCommon.cpp b/Engine/Source/Runtime/CoreUObject/Private/Misc/TransactionCommon.cpp
new file mode 100644
--- /dev/null
+++ b/Engine/Source/Runtime/CoreUObject/Private/Misc/TransactionCommon.cpp
@@ -0,0 +1,644 @@
+// Copyright Epic Games, Inc. All Rights Reserved.
+
+#include "Misc/TransactionCommon.h"
+#include "Serialization/ArchiveSerializedPropertyChain.h"
+#include "UObject/UnrealType.h"
+
+namespace UE::Transaction
+{
+
+FPersistentObjectRef::FPersistentObjectRef(UObject* InObject)
+{
+	RootObject = InObject;
+	{
+		auto UseOuter = [](const UObject* Obj)
+		{
+			if (Obj == nullptr || Obj->IsCapturingAsRootObjectForTransaction())
+			{
+				return false;
+			}
+
+			const bool bIsCDO = Obj->HasAllFlags(RF_ClassDefaultObject);
+			const UObject* CDO = bIsCDO ? Obj : nullptr;
+			const bool bIsClassCDO = (CDO != nullptr) ? (CDO->GetClass()->GetDefaultObject(false) == CDO) : false;
+			if (!bIsClassCDO && CDO)
+			{
+				// Likely a trashed CDO, try to recover
+				// Only known cause of this is ambiguous use of DSOs
+				CDO = CDO->GetClass()->GetDefaultObject(false);
+			}
+			const bool bIsDSO = Obj->HasAnyFlags(RF_DefaultSubObject);
+			if ((bIsCDO && bIsClassCDO) || bIsDSO)
+			{
+				return true;
+			}
+
+			for (OnUseOuterDelegate Delegate : GetOnUseOuterDelegates())
+			{
+				if (Delegate(Obj))
+				{
+					return true;
+				}
+			}
+			return false;
+		};
+
+		while (UseOuter(RootObject))
+		{
+			SubObjectHierarchyIDs.Add(RootObject->GetFName());
+			RootObject = RootObject->GetOuter();
+		}
+	}
+	check(RootObject);
+
+	if (SubObjectHierarchyIDs.Num() > 0)
+	{
+		ReferenceType = EReferenceType::SubObject;
+		Algo::Reverse(SubObjectHierarchyIDs);
+	}
+	else
+	{
+		ReferenceType = EReferenceType::RootObject;
+	}
+
+	// Make sure that when we look up the object we find the same thing:
+	checkSlow(Get() == InObject);
+}
+
+UObject* FPersistentObjectRef::Get() const
+{
+	if (ReferenceType == EReferenceType::SubObject)
+	{
+		check(SubObjectHierarchyIDs.Num() > 0);
+
+		UObject* CurrentObject = nullptr;
+
+		// Do we have a valid cached pointer?
+		if (!CachedRootObject.IsExplicitlyNull() && SubObjectHierarchyIDs.Num() == CachedSubObjectHierarchy.Num())
+		{
+			// Root object is a pointer test
+			{
+				CurrentObject = CachedRootObject.GetEvenIfUnreachable();
+				if (CurrentObject != RootObject)
+				{
+					CurrentObject = nullptr;
+				}
+			}
+
+			// All other sub-objects are a name test
+			for (int32 SubObjectIndex = 0; CurrentObject && SubObjectIndex < SubObjectHierarchyIDs.Num(); ++SubObjectIndex)
+			{
+				CurrentObject = CachedSubObjectHierarchy[SubObjectIndex].GetEvenIfUnreachable();
+				if (CurrentObject && CurrentObject->GetFName() != SubObjectHierarchyIDs[SubObjectIndex])
+				{
+					CurrentObject = nullptr;
+				}
+			}
+		}
+		if (CurrentObject)
+		{
+			return CurrentObject;
+		}
+
+		// Cached pointer is invalid
+		CachedRootObject.Reset();
+		CachedSubObjectHierarchy.Reset();
+
+		// Try to find and cache the subobject
+		CachedRootObject = RootObject;
+		CurrentObject = RootObject;
+		for (int32 SubObjectIndex = 0; CurrentObject && SubObjectIndex < SubObjectHierarchyIDs.Num(); ++SubObjectIndex)
+		{
+			CurrentObject = StaticFindObjectFast(UObject::StaticClass(), CurrentObject, SubObjectHierarchyIDs[SubObjectIndex]);
+			CachedSubObjectHierarchy.Add(CurrentObject);
+		}
+		if (CurrentObject)
+		{
+			check(!CachedRootObject.IsExplicitlyNull() && SubObjectHierarchyIDs.Num() == CachedSubObjectHierarchy.Num());
+			return CurrentObject;
+		}
+
+		// Cached pointer is invalid
+		CachedRootObject.Reset();
+		CachedSubObjectHierarchy.Reset();
+
+		return nullptr;
+	}
+
+	return RootObject;
+}
+
+void FPersistentObjectRef::AddReferencedObjects(FReferenceCollector& Collector)
+{
+	Collector.AddReferencedObject(RootObject);
+
+	if (ReferenceType == EReferenceType::SubObject)
+	{
+		// We can't refresh the resolved pointers during ARO, as it's not safe to call FindObject to update the cache if stale
+		// Instead we'll just ARO whatever we may have cached, as this may result in the resolved pointers being updated anyway
+		// Note: This is needed as sub-objects may be subject to GC while inside the transaction buffer, as the references from their root 
+		// object may have been removed (eg, a component on an actor will no longer be referenced by the actor after a delete operation)
+		for (TWeakObjectPtr<UObject>& CachedSubObject : CachedSubObjectHierarchy)
+		{
+			Collector.AddReferencedObject(CachedSubObject);
+		}
+	}
+}
+
+TArray<FPersistentObjectRef::OnUseOuterDelegate>& FPersistentObjectRef::GetOnUseOuterDelegates()
+{
+	static TArray<FPersistentObjectRef::OnUseOuterDelegate> Delegates;
+	return Delegates;
+}
+
+FSerializedTaggedData FSerializedTaggedData::FromOffsetAndSize(const int64 InOffset, const int64 InSize)
+{
+	return FSerializedTaggedData{ InOffset, InSize };
+}
+
+FSerializedTaggedData FSerializedTaggedData::FromStartAndEnd(const int64 InStart, const int64 InEnd)
+{
+	return FSerializedTaggedData{ InStart, InEnd - InStart };
+}
+
+void FSerializedTaggedData::AppendSerializedData(const int64 InOffset, const int64 InSize)
+{
+	if (DataOffset == INDEX_NONE)
+	{
+		DataOffset = InOffset;
+		DataSize = InSize;
+	}
+	else
+	{
+		DataOffset = FMath::Min(DataOffset, InOffset);
+		DataSize = FMath::Max(InOffset + InSize - DataOffset, DataSize);
+	}
+}
+
+void FSerializedTaggedData::AppendSerializedData(const FSerializedTaggedData& InData)
+{
+	if (InData.HasSerializedData())
+	{
+		AppendSerializedData(InData.DataOffset, InData.DataSize);
+	}
+}
+
+bool FSerializedTaggedData::HasSerializedData() const
+{
+	return DataOffset != INDEX_NONE && DataSize != 0;
+}
+
+
+void FSerializedObjectData::Read(void* Dest, int64 Offset, int64 Num) const
+{
+	checkSlow(Offset + Num <= Data.Num());
+	FMemory::Memcpy(Dest, &Data[Offset], Num);
+}
+
+void FSerializedObjectData::Write(const void* Src, int64 Offset, int64 Num)
+{
+	if (Offset == Data.Num())
+	{
+		Data.AddUninitialized(Num);
+	}
+	FMemory::Memcpy(&Data[Offset], Src, Num);
+}
+
+
+const FName TaggedDataKey_UnknownData = ".UnknownData";
+const FName TaggedDataKey_ScriptData = ".ScriptData";
+
+FSerializedObjectDataReader::FSerializedObjectDataReader(const FSerializedObject& InSerializedObject)
+	: SerializedObject(InSerializedObject)
+{
+	SetIsLoading(true);
+}
+
+void FSerializedObjectDataReader::Serialize(void* SerData, int64 Num)
+{
+	if (Num)
+	{
+		SerializedObject.SerializedData.Read(SerData, Offset, Num);
+		Offset += Num;
+	}
+}
+
+FArchive& FSerializedObjectDataReader::operator<<(FName& N)
+{
+	int32 NameIndex = 0;
+	*this << NameIndex;
+	N = SerializedObject.ReferencedNames[NameIndex];
+	return *this;
+}
+
+FArchive& FSerializedObjectDataReader::operator<<(UObject*& Res)
+{
+	int32 ObjectIndex = 0;
+	*this << ObjectIndex;
+	if (ObjectIndex != INDEX_NONE)
+	{
+		Res = SerializedObject.ReferencedObjects[ObjectIndex].Get();
+	}
+	else
+	{
+		Res = nullptr;
+	}
+	return *this;
+}
+
+
+namespace Internal
+{
+
+FSerializedObjectDataWriterCommon::FSerializedObjectDataWriterCommon(FSerializedObjectData& InSerializedData)
+	: SerializedData(InSerializedData)
+{
+	SetIsSaving(true);
+}
+
+void FSerializedObjectDataWriterCommon::Serialize(void* SerData, int64 Num)
+{
+	if (Num)
+	{
+		int64 DataIndex = Offset;
+		SerializedData.Write(SerData, Offset, Num);
+		Offset += Num;
+
+		OnDataSerialized(DataIndex, Num);
+	}
+}
+
+} // namespace Internal
+
+
+FSerializedObjectDataWriter::FSerializedObjectDataWriter(FSerializedObject& InSerializedObject)
+	: Internal::FSerializedObjectDataWriterCommon(InSerializedObject.SerializedData)
+	, SerializedObject(InSerializedObject)
+{
+	for (int32 ObjIndex = 0; ObjIndex < SerializedObject.ReferencedObjects.Num(); ++ObjIndex)
+	{
+		ObjectMap.Add(SerializedObject.ReferencedObjects[ObjIndex].Get(), ObjIndex);
+	}
+
+	for (int32 NameIndex = 0; NameIndex < SerializedObject.ReferencedNames.Num(); ++NameIndex)
+	{
+		NameMap.Add(SerializedObject.ReferencedNames[NameIndex], NameIndex);
+	}
+}
+
+FArchive& FSerializedObjectDataWriter::operator<<(FName& N)
+{
+	int32 NameIndex = INDEX_NONE;
+	const int32* NameIndexPtr = NameMap.Find(N);
+	if (NameIndexPtr)
+	{
+		NameIndex = *NameIndexPtr;
+	}
+	else
+	{
+		NameIndex = SerializedObject.ReferencedNames.Add(N);
+		NameMap.Add(N, NameIndex);
+	}
+
+	return *this << NameIndex;
+}
+
+FArchive& FSerializedObjectDataWriter::operator<<(UObject*& Res)
+{
+	int32 ObjectIndex = INDEX_NONE;
+	const int32* ObjIndexPtr = ObjectMap.Find(Res);
+	if (ObjIndexPtr)
+	{
+		ObjectIndex = *ObjIndexPtr;
+	}
+	else if (Res)
+	{
+		ObjectIndex = SerializedObject.ReferencedObjects.Add(FPersistentObjectRef(Res));
+		ObjectMap.Add(Res, ObjectIndex);
+	}
+
+	return *this << ObjectIndex;
+}
+
+
+FDiffableObjectDataWriter::FDiffableObjectDataWriter(FDiffableObject& InDiffableObject, TArrayView<const FProperty*> InPropertiesToSerialize)
+	: Internal::FSerializedObjectDataWriterCommon(InDiffableObject.SerializedData)
+	, DiffableObject(InDiffableObject)
+	, PropertiesToSerialize(InPropertiesToSerialize)
+{
+	SetWantBinaryPropertySerialization(true);
+}
+
+FName FDiffableObjectDataWriter::GetTaggedDataKey() const
+{
+	FName TaggedDataKey;
+
+	// Is this known property data?
+	if (TaggedDataKey.IsNone())
+	{
+		const FArchiveSerializedPropertyChain* PropertyChain = GetSerializedPropertyChain();
+		if (PropertyChain && PropertyChain->GetNumProperties() > 0)
+		{
+			TaggedDataKey = CachedSerializedTaggedPropertyKey.SyncCache(*PropertyChain);
+		}
+	}
+
+	// Is this script data?
+	if (TaggedDataKey.IsNone() && bIsPerformingScriptSerialization)
+	{
+		bWasUsingTaggedDataKey_ScriptData = true;
+		TaggedDataKey = FName(TaggedDataKey_ScriptData, TaggedDataKeyIndex_ScriptData);
+	}
+	else if (bWasUsingTaggedDataKey_ScriptData)
+	{
+		++TaggedDataKeyIndex_ScriptData;
+		bWasUsingTaggedDataKey_ScriptData = false;
+	}
+
+	// Is this unknown data?
+	if (TaggedDataKey.IsNone())
+	{
+		bWasUsingTaggedDataKey_UnknownData = true;
+		TaggedDataKey = FName(TaggedDataKey_UnknownData, TaggedDataKeyIndex_UnknownData);
+	}
+	else if (bWasUsingTaggedDataKey_UnknownData)
+	{
+		++TaggedDataKeyIndex_UnknownData;
+		bWasUsingTaggedDataKey_UnknownData = false;
+	}
+
+	return TaggedDataKey;
+}
+
+bool FDiffableObjectDataWriter::DoesObjectMatchDiffableObject(const UObject* Obj) const
+{
+	FNameBuilder ObjPathName;
+	Obj->GetPathName(nullptr, ObjPathName);
+	return FName(ObjPathName) == DiffableObject.ObjectInfo.ObjectPathName;
+}
+
+bool FDiffableObjectDataWriter::ShouldSkipProperty(const FProperty* InProperty) const
+{
+	return (PropertiesToSerialize.Num() > 0 && !PropertiesToSerialize.Contains(InProperty))
+		|| InProperty->HasAnyPropertyFlags(CPF_NonTransactional | CPF_Deprecated)
+		|| Internal::FSerializedObjectDataWriterCommon::ShouldSkipProperty(InProperty);
+}
+
+void FDiffableObjectDataWriter::MarkScriptSerializationStart(const UObject* Obj)
+{
+	FArchiveUObject::MarkScriptSerializationStart(Obj);
+
+	if (DoesObjectMatchDiffableObject(Obj))
+	{
+		bIsPerformingScriptSerialization = true;
+	}
+}
+
+void FDiffableObjectDataWriter::MarkScriptSerializationEnd(const UObject* Obj)
+{
+	FArchiveUObject::MarkScriptSerializationEnd(Obj);
+
+	if (DoesObjectMatchDiffableObject(Obj))
+	{
+		bIsPerformingScriptSerialization = false;
+	}
+}
+
+void FDiffableObjectDataWriter::OnDataSerialized(int64 InOffset, int64 InNum)
+{
+	// Track this offset index in the serialized data
+	const FName SerializedTaggedDataKey = GetTaggedDataKey();
+	if (!SerializedTaggedDataKey.IsNone())
+	{
+		FSerializedTaggedData& SerializedTaggedData = CachedSerializedTaggedDataEntry.SyncCache(DiffableObject, SerializedTaggedDataKey);
+		SerializedTaggedData.AppendSerializedData(InOffset, InNum);
+	}
+}
+
+FArchive& FDiffableObjectDataWriter::operator<<(FName& N)
+{
+	uint32 NameDisplayIndex = N.GetDisplayIndex().ToUnstableInt();
+	int32 NameNumericSuffix = N.GetNumber();
+	return *this << NameDisplayIndex << NameNumericSuffix;
+}
+
+FArchive& FDiffableObjectDataWriter::operator<<(UObject*& Res)
+{
+	FName ResPathName;
+	if (Res)
+	{
+		FNameBuilder ResPathNameBuilder;
+		Res->GetPathName(nullptr, ResPathNameBuilder);
+		ResPathName = FName(ResPathNameBuilder);
+	}
+	return *this << ResPathName;
+}
+
+FName FDiffableObjectDataWriter::FCachedPropertyKey::SyncCache(const FArchiveSerializedPropertyChain& InPropertyChain)
+{
+	const uint32 CurrentUpdateCount = InPropertyChain.GetUpdateCount();
+	if (CurrentUpdateCount != LastUpdateCount)
+	{
+		CachedKey = InPropertyChain.GetNumProperties() > 0 ? InPropertyChain.GetPropertyFromRoot(0)->GetFName() : FName();
+		LastUpdateCount = CurrentUpdateCount;
+	}
+
+	return CachedKey;
+}
+
+FSerializedTaggedData& FDiffableObjectDataWriter::FCachedTaggedDataEntry::SyncCache(FDiffableObject& InDiffableObject, const FName InSerializedTaggedDataKey)
+{
+	if (CachedKey != InSerializedTaggedDataKey)
+	{
+		CachedKey = InSerializedTaggedDataKey;
+		CachedEntryPtr = &InDiffableObject.SerializedTaggedData.FindOrAdd(InSerializedTaggedDataKey);
+	}
+
+	return *CachedEntryPtr;
+}
+
+
+namespace DiffUtil
+{
+static FArchive* GCurrentDiffableObjectDataWriter = nullptr; // Set before GetDiffableObject calls its serialize impl
+	
+bool IsGeneratingDiffableObject(const FArchive& Ar)
+{
+	return GCurrentDiffableObjectDataWriter == &Ar;
+}
+
+FDiffableObject GetDiffableObject(const UObject* Object, const FGetDiffableObjectOptions& Options)
+{
+	FDiffableObject DiffableObject;
+	DiffableObject.SetObject(Object);
+
+	if (Options.bSerializeEvenIfPendingKill || !DiffableObject.ObjectInfo.bIsPendingKill)
+	{
+		FDiffableObjectDataWriter DiffWriter(DiffableObject, Options.PropertiesToSerialize);
+		TGuardValue<FArchive*> GuardGCurrentDiffableObjectDataWriter(GCurrentDiffableObjectDataWriter, &DiffWriter);
+		
+		switch (Options.ObjectSerializationMode)
+		{
+		case EGetDiffableObjectMode::SerializeObject:
+			const_cast<UObject*>(Object)->Serialize(DiffWriter);
+			break;
+		case EGetDiffableObjectMode::SerializeProperties:
+			const_cast<UObject*>(Object)->SerializeScriptProperties(DiffWriter);
+			break;
+		case EGetDiffableObjectMode::Custom:
+			check(Options.CustomSerializer);
+			Options.CustomSerializer(DiffWriter);
+			break;
+		default:
+			checkf(false, TEXT("Unknown EGetDiffableObjectMode!"));
+			break;
+		}
+	}
+
+	return DiffableObject;
+}
+
+const FDiffableObject& FDiffableObjectArchetypeCache::GetArchetypeDiffableObject(const UObject* Archetype, const FGetDiffableObjectOptions& Options)
+{
+	if (Archetype)
+	{
+		if (const FDiffableObject* ExistingDiffableObject = ArchetypeDiffableObjects.Find(Archetype))
+		{
+			return *ExistingDiffableObject;
+		}
+		else
+		{
+			return ArchetypeDiffableObjects.Add(Archetype, GetDiffableObject(Archetype, Options));
+		}
+	}
+
+	static const FDiffableObject EmptyDiffableObject;
+	return EmptyDiffableObject;
+}
+
+void GenerateObjectDataDiff_Impl(const FDiffableObject& OldDiffableObject, const FDiffableObject& NewDiffableObject, FTransactionObjectDeltaChange& OutDeltaChange, const FGenerateObjectDiffOptions& DiffOptions)
+{
+	auto IsTaggedDataBlockIdentical = [&OldDiffableObject, &NewDiffableObject, &OutDeltaChange](const FName TaggedDataKey, const FSerializedTaggedData& OldSerializedTaggedData, const FSerializedTaggedData& NewSerializedTaggedData) -> bool
+	{
+		// Binary compare the serialized data to see if something has changed for this property
+		bool bIsTaggedDataIdentical = OldSerializedTaggedData.DataSize == NewSerializedTaggedData.DataSize;
+		if (bIsTaggedDataIdentical && OldSerializedTaggedData.HasSerializedData())
+		{
+			bIsTaggedDataIdentical = FMemory::Memcmp(OldDiffableObject.SerializedData.GetPtr(OldSerializedTaggedData.DataOffset), NewDiffableObject.SerializedData.GetPtr(NewSerializedTaggedData.DataOffset), NewSerializedTaggedData.DataSize) == 0;
+		}
+		return bIsTaggedDataIdentical;
+	};
+
+	auto ShouldCompareTaggedData = [&DiffOptions](const FName TaggedDataKey) -> bool
+	{
+		if (TaggedDataKey.GetComparisonIndex() == TaggedDataKey_ScriptData.GetComparisonIndex())
+		{
+			// Never compare script data, as it's assumed to be overhead around the tagged property serialization
+			return false;
+		}
+
+		if (DiffOptions.ShouldSkipProperty && DiffOptions.ShouldSkipProperty(TaggedDataKey))
+		{
+			return false;
+		}
+
+		return true;
+	};
+
+	auto ShouldCompareAsNonPropertyData = [](const FName TaggedDataKey)
+	{
+		return TaggedDataKey.GetComparisonIndex() == TaggedDataKey_UnknownData.GetComparisonIndex();
+	};
+
+	for (const TPair<FName, FSerializedTaggedData>& NewNamePropertyPair : NewDiffableObject.SerializedTaggedData)
+	{
+		if (!ShouldCompareTaggedData(NewNamePropertyPair.Key))
+		{
+			continue;
+		}
+
+		const FSerializedTaggedData* OldSerializedTaggedData = OldDiffableObject.SerializedTaggedData.Find(NewNamePropertyPair.Key);
+		if (ShouldCompareAsNonPropertyData(NewNamePropertyPair.Key))
+		{
+			if (DiffOptions.bFullDiff && !OutDeltaChange.bHasNonPropertyChanges && (!OldSerializedTaggedData || !IsTaggedDataBlockIdentical(NewNamePropertyPair.Key, *OldSerializedTaggedData, NewNamePropertyPair.Value)))
+			{
+				OutDeltaChange.bHasNonPropertyChanges = true;
+			}
+		}
+		else if (OldSerializedTaggedData)
+		{
+			if (!IsTaggedDataBlockIdentical(NewNamePropertyPair.Key, *OldSerializedTaggedData, NewNamePropertyPair.Value))
+			{
+				OutDeltaChange.ChangedProperties.AddUnique(NewNamePropertyPair.Key);
+			}
+		}
+		else if (DiffOptions.bFullDiff)
+		{
+			// Missing property, assume that the property changed
+			OutDeltaChange.ChangedProperties.AddUnique(NewNamePropertyPair.Key);
+		}
+	}
+
+	if (DiffOptions.bFullDiff)
+	{
+		for (const TPair<FName, FSerializedTaggedData>& OldNamePropertyPair : OldDiffableObject.SerializedTaggedData)
+		{
+			if (!ShouldCompareTaggedData(OldNamePropertyPair.Key) || NewDiffableObject.SerializedTaggedData.Contains(OldNamePropertyPair.Key))
+			{
+				continue;
+			}
+
+			if (ShouldCompareAsNonPropertyData(OldNamePropertyPair.Key))
+			{
+				OutDeltaChange.bHasNonPropertyChanges = true;
+			}
+			else
+			{
+				// Missing property, assume that the property changed
+				OutDeltaChange.ChangedProperties.AddUnique(OldNamePropertyPair.Key);
+			}
+		}
+	}
+}
+
+FTransactionObjectDeltaChange GenerateObjectDiff(const FDiffableObject& OldDiffableObject, const FDiffableObject& NewDiffableObject, const FGenerateObjectDiffOptions& DiffOptions, FDiffableObjectArchetypeCache* ArchetypeCache)
+{
+	FTransactionObjectDeltaChange DeltaChange;
+	GenerateObjectDiff(OldDiffableObject, NewDiffableObject, DeltaChange, DiffOptions, ArchetypeCache);
+	return DeltaChange;
+}
+
+void GenerateObjectDiff(const FDiffableObject& OldDiffableObject, const FDiffableObject& NewDiffableObject, FTransactionObjectDeltaChange& OutDeltaChange, const FGenerateObjectDiffOptions& DiffOptions, FDiffableObjectArchetypeCache* ArchetypeCache)
+{
+	if (DiffOptions.bFullDiff)
+	{
+		OutDeltaChange.bHasNameChange |= OldDiffableObject.ObjectInfo.ObjectName != NewDiffableObject.ObjectInfo.ObjectName;
+		OutDeltaChange.bHasOuterChange |= OldDiffableObject.ObjectInfo.ObjectOuterPathName != NewDiffableObject.ObjectInfo.ObjectOuterPathName;
+		OutDeltaChange.bHasExternalPackageChange |= OldDiffableObject.ObjectInfo.ObjectExternalPackageName != NewDiffableObject.ObjectInfo.ObjectExternalPackageName;
+		OutDeltaChange.bHasPendingKillChange |= OldDiffableObject.ObjectInfo.bIsPendingKill != NewDiffableObject.ObjectInfo.bIsPendingKill;
+	}
+
+	if (DiffOptions.bDiffDataEvenIfPendingKill || OldDiffableObject.ObjectInfo.bIsPendingKill == NewDiffableObject.ObjectInfo.bIsPendingKill)
+	{
+		GenerateObjectDataDiff_Impl(OldDiffableObject, NewDiffableObject, OutDeltaChange, DiffOptions);
+	}
+	else
+	{
+		if (!NewDiffableObject.ObjectInfo.bIsPendingKill)
+		{
+			check(OldDiffableObject.ObjectInfo.bIsPendingKill);
+
+			FDiffableObjectArchetypeCache LocalArchetypeCache;
+			FDiffableObjectArchetypeCache& ArchetypeCacheRef = ArchetypeCache ? *ArchetypeCache : LocalArchetypeCache;
+
+			// The object has been restored to life, so use the archetype state when generating the data diff (as there likely isn't any data cached for the old object; see GetDiffableObject)
+			const FDiffableObject& DiffableArchetypeObject = ArchetypeCacheRef.GetArchetypeDiffableObject(NewDiffableObject.ObjectArchetype.Get(), DiffOptions.ArchetypeOptions);
+			GenerateObjectDataDiff_Impl(DiffableArchetypeObject, NewDiffableObject, OutDeltaChange, DiffOptions);
+		}
+		// else, the object has become pending kill, so there's nothing more to diff after calculating the object info change above (as there likely isn't any cached data for the new object; see GetDiffableObject)
+	}
+}
+
+} // namespace DiffUtil
+
+} // namespace UE::Transaction
diff --git a/Engine/Source/Runtime/CoreUObject/Public/Misc/TransactionCommon.h b/Engine/Source/Runtime/CoreUObject/Public/Misc/TransactionCommon.h
new file mode 100644
--- /dev/null
+++ b/Engine/Source/Runtime/CoreUObject/Public/Misc/TransactionCommon.h
@@ -0,0 +1,503 @@
+// Copyright Epic Games, Inc. All Rights Reserved.
+
+#pragma once
+
+#include "CoreMinimal.h"
+#include "UObject/Object.h"
+#include "UObject/ObjectPtr.h"
+#include "UObject/WeakObjectPtr.h"
+#include "Serialization/ArchiveUObject.h"
+#include "Serialization/ArchiveSerializedPropertyChain.h"
+#include "Templates/TypeHash.h"
+#include "Misc/TransactionObjectEvent.h"
+#include "Containers/SortedMap.h"
+#include "Delegates/Delegate.h"
+
+#define UE_API COREUOBJECT_API
+
+class FReferenceCollector;
+
+namespace UE::Transaction
+{
+
+/**
+	This type is necessary because the blueprint system is destroying and creating
+	CDOs at edit time (usually on compile, but also on load), but also stores user
+	entered data in the CDO. We "need"  changes to a CDO to persist across instances
+	because as we undo and redo we  need to apply changes to different instances of
+	the CDO - alternatively we could destroy and create the CDO as part of a transaction
+	(this alternative is the reason for the bunny ears around need).
+
+	DanO: My long term preference is for the editor to use a dynamic, mutable type
+	(rather than the CDO) to store editor data. The CDO can then be re-instanced (or not)
+	as runtime code requires.
+*/
+struct FPersistentObjectRef
+{
+public:
+	FPersistentObjectRef() = default;
+	UE_API explicit FPersistentObjectRef(UObject* InObject);
+
+	bool Serialize(FArchive& Ar)
+	{
+		Ar << *this;
+		return true;
+	}
+
+	friend FArchive& operator<<(FArchive& Ar, FPersistentObjectRef& ReferencedObject)
+	{
+		Ar << (std::underlying_type_t<EReferenceType>&)ReferencedObject.ReferenceType;
+		Ar << ReferencedObject.RootObject;
+		Ar << ReferencedObject.SubObjectHierarchyIDs;
+		return Ar;
+	}
+
+	friend bool operator==(const FPersistentObjectRef& LHS, const FPersistentObjectRef& RHS)
+	{
+		return LHS.ReferenceType == RHS.ReferenceType
+			&& LHS.RootObject == RHS.RootObject
+			&& (LHS.ReferenceType != EReferenceType::SubObject || LHS.SubObjectHierarchyIDs == RHS.SubObjectHierarchyIDs);
+	}
+
+	friend bool operator!=(const FPersistentObjectRef& LHS, const FPersistentObjectRef& RHS)
+	{
+		return !(LHS == RHS);
+	}
+
+	friend uint32 GetTypeHash(const FPersistentObjectRef& InObjRef)
+	{
+		return HashCombine(::GetTypeHash(InObjRef.ReferenceType), GetTypeHash(InObjRef.RootObject));
+	}
+
+	bool IsRootObjectReference() const
+	{
+		return ReferenceType == EReferenceType::RootObject;
+	}
+
+	bool IsSubObjectReference() const
+	{
+		return ReferenceType == EReferenceType::SubObject;
+	}
+
+	UE_API UObject* Get() const;
+	
+	UE_API void AddReferencedObjects(FReferenceCollector& Collector);
+
+	using OnUseOuterDelegate = bool(*)(const UObject*);
+	UE_API static TArray<OnUseOuterDelegate>& GetOnUseOuterDelegates();
+
+private:
+	/** This enum represents all of the different special cases we are handling with this type */
+	enum class EReferenceType : uint8
+	{
+		Unknown,
+		RootObject,
+		SubObject,
+	};
+
+	/** The reference type we're handling */
+	EReferenceType ReferenceType = EReferenceType::Unknown;
+	/** Stores the object pointer when ReferenceType==RootObject, and the outermost pointer of the sub-object chain when when ReferenceType==SubObject */
+	TObjectPtr<UObject> RootObject = nullptr;
+	/** Stores the sub-object name chain when ReferenceType==SubObject */
+	TArray<FName, TInlineAllocator<4>> SubObjectHierarchyIDs;
+
+	/** Cached pointers corresponding to RootObject when ReferenceType==SubObject (@note cache needs testing on access as it may have become stale) */
+	mutable TWeakObjectPtr<UObject> CachedRootObject;
+	/** Cache of pointers corresponding to the items within SubObjectHierarchyIDs when ReferenceType==SubObject (@note cache needs testing on access as it may have become stale) */
+	mutable TArray<TWeakObjectPtr<UObject>, TInlineAllocator<4>> CachedSubObjectHierarchy;
+};
+
+struct FSerializedTaggedData
+{
+public:
+	static UE_API FSerializedTaggedData FromOffsetAndSize(const int64 InOffset, const int64 InSize);
+	static UE_API FSerializedTaggedData FromStartAndEnd(const int64 InStart, const int64 InEnd);
+
+	UE_API void AppendSerializedData(const int64 InOffset, const int64 InSize);
+	UE_API void AppendSerializedData(const FSerializedTaggedData& InData);
+
+	UE_API bool HasSerializedData() const;
+
+	int64 GetStart() const
+	{
+		return DataOffset;
+	}
+
+	int64 GetEnd() const
+	{
+		return DataOffset + DataSize;
+	}
+
+	friend bool operator==(const FSerializedTaggedData& LHS, const FSerializedTaggedData& RHS)
+	{
+		return LHS.DataOffset == RHS.DataOffset
+			&& LHS.DataSize == RHS.DataSize;
+	}
+
+	friend bool operator!=(const FSerializedTaggedData& LHS, const FSerializedTaggedData& RHS)
+	{
+		return !(LHS == RHS);
+	}
+
+	/** Offset to the start of the tagged data within the serialized object */
+	int64 DataOffset = INDEX_NONE;
+
+	/** Size (in bytes) of the tagged data within the serialized object */
+	int64 DataSize = 0;
+};
+
+struct FSerializedObjectData
+{
+public:
+	bool Serialize(FArchive& Ar)
+	{
+		Ar << *this;
+		return true;
+	}
+
+	friend FArchive& operator<<(FArchive& Ar, FSerializedObjectData& SerializedData)
+	{
+		Ar << SerializedData.Data;
+		return Ar;
+	}
+
+	friend bool operator==(const FSerializedObjectData& LHS, const FSerializedObjectData& RHS)
+	{
+		return LHS.Data == RHS.Data;
+	}
+
+	friend bool operator!=(const FSerializedObjectData& LHS, const FSerializedObjectData& RHS)
+	{
+		return !(LHS == RHS);
+	}
+
+	UE_API void Read(void* Dest, int64 Offset, int64 Num) const;
+	UE_API void Write(const void* Src, int64 Offset, int64 Num);
+	
+	const void* GetPtr(int64 Offset) const
+	{
+		return &Data[Offset];
+	}
+
+	void Reset()
+	{
+		Data.Reset();
+	}
+
+	int64 Num() const
+	{
+		return Data.Num();
+	}
+
+private:
+	TArray64<uint8> Data;
+};
+
+struct FSerializedObjectInfo : public FTransactionObjectId
+{
+public:
+	FSerializedObjectInfo() = default;
+
+	explicit FSerializedObjectInfo(const UObject* InObject)
+	{
+		SetObject(InObject);
+	}
+
+	void SetObject(const UObject* InObject)
+	{
+		FTransactionObjectId::SetObject(InObject);
+		bIsPendingKill = !IsValid(InObject);
+	}
+
+	void Reset()
+	{
+		FTransactionObjectId::Reset();
+		bIsPendingKill = false;
+	}
+
+	void Swap(FSerializedObjectInfo& Other)
+	{
+		FTransactionObjectId::Swap(Other);
+		Exchange(bIsPendingKill, Other.bIsPendingKill);
+	}
+
+	/** The pending kill state of the object when it was serialized */
+	bool bIsPendingKill = false;
+};
+
+struct FSerializedObject
+{
+public:
+	void Reset()
+	{
+		SerializedData.Reset();
+		ReferencedObjects.Reset();
+		ReferencedNames.Reset();
+	}
+
+	void Swap(FSerializedObject& Other)
+	{
+		Exchange(SerializedData, Other.SerializedData);
+		Exchange(ReferencedObjects, Other.ReferencedObjects);
+		Exchange(ReferencedNames, Other.ReferencedNames);
+	}
+
+	/** The serialized data for the transacted object */
+	FSerializedObjectData SerializedData;
+
+	/** External objects referenced by the transacted object */
+	TArray<FPersistentObjectRef> ReferencedObjects;
+	
+	/** Names referenced by the transacted object */
+	TArray<FName> ReferencedNames;
+};
+
+struct FDiffableObject
+{
+public:
+	void SetObject(const UObject* InObject)
+	{
+		ObjectInfo.SetObject(InObject);
+		ObjectArchetype = FPersistentObjectRef(InObject->GetArchetype());
+		ObjectFlags = InObject->GetFlags();
+	}
+
+	void Reset()
+	{
+		ObjectInfo.Reset();
+		ObjectArchetype = FPersistentObjectRef();
+		ObjectFlags = RF_NoFlags;
+		SerializedData.Reset();
+		SerializedTaggedData.Reset();
+	}
+
+	void Swap(FDiffableObject& Other)
+	{
+		ObjectInfo.Swap(Other.ObjectInfo);
+		Exchange(ObjectArchetype, Other.ObjectArchetype);
+		Exchange(ObjectFlags, Other.ObjectFlags);
+		Exchange(SerializedData, Other.SerializedData);
+		Exchange(SerializedTaggedData, Other.SerializedTaggedData);
+	}
+
+	/** Information about the object when it was serialized */
+	FSerializedObjectInfo ObjectInfo;
+
+	/** The archetype of the object when it was serialized */
+	FPersistentObjectRef ObjectArchetype;
+
+	/** The flags of the object when it was serialized */
+	EObjectFlags ObjectFlags = RF_NoFlags;
+
+	/** The serialized data for the diffable object */
+	FSerializedObjectData SerializedData;
+
+	/** Information about tagged data (mainly properties) that were serialized within this object */
+	TSortedMap<FName, FSerializedTaggedData, FDefaultAllocator, FNameFastLess> SerializedTaggedData;
+};
+
+UE_API extern const FName TaggedDataKey_UnknownData;
+UE_API extern const FName TaggedDataKey_ScriptData;
+
+/** Core archive to read a transaction object from the buffer. */
+class FSerializedObjectDataReader : public FArchiveUObject
+{
+public:
+	UE_API FSerializedObjectDataReader(const FSerializedObject& InSerializedObject);
+
+	virtual int64 Tell() override { return Offset; }
+	virtual void Seek(int64 InPos) override { Offset = InPos; }
+	virtual int64 TotalSize() override { return SerializedObject.SerializedData.Num(); }
+
+protected:
+	UE_API virtual void Serialize(void* SerData, int64 Num) override;
+
+	using FArchiveUObject::operator<<;
+	UE_API virtual FArchive& operator<<(class FName& N) override;
+	UE_API virtual FArchive& operator<<(class UObject*& Res) override;
+
+	const FSerializedObject& SerializedObject;
+	int64 Offset = 0;
+};
+
+namespace Internal
+{
+
+/** Core archive to write a transaction object to the buffer. */
+class FSerializedObjectDataWriterCommon : public FArchiveUObject
+{
+public:
+	UE_API FSerializedObjectDataWriterCommon(FSerializedObjectData& InSerializedData);
+
+	virtual int64 Tell() override { return Offset; }
+	virtual void Seek(int64 InPos) override { checkSlow(Offset <= SerializedData.Num()); Offset = InPos; }
+	virtual int64 TotalSize() override { return SerializedData.Num(); }
+
+protected:
+	UE_API virtual void Serialize(void* SerData, int64 Num) override;
+
+	virtual void OnDataSerialized(int64 InOffset, int64 InNum) {}
+
+	FSerializedObjectData& SerializedData;
+	int64 Offset = 0;
+};
+
+} // namespace Internal
+
+/** Core archive to write a transaction object to the buffer. */
+class FSerializedObjectDataWriter : public Internal::FSerializedObjectDataWriterCommon
+{
+public:
+	UE_API FSerializedObjectDataWriter(FSerializedObject& InSerializedObject);
+
+protected:
+	using FArchiveUObject::operator<<;
+	UE_API virtual FArchive& operator<<(class FName& N) override;
+	UE_API virtual FArchive& operator<<(class UObject*& Res) override;
+
+	FSerializedObject& SerializedObject;
+	int64 Offset = 0;
+
+	TMap<UObject*, int32> ObjectMap;
+	TMap<FName, int32> NameMap;
+};
+
+/** Core archive to write a diffable object to the buffer. */
+class FDiffableObjectDataWriter : public Internal::FSerializedObjectDataWriterCommon
+{
+public:
+	UE_API FDiffableObjectDataWriter(FDiffableObject& InDiffableObject, TArrayView<const FProperty*> InPropertiesToSerialize = TArrayView<const FProperty*>());
+
+protected:
+	UE_API FName GetTaggedDataKey() const;
+
+	UE_API bool DoesObjectMatchDiffableObject(const UObject* Obj) const;
+
+	UE_API virtual bool ShouldSkipProperty(const FProperty* InProperty) const override;
+
+	UE_API virtual void MarkScriptSerializationStart(const UObject* Obj) override;
+	UE_API virtual void MarkScriptSerializationEnd(const UObject* Obj) override;
+
+	UE_API virtual void OnDataSerialized(int64 InOffset, int64 InNum) override;
+
+	using FArchiveUObject::operator<<;
+	UE_API virtual FArchive& operator<<(class FName& N) override;
+	UE_API virtual FArchive& operator<<(class UObject*& Res) override;
+
+private:
+	struct FCachedPropertyKey
+	{
+	public:
+		FName SyncCache(const FArchiveSerializedPropertyChain& InPropertyChain);
+
+	private:
+		FName CachedKey;
+		uint32 LastUpdateCount = 0;
+	};
+
+	struct FCachedTaggedDataEntry
+	{
+	public:
+		FSerializedTaggedData& SyncCache(FDiffableObject& InDiffableObject, const FName InSerializedTaggedDataKey);
+
+	private:
+		FName CachedKey;
+		FSerializedTaggedData* CachedEntryPtr = nullptr;
+	};
+
+	FDiffableObject& DiffableObject;
+	TArrayView<const FProperty*> PropertiesToSerialize;
+
+	bool bIsPerformingScriptSerialization = false;
+
+	mutable bool bWasUsingTaggedDataKey_UnknownData = false;
+	mutable bool bWasUsingTaggedDataKey_ScriptData = false;
+	mutable int32 TaggedDataKeyIndex_UnknownData = 0;
+	mutable int32 TaggedDataKeyIndex_ScriptData = 0;
+	mutable FCachedPropertyKey CachedSerializedTaggedPropertyKey;
+	mutable FCachedTaggedDataEntry CachedSerializedTaggedDataEntry;
+};
+
+namespace DiffUtil
+{
+UE_API bool IsGeneratingDiffableObject(const FArchive& Ar);
+
+enum class EGetDiffableObjectMode : uint8
+{
+	/** Serialize the entire object state by calling its Serialize function */
+	SerializeObject,
+
+	/** Serialize the property state of the object by calling its SerializeScriptProperties function */
+	SerializeProperties,
+
+	/** Serialize the object via a custom serialize function */
+	Custom,
+};
+
+struct FGetDiffableObjectOptions
+{
+	/** How should we serialize the object for diffing? */
+	EGetDiffableObjectMode ObjectSerializationMode = EGetDiffableObjectMode::SerializeObject;
+
+	/** Custom serializer for the object (must be set when ObjectSerializationMode == Custom) */
+	TFunction<void(FDiffableObjectDataWriter&)> CustomSerializer;
+
+	/** Optional list of properties to serialize on the object, or an empty array to serialize all properties */
+	TArrayView<const FProperty*> PropertiesToSerialize;
+
+	/** Should we still serialize this object if it's considered pending kill? */
+	bool bSerializeEvenIfPendingKill = false;
+};
+
+/**
+ * Get an object snapshot that can be diffed later.
+ */
+UE_API FDiffableObject GetDiffableObject(const UObject* Object, const FGetDiffableObjectOptions& Options = FGetDiffableObjectOptions());
+
+struct FGenerateObjectDiffOptions
+{
+	FGenerateObjectDiffOptions()
+	{
+		ArchetypeOptions.bSerializeEvenIfPendingKill = true;
+	}
+
+	/** Options used when getting the diffable state of archive objects */
+	FGetDiffableObjectOptions ArchetypeOptions;
+
+	/** Optional function used to skip comparing certain properties within the diffable data */
+	TFunction<bool(FName)> ShouldSkipProperty;
+
+	/** Should we perform a "full" diff? (compares object info, and considers any missing properties to have been changed) */
+	bool bFullDiff = true;
+
+	/**
+	 * Should we diff the object state even if it's considered pending kill?
+	 * When false:
+	 *   If the old object was pending kill and the new object isn't, then the diff will be performed against the archetype of the new object.
+	 *   If the old object wasn't pending kill and the new object is, then the data diff will be skipped and only the object info diffed.
+	 */
+	bool bDiffDataEvenIfPendingKill = false;
+};
+
+/** Optional cache for requests to get the diffable state for a given archetype object */
+class FDiffableObjectArchetypeCache
+{
+public:
+	const FDiffableObject& GetArchetypeDiffableObject(const UObject* Archetype, const FGetDiffableObjectOptions& Options);
+
+private:
+	TMap<const UObject*, FDiffableObject> ArchetypeDiffableObjects;
+};
+
+/**
+ * Generate a diff between the two object snapshots.
+ */
+UE_API FTransactionObjectDeltaChange GenerateObjectDiff(const FDiffableObject& OldDiffableObject, const FDiffableObject& NewDiffableObject, const FGenerateObjectDiffOptions& DiffOptions = FGenerateObjectDiffOptions(), FDiffableObjectArchetypeCache* ArchetypeCache = nullptr);
+UE_API void GenerateObjectDiff(const FDiffableObject& OldDiffableObject, const FDiffableObject& NewDiffableObject, FTransactionObjectDeltaChange& OutDeltaChange, const FGenerateObjectDiffOptions& DiffOptions = FGenerateObjectDiffOptions(), FDiffableObjectArchetypeCache* ArchetypeCache = nullptr);
+
+} // namespace DiffUtil
+
+} // namespace UE::Transaction
+
+#undef UE_API
\ No newline at end of file
diff --git a/Engine/Source/Runtime/Engine/Private/Tests/TransactionDiffingTests.cpp b/Engine/Source/Runtime/Engine/Private/Tests/TransactionDiffingTests.cpp
--- a/Engine/Source/Runtime/Engine/Private/Tests/TransactionDiffingTests.cpp
+++ b/Engine/Source/Runtime/Engine/Private/Tests/TransactionDiffingTests.cpp
@@ -2,7 +2,7 @@
 
 #include "Tests/TransactionDiffingTests.h"
 #include "Misc/AutomationTest.h"
-#include "TransactionCommon.h"
+#include "Misc/TransactionCommon.h"
 
 #include UE_INLINE_GENERATED_CPP_BY_NAME(TransactionDiffingTests)
 
diff --git a/Engine/Source/Runtime/Engine/Private/TransactionCommon.cpp b/Engine/Source/Runtime/Engine/Private/TransactionCommon.cpp
deleted file mode 100644
diff --git a/Engine/Source/Runtime/Engine/Private/UnrealEngine.cpp b/Engine/Source/Runtime/Engine/Private/UnrealEngine.cpp
--- a/Engine/Source/Runtime/Engine/Private/UnrealEngine.cpp
+++ b/Engine/Source/Runtime/Engine/Private/UnrealEngine.cpp
@@ -135,6 +135,7 @@ UnrealEngine.cpp: Implements the UEngine class and helpers.
 #include "UObject/UObjectSerializeContext.h"
 #include "VerseVM/VVMDuplication.h"
 #include "PSOPrecacheSettings.h"
+#include "Components/ActorComponent.h"
 
 #include "Iris/IrisConfig.h"
 #include "Iris/ReplicationSystem/NetObjectFactoryRegistry.h"
@@ -198,6 +199,7 @@ UnrealEngine.cpp: Implements the UEngine class and helpers.
 #include "Engine/UserInterfaceSettings.h"
 #include "ComponentRecreateRenderStateContext.h"
 #include "Materials/MaterialRenderProxy.h"
+#include "Misc/TransactionCommon.h"
 
 #include "IMessageRpcClient.h"
 #include "IMessagingRpcModule.h"
@@ -294,6 +296,18 @@ IMPLEMENT_MODULE( FEngineModule, Engine );
 void InitScreenPercentage();
 void InitCustomDepth();
 
+namespace UE::Transaction
+{
+// This method adds UActorComponent objects to the outer object scan in
+// FSerializedObjectDataWriter and FDiffableObjectDataWriter.  This was the 
+// only reference to Engine types used by the mentioned above classes.
+bool IsActorComponentCreatedByConstructionScript(const UObject* Obj)
+{
+	const UActorComponent* AsComponent = Cast<UActorComponent>(Obj);
+	return AsComponent && AsComponent->IsCreatedByConstructionScript();
+}
+} // namespace UE::Transaction
+
 void OnChangeEngineCVarRequiringRecreateRenderState(IConsoleVariable* Var)
 {
 	// Propgate cvar change to static draw lists
@@ -303,6 +317,9 @@ void OnChangeEngineCVarRequiringRecreateRenderState(IConsoleVariable* Var)
 void FEngineModule::StartupModule()
 {
 	using namespace UE::UniversalObjectLocator;
+	using namespace UE::Transaction;
+	
+	FPersistentObjectRef::GetOnUseOuterDelegates().Add(&IsActorComponentCreatedByConstructionScript);
 
 	// Setup delegate callback for ProfilingHelpers to access current map name
 	extern const FString GetMapNameStatic();
@@ -416,6 +433,8 @@ void FEngineModule::StartupModule()
 
 void FEngineModule::ShutdownModule()
 {
+	using namespace UE::Transaction;
+
 	IPrimitiveComponent::RemoveImplementer(UPrimitiveComponent::StaticClass());
 	IStaticMeshComponent::RemoveImplementer(UStaticMeshComponent::StaticClass());
 	ISkinnedMeshComponent::RemoveImplementer(USkinnedMeshComponent::StaticClass());
@@ -439,6 +458,8 @@ void FEngineModule::ShutdownModule()
 	FSkinWeightProfileManager::OnShutdown();
 
 	UE::Anim::FSkeletonRemappingRegistry::Destroy();
+
+	FPersistentObjectRef::GetOnUseOuterDelegates().Remove(&IsActorComponentCreatedByConstructionScript);
 }
 
 /* Global variables
diff --git a/Engine/Source/Runtime/Engine/Public/ActorTransactionAnnotation.h b/Engine/Source/Runtime/Engine/Public/ActorTransactionAnnotation.h
--- a/Engine/Source/Runtime/Engine/Public/ActorTransactionAnnotation.h
+++ b/Engine/Source/Runtime/Engine/Public/ActorTransactionAnnotation.h
@@ -6,7 +6,7 @@
 #include "ComponentInstanceDataCache.h"
 
 #if WITH_EDITOR
-#include "TransactionCommon.h"
+#include "Misc/TransactionCommon.h"
 #include "Misc/ITransactionObjectAnnotation.h"
 #endif
 
diff --git a/Engine/Source/Runtime/Engine/Public/TransactionCommon.h b/Engine/Source/Runtime/Engine/Public/TransactionCommon.h
--- a/Engine/Source/Runtime/Engine/Public/TransactionCommon.h
+++ b/Engine/Source/Runtime/Engine/Public/TransactionCommon.h
@@ -2,494 +2,6 @@
 
 #pragma once
 
-#include "CoreMinimal.h"
-#include "UObject/Object.h"
-#include "UObject/ObjectPtr.h"
-#include "UObject/WeakObjectPtr.h"
-#include "Serialization/ArchiveUObject.h"
-#include "Serialization/ArchiveSerializedPropertyChain.h"
-#include "Templates/TypeHash.h"
-#include "Misc/TransactionObjectEvent.h"
-#include "Containers/SortedMap.h"
+UE_DEPRECATED_HEADER(5.8, "Use Misc/TransactionCommon.h instead of TransactionCommon.h")
 
-class FReferenceCollector;
-
-namespace UE::Transaction
-{
-
-/**
-	This type is necessary because the blueprint system is destroying and creating
-	CDOs at edit time (usually on compile, but also on load), but also stores user
-	entered data in the CDO. We "need"  changes to a CDO to persist across instances
-	because as we undo and redo we  need to apply changes to different instances of
-	the CDO - alternatively we could destroy and create the CDO as part of a transaction
-	(this alternative is the reason for the bunny ears around need).
-
-	DanO: My long term preference is for the editor to use a dynamic, mutable type
-	(rather than the CDO) to store editor data. The CDO can then be re-instanced (or not)
-	as runtime code requires.
-*/
-struct FPersistentObjectRef
-{
-public:
-	FPersistentObjectRef() = default;
-	ENGINE_API explicit FPersistentObjectRef(UObject* InObject);
-
-	bool Serialize(FArchive& Ar)
-	{
-		Ar << *this;
-		return true;
-	}
-
-	friend FArchive& operator<<(FArchive& Ar, FPersistentObjectRef& ReferencedObject)
-	{
-		Ar << (std::underlying_type_t<EReferenceType>&)ReferencedObject.ReferenceType;
-		Ar << ReferencedObject.RootObject;
-		Ar << ReferencedObject.SubObjectHierarchyIDs;
-		return Ar;
-	}
-
-	friend bool operator==(const FPersistentObjectRef& LHS, const FPersistentObjectRef& RHS)
-	{
-		return LHS.ReferenceType == RHS.ReferenceType
-			&& LHS.RootObject == RHS.RootObject
-			&& (LHS.ReferenceType != EReferenceType::SubObject || LHS.SubObjectHierarchyIDs == RHS.SubObjectHierarchyIDs);
-	}
-
-	friend bool operator!=(const FPersistentObjectRef& LHS, const FPersistentObjectRef& RHS)
-	{
-		return !(LHS == RHS);
-	}
-
-	friend uint32 GetTypeHash(const FPersistentObjectRef& InObjRef)
-	{
-		return HashCombine(::GetTypeHash(InObjRef.ReferenceType), GetTypeHash(InObjRef.RootObject));
-	}
-
-	bool IsRootObjectReference() const
-	{
-		return ReferenceType == EReferenceType::RootObject;
-	}
-
-	bool IsSubObjectReference() const
-	{
-		return ReferenceType == EReferenceType::SubObject;
-	}
-
-	ENGINE_API UObject* Get() const;
-	
-	ENGINE_API void AddReferencedObjects(FReferenceCollector& Collector);
-
-private:
-	/** This enum represents all of the different special cases we are handling with this type */
-	enum class EReferenceType : uint8
-	{
-		Unknown,
-		RootObject,
-		SubObject,
-	};
-
-	/** The reference type we're handling */
-	EReferenceType ReferenceType = EReferenceType::Unknown;
-	/** Stores the object pointer when ReferenceType==RootObject, and the outermost pointer of the sub-object chain when when ReferenceType==SubObject */
-	TObjectPtr<UObject> RootObject = nullptr;
-	/** Stores the sub-object name chain when ReferenceType==SubObject */
-	TArray<FName, TInlineAllocator<4>> SubObjectHierarchyIDs;
-
-	/** Cached pointers corresponding to RootObject when ReferenceType==SubObject (@note cache needs testing on access as it may have become stale) */
-	mutable TWeakObjectPtr<UObject> CachedRootObject;
-	/** Cache of pointers corresponding to the items within SubObjectHierarchyIDs when ReferenceType==SubObject (@note cache needs testing on access as it may have become stale) */
-	mutable TArray<TWeakObjectPtr<UObject>, TInlineAllocator<4>> CachedSubObjectHierarchy;
-};
-
-struct FSerializedTaggedData
-{
-public:
-	static ENGINE_API FSerializedTaggedData FromOffsetAndSize(const int64 InOffset, const int64 InSize);
-	static ENGINE_API FSerializedTaggedData FromStartAndEnd(const int64 InStart, const int64 InEnd);
-
-	ENGINE_API void AppendSerializedData(const int64 InOffset, const int64 InSize);
-	ENGINE_API void AppendSerializedData(const FSerializedTaggedData& InData);
-
-	ENGINE_API bool HasSerializedData() const;
-
-	int64 GetStart() const
-	{
-		return DataOffset;
-	}
-
-	int64 GetEnd() const
-	{
-		return DataOffset + DataSize;
-	}
-
-	friend bool operator==(const FSerializedTaggedData& LHS, const FSerializedTaggedData& RHS)
-	{
-		return LHS.DataOffset == RHS.DataOffset
-			&& LHS.DataSize == RHS.DataSize;
-	}
-
-	friend bool operator!=(const FSerializedTaggedData& LHS, const FSerializedTaggedData& RHS)
-	{
-		return !(LHS == RHS);
-	}
-
-	/** Offset to the start of the tagged data within the serialized object */
-	int64 DataOffset = INDEX_NONE;
-
-	/** Size (in bytes) of the tagged data within the serialized object */
-	int64 DataSize = 0;
-};
-
-struct FSerializedObjectData
-{
-public:
-	bool Serialize(FArchive& Ar)
-	{
-		Ar << *this;
-		return true;
-	}
-
-	friend FArchive& operator<<(FArchive& Ar, FSerializedObjectData& SerializedData)
-	{
-		Ar << SerializedData.Data;
-		return Ar;
-	}
-
-	friend bool operator==(const FSerializedObjectData& LHS, const FSerializedObjectData& RHS)
-	{
-		return LHS.Data == RHS.Data;
-	}
-
-	friend bool operator!=(const FSerializedObjectData& LHS, const FSerializedObjectData& RHS)
-	{
-		return !(LHS == RHS);
-	}
-
-	ENGINE_API void Read(void* Dest, int64 Offset, int64 Num) const;
-	ENGINE_API void Write(const void* Src, int64 Offset, int64 Num);
-	
-	const void* GetPtr(int64 Offset) const
-	{
-		return &Data[Offset];
-	}
-
-	void Reset()
-	{
-		Data.Reset();
-	}
-
-	int64 Num() const
-	{
-		return Data.Num();
-	}
-
-private:
-	TArray64<uint8> Data;
-};
-
-struct FSerializedObjectInfo : public FTransactionObjectId
-{
-public:
-	FSerializedObjectInfo() = default;
-
-	explicit FSerializedObjectInfo(const UObject* InObject)
-	{
-		SetObject(InObject);
-	}
-
-	void SetObject(const UObject* InObject)
-	{
-		FTransactionObjectId::SetObject(InObject);
-		bIsPendingKill = !IsValid(InObject);
-	}
-
-	void Reset()
-	{
-		FTransactionObjectId::Reset();
-		bIsPendingKill = false;
-	}
-
-	void Swap(FSerializedObjectInfo& Other)
-	{
-		FTransactionObjectId::Swap(Other);
-		Exchange(bIsPendingKill, Other.bIsPendingKill);
-	}
-
-	/** The pending kill state of the object when it was serialized */
-	bool bIsPendingKill = false;
-};
-
-struct FSerializedObject
-{
-public:
-	void Reset()
-	{
-		SerializedData.Reset();
-		ReferencedObjects.Reset();
-		ReferencedNames.Reset();
-	}
-
-	void Swap(FSerializedObject& Other)
-	{
-		Exchange(SerializedData, Other.SerializedData);
-		Exchange(ReferencedObjects, Other.ReferencedObjects);
-		Exchange(ReferencedNames, Other.ReferencedNames);
-	}
-
-	/** The serialized data for the transacted object */
-	FSerializedObjectData SerializedData;
-
-	/** External objects referenced by the transacted object */
-	TArray<FPersistentObjectRef> ReferencedObjects;
-	
-	/** Names referenced by the transacted object */
-	TArray<FName> ReferencedNames;
-};
-
-struct FDiffableObject
-{
-public:
-	void SetObject(const UObject* InObject)
-	{
-		ObjectInfo.SetObject(InObject);
-		ObjectArchetype = FPersistentObjectRef(InObject->GetArchetype());
-		ObjectFlags = InObject->GetFlags();
-	}
-
-	void Reset()
-	{
-		ObjectInfo.Reset();
-		ObjectArchetype = FPersistentObjectRef();
-		ObjectFlags = RF_NoFlags;
-		SerializedData.Reset();
-		SerializedTaggedData.Reset();
-	}
-
-	void Swap(FDiffableObject& Other)
-	{
-		ObjectInfo.Swap(Other.ObjectInfo);
-		Exchange(ObjectArchetype, Other.ObjectArchetype);
-		Exchange(ObjectFlags, Other.ObjectFlags);
-		Exchange(SerializedData, Other.SerializedData);
-		Exchange(SerializedTaggedData, Other.SerializedTaggedData);
-	}
-
-	/** Information about the object when it was serialized */
-	FSerializedObjectInfo ObjectInfo;
-
-	/** The archetype of the object when it was serialized */
-	FPersistentObjectRef ObjectArchetype;
-
-	/** The flags of the object when it was serialized */
-	EObjectFlags ObjectFlags = RF_NoFlags;
-
-	/** The serialized data for the diffable object */
-	FSerializedObjectData SerializedData;
-
-	/** Information about tagged data (mainly properties) that were serialized within this object */
-	TSortedMap<FName, FSerializedTaggedData, FDefaultAllocator, FNameFastLess> SerializedTaggedData;
-};
-
-ENGINE_API extern const FName TaggedDataKey_UnknownData;
-ENGINE_API extern const FName TaggedDataKey_ScriptData;
-
-/** Core archive to read a transaction object from the buffer. */
-class FSerializedObjectDataReader : public FArchiveUObject
-{
-public:
-	ENGINE_API FSerializedObjectDataReader(const FSerializedObject& InSerializedObject);
-
-	virtual int64 Tell() override { return Offset; }
-	virtual void Seek(int64 InPos) override { Offset = InPos; }
-	virtual int64 TotalSize() override { return SerializedObject.SerializedData.Num(); }
-
-protected:
-	ENGINE_API virtual void Serialize(void* SerData, int64 Num) override;
-
-	using FArchiveUObject::operator<<;
-	ENGINE_API virtual FArchive& operator<<(class FName& N) override;
-	ENGINE_API virtual FArchive& operator<<(class UObject*& Res) override;
-
-	const FSerializedObject& SerializedObject;
-	int64 Offset = 0;
-};
-
-namespace Internal
-{
-
-/** Core archive to write a transaction object to the buffer. */
-class FSerializedObjectDataWriterCommon : public FArchiveUObject
-{
-public:
-	ENGINE_API FSerializedObjectDataWriterCommon(FSerializedObjectData& InSerializedData);
-
-	virtual int64 Tell() override { return Offset; }
-	virtual void Seek(int64 InPos) override { checkSlow(Offset <= SerializedData.Num()); Offset = InPos; }
-	virtual int64 TotalSize() override { return SerializedData.Num(); }
-
-protected:
-	ENGINE_API virtual void Serialize(void* SerData, int64 Num) override;
-
-	virtual void OnDataSerialized(int64 InOffset, int64 InNum) {}
-
-	FSerializedObjectData& SerializedData;
-	int64 Offset = 0;
-};
-
-} // namespace Internal
-
-/** Core archive to write a transaction object to the buffer. */
-class FSerializedObjectDataWriter : public Internal::FSerializedObjectDataWriterCommon
-{
-public:
-	ENGINE_API FSerializedObjectDataWriter(FSerializedObject& InSerializedObject);
-
-protected:
-	using FArchiveUObject::operator<<;
-	ENGINE_API virtual FArchive& operator<<(class FName& N) override;
-	ENGINE_API virtual FArchive& operator<<(class UObject*& Res) override;
-
-	FSerializedObject& SerializedObject;
-	int64 Offset = 0;
-
-	TMap<UObject*, int32> ObjectMap;
-	TMap<FName, int32> NameMap;
-};
-
-/** Core archive to write a diffable object to the buffer. */
-class FDiffableObjectDataWriter : public Internal::FSerializedObjectDataWriterCommon
-{
-public:
-	ENGINE_API FDiffableObjectDataWriter(FDiffableObject& InDiffableObject, TArrayView<const FProperty*> InPropertiesToSerialize = TArrayView<const FProperty*>());
-
-protected:
-	ENGINE_API FName GetTaggedDataKey() const;
-
-	ENGINE_API bool DoesObjectMatchDiffableObject(const UObject* Obj) const;
-
-	ENGINE_API virtual bool ShouldSkipProperty(const FProperty* InProperty) const override;
-
-	ENGINE_API virtual void MarkScriptSerializationStart(const UObject* Obj) override;
-	ENGINE_API virtual void MarkScriptSerializationEnd(const UObject* Obj) override;
-
-	ENGINE_API virtual void OnDataSerialized(int64 InOffset, int64 InNum) override;
-
-	using FArchiveUObject::operator<<;
-	ENGINE_API virtual FArchive& operator<<(class FName& N) override;
-	ENGINE_API virtual FArchive& operator<<(class UObject*& Res) override;
-
-private:
-	struct FCachedPropertyKey
-	{
-	public:
-		FName SyncCache(const FArchiveSerializedPropertyChain& InPropertyChain);
-
-	private:
-		FName CachedKey;
-		uint32 LastUpdateCount = 0;
-	};
-
-	struct FCachedTaggedDataEntry
-	{
-	public:
-		FSerializedTaggedData& SyncCache(FDiffableObject& InDiffableObject, const FName InSerializedTaggedDataKey);
-
-	private:
-		FName CachedKey;
-		FSerializedTaggedData* CachedEntryPtr = nullptr;
-	};
-
-	FDiffableObject& DiffableObject;
-	TArrayView<const FProperty*> PropertiesToSerialize;
-
-	bool bIsPerformingScriptSerialization = false;
-
-	mutable bool bWasUsingTaggedDataKey_UnknownData = false;
-	mutable bool bWasUsingTaggedDataKey_ScriptData = false;
-	mutable int32 TaggedDataKeyIndex_UnknownData = 0;
-	mutable int32 TaggedDataKeyIndex_ScriptData = 0;
-	mutable FCachedPropertyKey CachedSerializedTaggedPropertyKey;
-	mutable FCachedTaggedDataEntry CachedSerializedTaggedDataEntry;
-};
-
-namespace DiffUtil
-{
-ENGINE_API bool IsGeneratingDiffableObject(const FArchive& Ar);
-
-enum class EGetDiffableObjectMode : uint8
-{
-	/** Serialize the entire object state by calling its Serialize function */
-	SerializeObject,
-
-	/** Serialize the property state of the object by calling its SerializeScriptProperties function */
-	SerializeProperties,
-
-	/** Serialize the object via a custom serialize function */
-	Custom,
-};
-
-struct FGetDiffableObjectOptions
-{
-	/** How should we serialize the object for diffing? */
-	EGetDiffableObjectMode ObjectSerializationMode = EGetDiffableObjectMode::SerializeObject;
-
-	/** Custom serializer for the object (must be set when ObjectSerializationMode == Custom) */
-	TFunction<void(FDiffableObjectDataWriter&)> CustomSerializer;
-
-	/** Optional list of properties to serialize on the object, or an empty array to serialize all properties */
-	TArrayView<const FProperty*> PropertiesToSerialize;
-
-	/** Should we still serialize this object if it's considered pending kill? */
-	bool bSerializeEvenIfPendingKill = false;
-};
-
-/**
- * Get an object snapshot that can be diffed later.
- */
-ENGINE_API FDiffableObject GetDiffableObject(const UObject* Object, const FGetDiffableObjectOptions& Options = FGetDiffableObjectOptions());
-
-struct FGenerateObjectDiffOptions
-{
-	FGenerateObjectDiffOptions()
-	{
-		ArchetypeOptions.bSerializeEvenIfPendingKill = true;
-	}
-
-	/** Options used when getting the diffable state of archive objects */
-	FGetDiffableObjectOptions ArchetypeOptions;
-
-	/** Optional function used to skip comparing certain properties within the diffable data */
-	TFunction<bool(FName)> ShouldSkipProperty;
-
-	/** Should we perform a "full" diff? (compares object info, and considers any missing properties to have been changed) */
-	bool bFullDiff = true;
-
-	/**
-	 * Should we diff the object state even if it's considered pending kill?
-	 * When false:
-	 *   If the old object was pending kill and the new object isn't, then the diff will be performed against the archetype of the new object.
-	 *   If the old object wasn't pending kill and the new object is, then the data diff will be skipped and only the object info diffed.
-	 */
-	bool bDiffDataEvenIfPendingKill = false;
-};
-
-/** Optional cache for requests to get the diffable state for a given archetype object */
-class FDiffableObjectArchetypeCache
-{
-public:
-	const FDiffableObject& GetArchetypeDiffableObject(const UObject* Archetype, const FGetDiffableObjectOptions& Options);
-
-private:
-	TMap<const UObject*, FDiffableObject> ArchetypeDiffableObjects;
-};
-
-/**
- * Generate a diff between the two object snapshots.
- */
-ENGINE_API FTransactionObjectDeltaChange GenerateObjectDiff(const FDiffableObject& OldDiffableObject, const FDiffableObject& NewDiffableObject, const FGenerateObjectDiffOptions& DiffOptions = FGenerateObjectDiffOptions(), FDiffableObjectArchetypeCache* ArchetypeCache = nullptr);
-ENGINE_API void GenerateObjectDiff(const FDiffableObject& OldDiffableObject, const FDiffableObject& NewDiffableObject, FTransactionObjectDeltaChange& OutDeltaChange, const FGenerateObjectDiffOptions& DiffOptions = FGenerateObjectDiffOptions(), FDiffableObjectArchetypeCache* ArchetypeCache = nullptr);
-
-} // namespace DiffUtil
-
-} // namespace UE::Transaction
+#include "Misc/TransactionCommon.h"
diff --git a/Engine/Source/Runtime/GeometryFramework/Private/UDynamicMesh.cpp b/Engine/Source/Runtime/GeometryFramework/Private/UDynamicMesh.cpp
--- a/Engine/Source/Runtime/GeometryFramework/Private/UDynamicMesh.cpp
+++ b/Engine/Source/Runtime/GeometryFramework/Private/UDynamicMesh.cpp
@@ -12,7 +12,7 @@
 #include "HAL/IConsoleManager.h"
 #include "Engine/Engine.h"
 #include "Generators/MinimalBoxMeshGenerator.h"
-#include "TransactionCommon.h"
+#include "Misc/TransactionCommon.h"
 
 
 #if WITH_EDITOR
diff --git a/Engine/Source/Runtime/NavigationSystem/Private/NavMesh/PImplRecastNavMesh.cpp b/Engine/Source/Runtime/NavigationSystem/Private/NavMesh/PImplRecastNavMesh.cpp
--- a/Engine/Source/Runtime/NavigationSystem/Private/NavMesh/PImplRecastNavMesh.cpp
+++ b/Engine/Source/Runtime/NavigationSystem/Private/NavMesh/PImplRecastNavMesh.cpp
@@ -2,7 +2,7 @@
 
 #include "NavMesh/PImplRecastNavMesh.h"
 #include "NavigationSystem.h"
-#include "TransactionCommon.h"
+#include "Misc/TransactionCommon.h"
 
 #if WITH_RECAST
 
