FRemoteWorkPriority fixes:
- ensure that dependent priorities are always unique, we now encode a unique id into the low bits that includes the creator's serverid plus a unique rolling id

#Synced-CL 52074282

diff --git a/Engine/Source/Runtime/CoreUObject/Private/UObject/RemoteExecutor.cpp b/Engine/Source/Runtime/CoreUObject/Private/UObject/RemoteExecutor.cpp
--- a/Engine/Source/Runtime/CoreUObject/Private/UObject/RemoteExecutor.cpp
+++ b/Engine/Source/Runtime/CoreUObject/Private/UObject/RemoteExecutor.cpp
@@ -112,6 +112,7 @@ class FRemoteExecutor
 	TArray<FRemoteSubsystemBase*> Subsystems;
 
 	uint32 NextTransactionRequestId = 0;
+	uint32 NextUniqueId = 0;
 
 	TArray<FRemoteExecutorWork> PendingWorks;
 	TArray<FRemoteExecutorWork> PendingMigrationWorks;
@@ -208,6 +209,16 @@ public:
 		return FRemoteTransactionId((RawRequestId++ % 0x800000u) + 1000u);
 	}
 
+	uint32 GenerateNextUniqueId()
+	{
+		// pack the local server id into the top bits
+		uint32 UniqueIdMask = (1u << (32 - REMOTE_OBJECT_SERVER_ID_BIT_SIZE)) - 1;
+		
+		uint32 UniqueId = (NextUniqueId++) & UniqueIdMask;
+		UniqueId |= LocalServerId.GetIdNumber() << (32 - REMOTE_OBJECT_SERVER_ID_BIT_SIZE);
+		return UniqueId;
+	}
+
 	void EnqueueWork(FName InWorkName, FRemoteWorkPriority InWorkPriority, bool bIsTransactional, int32 InWorkListIndex, TUniqueFunction<void(void)>&& InWork)
 	{
 		FRemoteExecutorWork& NewWork = InWorkPriority.GetWorkType() == ERemoteWorkType::Normal ? PendingWorks.Emplace_GetRef() : PendingMigrationWorks.Emplace_GetRef();
@@ -1228,7 +1239,7 @@ void ReadyRemoteMultiServerCommit(FRemoteServerId ServerId, FRemoteTransactionId
 
 FRemoteWorkPriority CreateRootWorkPriority()
 {
-	FRemoteWorkPriority RootWorkPriority = FRemoteWorkPriority::CreateRootWorkPriority(GRemoteExecutor->GetLocalServerId(), GRemoteExecutor->GenerateNextTransactionId());
+	FRemoteWorkPriority RootWorkPriority = FRemoteWorkPriority::CreateRootWorkPriority(GRemoteExecutor->GetLocalServerId());
 	return RootWorkPriority;
 }
 
@@ -1258,7 +1269,7 @@ void ExecuteTransactional(FName WorkName, const TFunctionRef<void(void)>& Work)
 {
 	if (GRemoteExecutor)
 	{
-		FRemoteWorkPriority RootWorkPriority = FRemoteWorkPriority::CreateRootWorkPriority(GRemoteExecutor->GetLocalServerId(), GRemoteExecutor->GenerateNextTransactionId());
+		FRemoteWorkPriority RootWorkPriority = FRemoteWorkPriority::CreateRootWorkPriority(GRemoteExecutor->GetLocalServerId());
 		ExecuteTransactionalWithExplicitPriority(WorkName, RootWorkPriority, Work);
 	}
 	else
@@ -1287,7 +1298,7 @@ void EnqueueWork(FName WorkName, bool bIsTransactional, TUniqueFunction<void(voi
 {
 	if (GRemoteExecutor)
 	{
-		FRemoteWorkPriority RootWorkPriority = FRemoteWorkPriority::CreateRootWorkPriority(GRemoteExecutor->GetLocalServerId(), GRemoteExecutor->GenerateNextTransactionId());
+		FRemoteWorkPriority RootWorkPriority = FRemoteWorkPriority::CreateRootWorkPriority(GRemoteExecutor->GetLocalServerId());
 		EnqueueWorkWithExplicitPriority(WorkName, RootWorkPriority, bIsTransactional, MoveTemp(InWork));
 	}
 	else
@@ -1300,7 +1311,7 @@ void EnqueueMigrationWork(FName WorkName, TUniqueFunction<void(void)>&& InWork)
 {
 	if (GRemoteExecutor)
 	{
-		FRemoteWorkPriority RootWorkPriority = FRemoteWorkPriority::CreateRootWorkPriority(GRemoteExecutor->GetLocalServerId(), GRemoteExecutor->GenerateNextTransactionId(), ERemoteWorkType::Migration);
+		FRemoteWorkPriority RootWorkPriority = FRemoteWorkPriority::CreateRootWorkPriority(GRemoteExecutor->GetLocalServerId(), ERemoteWorkType::Migration);
 		EnqueueWorkWithExplicitPriority(WorkName, RootWorkPriority, /*bIsTransactional=*/ true, MoveTemp(InWork));
 	}
 	else
@@ -1335,7 +1346,7 @@ void EnqueueWorkList(FName WorkName, bool bIsTransactional, TArray<TUniqueFuncti
 #if UE_WITH_REMOTE_OBJECT_HANDLE
 	if (GRemoteExecutor)
 	{
-		FRemoteWorkPriority RootWorkPriority = FRemoteWorkPriority::CreateRootWorkPriority(GRemoteExecutor->GetLocalServerId(), GRemoteExecutor->GenerateNextTransactionId());
+		FRemoteWorkPriority RootWorkPriority = FRemoteWorkPriority::CreateRootWorkPriority(GRemoteExecutor->GetLocalServerId());
 		EnqueueWorkListWithExplicitPriority(WorkName, RootWorkPriority, bIsTransactional, MoveTemp(InWorkList));
 	}
 	else
@@ -1370,47 +1381,47 @@ const TCHAR* EnumToString(ERemoteWorkType WorkType)
 
 FString FRemoteWorkPriority::ToString() const
 {
-	return FString::Format(TEXT("[pri: rsi {0} depth {1} id {2} type {3}]"),
-		{ *GetRootServerId().ToString(),
+	return FString::Format(TEXT("[pri: type {0} rsi {1} depth {2} id {3} ]"),
+		{ EnumToString(GetWorkType()),
+		*GetRootServerId().ToString(),
 		GetWorkDepth(),
-		*GetRootWorkTransactionId().ToString(),
-		EnumToString(GetWorkType()) });
+		GetUniqueId() });
 }
 
-FRemoteWorkPriority FRemoteWorkPriority::CreateRootWorkPriority(FRemoteServerId ServerId, FRemoteTransactionId TransactionId, ERemoteWorkType WorkType /*= ERemoteWorkType::Normal*/)
+FRemoteWorkPriority FRemoteWorkPriority::CreateRootWorkPriority(FRemoteServerId ServerId, ERemoteWorkType WorkType /*= ERemoteWorkType::Normal*/)
 {
-	FRemoteWorkPriority Result(WorkType);
-
-	uint32 RawServerId = ServerId.GetIdNumber();
-
-	uint32 RawTransactionId = 0;
-	FMemory::Memcpy(&RawTransactionId, &TransactionId, sizeof(FRemoteTransactionId));
+	FRemoteWorkPriority Result;
 
-	constexpr uint64 RootWorkDepth = 0xFFull;
+	uint64 RawWorkType = (uint64)WorkType;
+	uint64 RawServerId = ServerId.GetIdNumber();
+	uint64 RawWorkDepth = 0xFFull;
+	uint64 RawUniqueId = GRemoteExecutor->GenerateNextUniqueId();
 
-	Result.PackedData |= (static_cast<uint64>(RawServerId & 0x7FFFFFFFu) << 32);
-	Result.PackedData |= (RootWorkDepth << 24);
-	Result.PackedData |= (RawTransactionId & 0xFFFFFFu);
+	Result.PackedData = PackPriority(RawWorkType, RawServerId, RawWorkDepth, RawUniqueId);
 
 	checkf(WorkType == Result.GetWorkType(), TEXT("Work type (%u) was not properly encoded (%d)"), (uint32)WorkType, (uint32)Result.GetWorkType());
 	checkf(ServerId == Result.GetRootServerId(), TEXT("Server id (%s) was not properly encoded (%s)"), *ServerId.ToString(), *Result.GetRootServerId().ToString());
-	checkf(RootWorkDepth == Result.GetWorkDepth(), TEXT("Work depth (%llu) was not properly encoded (%u)"), RootWorkDepth, Result.GetWorkDepth());
-	checkf(TransactionId == Result.GetRootWorkTransactionId(), TEXT("Transaction id (%u) was not properly encoded (%u)"), TransactionId.GetIdNumber(), Result.GetRootWorkTransactionId().GetIdNumber());
+	checkf(RawWorkDepth == Result.GetWorkDepth(), TEXT("Work depth (%llu) was not properly encoded (%u)"), RawWorkDepth, Result.GetWorkDepth());
+	checkf(RawUniqueId == Result.GetUniqueId(), TEXT("Unique id (%llu) was not properly encoded (%llu)"), RawUniqueId, Result.GetUniqueId());
 
 	return Result;
 }
 
 FRemoteWorkPriority FRemoteWorkPriority::CreateDependentWorkPriority() const
 {
-	FRemoteWorkPriority Result = *this;
-	Result.PackedData &= 0xFFFFFFFF00FFFFFFull;
+	FRemoteWorkPriority Result;
 
-	// work depth must be non-zero else we would wrap back around
-	check(GetWorkDepth() > 0);
+	uint64 RawWorkType = (uint64)GetWorkType();
+	uint64 RawServerId = GetRootServerId().GetIdNumber();
+	uint64 RawWorkDepth = GetWorkDepth() - 1;
+	uint64 RawUniqueId = GRemoteExecutor->GenerateNextUniqueId();
 
-	Result.PackedData |= static_cast<uint64>(GetWorkDepth() - 1) << 24;
+	Result.PackedData = PackPriority(RawWorkType, RawServerId, RawWorkDepth, RawUniqueId);
 
-	checkf(GetWorkType() == Result.GetWorkType(), TEXT("Work type (%u) changed when creating dependent work priority (%d)"), (uint32)GetWorkType(), (uint32)Result.GetWorkType());
+	checkf(GetWorkType() == Result.GetWorkType(), TEXT("Work type (%u) was not properly encoded (%d)"), (uint32)GetWorkType(), (uint32)Result.GetWorkType());
+	checkf(GetRootServerId() == Result.GetRootServerId(), TEXT("Server id (%s) was not properly encoded (%s)"), *GetRootServerId().ToString(), *Result.GetRootServerId().ToString());
+	checkf(RawWorkDepth == Result.GetWorkDepth(), TEXT("Work depth (%llu) was not properly encoded (%u)"), RawWorkDepth, Result.GetWorkDepth());
+	checkf(RawUniqueId == Result.GetUniqueId(), TEXT("Unique id (%llu) was not properly encoded (%llu)"), RawUniqueId, Result.GetUniqueId());
 
 	return Result;
 }
diff --git a/Engine/Source/Runtime/CoreUObject/Public/UObject/RemoteExecutor.h b/Engine/Source/Runtime/CoreUObject/Public/UObject/RemoteExecutor.h
--- a/Engine/Source/Runtime/CoreUObject/Public/UObject/RemoteExecutor.h
+++ b/Engine/Source/Runtime/CoreUObject/Public/UObject/RemoteExecutor.h
@@ -59,44 +59,97 @@ COREUOBJECT_API const TCHAR* EnumToString(ERemoteWorkType WorkType);
 
 struct FRemoteWorkPriority
 {
+	uint64 PackedData = 0;
+
 	// the data is packed in a uint64 to allow passing by value in a register
-	// [63]    : 0 if this priority represents transactional migration work, 1 otherwise
-	// [62:32] : FRemoteServerId RootServerId
-	// [31:24] : uint8 WorkDepth
-	// [23:0]  : FRemoteTransactionId RootWorkTransactionId
-	//
 	// this allows comparison of priority to be a single uint64 compare-less-than
 	// 
 	// the WorkDepth is stored subtracted from 0xFF so that less-than means higher priority
 	//   0xFF is the root depth
 	//   0xFE is one deeper
-	//   
-	uint64 PackedData = 0;
-	static_assert(sizeof(FRemoteServerId) == 4);
-	static_assert(sizeof(FRemoteTransactionId) == 4);
-
-	FRemoteServerId GetRootServerId() const { return FRemoteServerId::FromIdNumber(static_cast<uint32>((PackedData >> 32) & 0x7FFFFFFFu)); }
-	uint32 GetWorkDepth() const { return (PackedData >> 24) & 0xFFu; }
-	FRemoteTransactionId GetRootWorkTransactionId() const { return FRemoteTransactionId(PackedData & 0xFFFFFFu); }
-	ERemoteWorkType GetWorkType() const { return !(PackedData >> 63) ? ERemoteWorkType::Migration : ERemoteWorkType::Normal; }
+	//
+	using FPriorityPacking = TIntegerSequence<int32,
+		1, // 0 if this priority represents (higher priority) transactional migration work, 1 otherwise
+		REMOTE_OBJECT_SERVER_ID_BIT_SIZE, // FRemoteServerId RootServerId
+		8, // uint8 WorkDepth
+		32>; // unique id bits to ensure this priority is unique
+
+	ERemoteWorkType GetWorkType() const { return !(UnpackPriorityElement(PackedData, 0)) ? ERemoteWorkType::Migration : ERemoteWorkType::Normal; }
+	FRemoteServerId GetRootServerId() const { return FRemoteServerId::FromIdNumber((uint32)UnpackPriorityElement(PackedData, 1)); }
+	uint32 GetWorkDepth() const { return (uint32)UnpackPriorityElement(PackedData, 2); }
+	uint32 GetUniqueId() const { return (uint32)UnpackPriorityElement(PackedData, 3); }
 	bool IsValid() const { return GetRootServerId().IsValid(); }
 
 	FRemoteWorkPriority()
-		: FRemoteWorkPriority(ERemoteWorkType::Normal)
 	{
+		PackedData = PackPriority(ERemoteWorkType::Normal, 0, 0, 0);
 	}
 	
-	explicit FRemoteWorkPriority(ERemoteWorkType WorkType)
-		: PackedData(static_cast<uint64>(WorkType) << 63)
-	{
-	}
-
-	COREUOBJECT_API static FRemoteWorkPriority CreateRootWorkPriority(FRemoteServerId ServerId, FRemoteTransactionId TransactionId, ERemoteWorkType WorkType = ERemoteWorkType::Normal);
+	COREUOBJECT_API static FRemoteWorkPriority CreateRootWorkPriority(FRemoteServerId ServerId, ERemoteWorkType WorkType = ERemoteWorkType::Normal);
 
 	// creates a copy of this priority with the WorkDepth decreased by 1
 	COREUOBJECT_API FRemoteWorkPriority CreateDependentWorkPriority() const;
 
 	COREUOBJECT_API FString ToString() const;
+
+private:
+	template<int32... Bits, typename... ArgTypes>
+	static uint64 PackPriorityImpl(TIntegerSequence<int32, Bits...>, ArgTypes... InValues)
+	{
+		static_assert(sizeof...(Bits) == sizeof...(ArgTypes));
+		static_assert((Bits + ...) <= 64, "sum of the bit sizes must be <= 64");
+
+		int32 BitWidths[] = { Bits... };
+		uint64 Values[] = { static_cast<uint64>(InValues)... };
+
+		uint64 Result = 0;
+		int32 Offset = 64; // pack the values in decreasing significance
+
+		for(int32 Index = 0; Index < sizeof...(Bits); Index++)
+		{
+			Offset -= BitWidths[Index];
+			uint64 Mask = (1ULL << BitWidths[Index]) - 1;
+			Result |= (Values[Index] & Mask) << Offset;
+		}
+
+		return Result;
+	}
+
+	template<int32... Bits>
+	static uint64 UnpackPriorityElementImpl(TIntegerSequence<int32, Bits...>, uint64 PackedPriority, int32 OutIndex)
+	{
+		static_assert((Bits + ...) <= 64, "sum of the bit sizes must be <= 64");
+
+		int32 BitWidths[] = { Bits... };
+
+		int32 Offset = 64;
+		uint64 Result = 0;
+		for(int32 Index = 0; Index < sizeof...(Bits); Index++)
+		{
+			Offset -= BitWidths[Index];
+
+			uint64 Mask = (1ULL << BitWidths[Index]) - 1;
+
+			if (Index == OutIndex)
+			{
+				Result = (PackedPriority >> Offset) & Mask;
+				break;
+			}
+		}
+
+		return Result;
+	}
+
+	template<typename... ArgTypes>
+	static uint64 PackPriority(ArgTypes... Values)
+	{
+		return PackPriorityImpl(FPriorityPacking{}, Values...);
+	}
+
+	static uint64 UnpackPriorityElement(uint64 PackedPriority, int32 OutIndex)
+	{
+		return UnpackPriorityElementImpl(FPriorityPacking{}, PackedPriority, OutIndex);
+	}
 };
 
 COREUOBJECT_API FArchive& operator<<(FArchive& Ar, FRemoteWorkPriority& Priority);
