RemoteHeap: implemented read sharing

#Synced-CL 52114874
#Synced-CL 52114874

diff --git a/Engine/Source/Runtime/AutoRTFM/Private/API.cpp b/Engine/Source/Runtime/AutoRTFM/Private/API.cpp
--- a/Engine/Source/Runtime/AutoRTFM/Private/API.cpp
+++ b/Engine/Source/Runtime/AutoRTFM/Private/API.cpp
@@ -1017,43 +1017,43 @@ void ForTheRuntime::RegisterOnAbortFromTheOpen(TTask<void()>&& Work)
 	RTFM_OnAbortInternal(std::move(Work));
 }
 
-void ForTheRuntime::RedirectedLoad(uint32_t AddressSpace, void* DestPointer, uint64_t Size, uint64_t SourceAddress)
+void ForTheRuntime::RedirectedLoad(uint32_t AddressSpace, void* DestPointer, uint64_t Size, uint64_t SourceAddress, bool bWillWriteHint)
 {
 	if (FHeapRedirectCallbacks* HeapRedirects = FindHeapRedirectCallbacksForAddressSpace(AddressSpace))
 	{
-		HeapRedirects->RedirectedLoad(DestPointer, Size, SourceAddress);
+		HeapRedirects->RedirectedLoad(DestPointer, Size, SourceAddress, bWillWriteHint);
 	}
 }
 
-void ForTheRuntime::RedirectedLoad8(uint32_t AddressSpace, void* DestPointer, uint64_t SourceAddress)
+void ForTheRuntime::RedirectedLoad8(uint32_t AddressSpace, void* DestPointer, uint64_t SourceAddress, bool bWillWriteHint)
 {
 	if (FHeapRedirectCallbacks* HeapRedirects = FindHeapRedirectCallbacksForAddressSpace(AddressSpace))
 	{
-		HeapRedirects->RedirectedLoad8(DestPointer, SourceAddress);
+		HeapRedirects->RedirectedLoad8(DestPointer, SourceAddress, bWillWriteHint);
 	}
 }
 
-void ForTheRuntime::RedirectedLoad4(uint32_t AddressSpace, void* DestPointer, uint64_t SourceAddress)
+void ForTheRuntime::RedirectedLoad4(uint32_t AddressSpace, void* DestPointer, uint64_t SourceAddress, bool bWillWriteHint)
 {
 	if (FHeapRedirectCallbacks* HeapRedirects = FindHeapRedirectCallbacksForAddressSpace(AddressSpace))
 	{
-		HeapRedirects->RedirectedLoad4(DestPointer, SourceAddress);
+		HeapRedirects->RedirectedLoad4(DestPointer, SourceAddress, bWillWriteHint);
 	}
 }
 
-void ForTheRuntime::RedirectedLoad2(uint32_t AddressSpace, void* DestPointer, uint64_t SourceAddress)
+void ForTheRuntime::RedirectedLoad2(uint32_t AddressSpace, void* DestPointer, uint64_t SourceAddress, bool bWillWriteHint)
 {
 	if (FHeapRedirectCallbacks* HeapRedirects = FindHeapRedirectCallbacksForAddressSpace(AddressSpace))
 	{
-		HeapRedirects->RedirectedLoad2(DestPointer, SourceAddress);
+		HeapRedirects->RedirectedLoad2(DestPointer, SourceAddress, bWillWriteHint);
 	}
 }
 
-void ForTheRuntime::RedirectedLoad1(uint32_t AddressSpace, void* DestPointer, uint64_t SourceAddress)
+void ForTheRuntime::RedirectedLoad1(uint32_t AddressSpace, void* DestPointer, uint64_t SourceAddress, bool bWillWriteHint)
 {
 	if (FHeapRedirectCallbacks* HeapRedirects = FindHeapRedirectCallbacksForAddressSpace(AddressSpace))
 	{
-		HeapRedirects->RedirectedLoad1(DestPointer, SourceAddress);
+		HeapRedirects->RedirectedLoad1(DestPointer, SourceAddress, bWillWriteHint);
 	}
 }
 
diff --git a/Engine/Source/Runtime/AutoRTFM/Private/CompilerInterface.cpp b/Engine/Source/Runtime/AutoRTFM/Private/CompilerInterface.cpp
--- a/Engine/Source/Runtime/AutoRTFM/Private/CompilerInterface.cpp
+++ b/Engine/Source/Runtime/AutoRTFM/Private/CompilerInterface.cpp
@@ -107,45 +107,45 @@ static bool AutoRTFMIsPointerNegative(const void* InPointer)
 }
 
 extern "C" AUTORTFM_DISABLE_NO_SANITIZE UE_AUTORTFM_API AUTORTFM_NO_ASAN void autortfm_redirected_load(
-	uint32 AddressSpace, void* DestPointer, uint64 Size, uint64 SourceAddress) noexcept
+	uint32 AddressSpace, void* DestPointer, uint64 Size, uint64 SourceAddress, bool bWillWriteHint) noexcept
 {
 	ForTheRuntime::RedirectedLoad(
-		AddressSpace ? AddressSpace : AutoRTFMExtractAddressSpaceFromAddress(SourceAddress), DestPointer, Size, SourceAddress);
+		AddressSpace ? AddressSpace : AutoRTFMExtractAddressSpaceFromAddress(SourceAddress), DestPointer, Size, SourceAddress, bWillWriteHint);
 }
 
 extern "C" AUTORTFM_DISABLE_NO_SANITIZE UE_AUTORTFM_API AUTORTFM_NO_ASAN uint64 autortfm_redirected_load_8(
-	uint32 AddressSpace, uint64 SourceAddress) noexcept
+	uint32 AddressSpace, uint64 SourceAddress, bool bWillWriteHint) noexcept
 {
 	uint64 Result = 0;
 	ForTheRuntime::RedirectedLoad8(
-		AddressSpace ? AddressSpace : AutoRTFMExtractAddressSpaceFromAddress(SourceAddress), &Result, SourceAddress);
+		AddressSpace ? AddressSpace : AutoRTFMExtractAddressSpaceFromAddress(SourceAddress), &Result, SourceAddress, bWillWriteHint);
 	return Result;
 }
 
 extern "C" AUTORTFM_DISABLE_NO_SANITIZE UE_AUTORTFM_API AUTORTFM_NO_ASAN uint32 autortfm_redirected_load_4(
-	uint32 AddressSpace, uint64 SourceAddress) noexcept
+	uint32 AddressSpace, uint64 SourceAddress, bool bWillWriteHint) noexcept
 {
 	uint32 Result = 0;
 	ForTheRuntime::RedirectedLoad4(
-		AddressSpace ? AddressSpace : AutoRTFMExtractAddressSpaceFromAddress(SourceAddress), &Result, SourceAddress);
+		AddressSpace ? AddressSpace : AutoRTFMExtractAddressSpaceFromAddress(SourceAddress), &Result, SourceAddress, bWillWriteHint);
 	return Result;
 }
 
 extern "C" AUTORTFM_DISABLE_NO_SANITIZE UE_AUTORTFM_API AUTORTFM_NO_ASAN uint16 autortfm_redirected_load_2(
-	uint32 AddressSpace, uint64 SourceAddress) noexcept
+	uint32 AddressSpace, uint64 SourceAddress, bool bWillWriteHint) noexcept
 {
 	uint16 Result = 0;
 	ForTheRuntime::RedirectedLoad2(
-		AddressSpace ? AddressSpace : AutoRTFMExtractAddressSpaceFromAddress(SourceAddress), &Result, SourceAddress);
+		AddressSpace ? AddressSpace : AutoRTFMExtractAddressSpaceFromAddress(SourceAddress), &Result, SourceAddress, bWillWriteHint);
 	return Result;
 }
 
 extern "C" AUTORTFM_DISABLE_NO_SANITIZE UE_AUTORTFM_API AUTORTFM_NO_ASAN uint8 autortfm_redirected_load_1(
-	uint32 AddressSpace, uint64 SourceAddress) noexcept
+	uint32 AddressSpace, uint64 SourceAddress, bool bWillWriteHint) noexcept
 {
 	uint8 Result = 0;
 	ForTheRuntime::RedirectedLoad1(
-		AddressSpace ? AddressSpace : AutoRTFMExtractAddressSpaceFromAddress(SourceAddress), &Result, SourceAddress);
+		AddressSpace ? AddressSpace : AutoRTFMExtractAddressSpaceFromAddress(SourceAddress), &Result, SourceAddress, bWillWriteHint);
 	return Result;
 }
 
@@ -235,7 +235,7 @@ AUTORTFM_DISABLE_NO_SANITIZE void autortfm_memmove_redirectable_impl(void* Dst,
 
 		if (bSrcIsRedirected)
 		{
-			Data = autortfm_redirected_load_8(0, (uint64)SrcCursor);
+			Data = autortfm_redirected_load_8(0, (uint64)SrcCursor, false);
 		}
 		else
 		{
@@ -274,7 +274,7 @@ AUTORTFM_DISABLE_NO_SANITIZE void autortfm_memmove_redirectable_impl(void* Dst,
 
 		if (bSrcIsRedirected)
 		{
-			Data = autortfm_redirected_load_1(0, (uint64)SrcCursor);
+			Data = autortfm_redirected_load_1(0, (uint64)SrcCursor, false);
 		}
 		else
 		{
diff --git a/Engine/Source/Runtime/AutoRTFM/Public/AutoRTFM.h b/Engine/Source/Runtime/AutoRTFM/Public/AutoRTFM.h
--- a/Engine/Source/Runtime/AutoRTFM/Public/AutoRTFM.h
+++ b/Engine/Source/Runtime/AutoRTFM/Public/AutoRTFM.h
@@ -319,11 +319,11 @@ UE_AUTORTFM_API void PopOnAbortHandlerInternal(const void* Key);
 UE_AUTORTFM_API void PopAllOnAbortHandlersInternal(const void* Key);
 UE_AUTORTFM_API void RegisterOnCommitFromTheOpen(TTask<void()>&& Work);
 UE_AUTORTFM_API void RegisterOnAbortFromTheOpen(TTask<void()>&& Work);
-UE_AUTORTFM_API void RedirectedLoad(uint32_t AddressSpace, void* DestPointer, uint64_t Size, uint64_t SourceAddress);
-UE_AUTORTFM_API void RedirectedLoad8(uint32_t AddressSpace, void* DestPointer, uint64_t SourceAddress);
-UE_AUTORTFM_API void RedirectedLoad4(uint32_t AddressSpace, void* DestPointer, uint64_t SourceAddress);
-UE_AUTORTFM_API void RedirectedLoad2(uint32_t AddressSpace, void* DestPointer, uint64_t SourceAddress);
-UE_AUTORTFM_API void RedirectedLoad1(uint32_t AddressSpace, void* DestPointer, uint64_t SourceAddress);
+UE_AUTORTFM_API void RedirectedLoad(uint32_t AddressSpace, void* DestPointer, uint64_t Size, uint64_t SourceAddress, bool bWillWriteHint);
+UE_AUTORTFM_API void RedirectedLoad8(uint32_t AddressSpace, void* DestPointer, uint64_t SourceAddress, bool bWillWriteHint);
+UE_AUTORTFM_API void RedirectedLoad4(uint32_t AddressSpace, void* DestPointer, uint64_t SourceAddress, bool bWillWriteHint);
+UE_AUTORTFM_API void RedirectedLoad2(uint32_t AddressSpace, void* DestPointer, uint64_t SourceAddress, bool bWillWriteHint);
+UE_AUTORTFM_API void RedirectedLoad1(uint32_t AddressSpace, void* DestPointer, uint64_t SourceAddress, bool bWillWriteHint);
 UE_AUTORTFM_API void RedirectedStore(uint32_t AddressSpace, uint64_t DestAddress, uint64_t Size, const void* SourcePointer);
 UE_AUTORTFM_API void RedirectedStore8(uint32_t AddressSpace, uint64_t DestAddress, const void* SourcePointer);
 UE_AUTORTFM_API void RedirectedStore4(uint32_t AddressSpace, uint64_t DestAddress, const void* SourcePointer);
@@ -953,11 +953,11 @@ UE_AUTORTFM_CRITICAL_INLINE void PopAllOnAbortHandlers(const void* Key)
 struct FHeapRedirectCallbacks
 {
 	uint32_t AddressSpace;
-	void (*RedirectedLoad)(void* DestPointer, uint64_t Size, uint64_t SourceAddress);
-	void (*RedirectedLoad8)(void* DestPointer, uint64_t SourceAddress);
-	void (*RedirectedLoad4)(void* DestPointer, uint64_t SourceAddress);
-	void (*RedirectedLoad2)(void* DestPointer, uint64_t SourceAddress);
-	void (*RedirectedLoad1)(void* DestPointer, uint64_t SourceAddress);
+	void (*RedirectedLoad)(void* DestPointer, uint64_t Size, uint64_t SourceAddress, bool bWillWriteHint);
+	void (*RedirectedLoad8)(void* DestPointer, uint64_t SourceAddress, bool bWillWriteHint);
+	void (*RedirectedLoad4)(void* DestPointer, uint64_t SourceAddress, bool bWillWriteHint);
+	void (*RedirectedLoad2)(void* DestPointer, uint64_t SourceAddress, bool bWillWriteHint);
+	void (*RedirectedLoad1)(void* DestPointer, uint64_t SourceAddress, bool bWillWriteHint);
 	void (*RedirectedStore)(uint64_t DestAddress, uint64_t Size, const void* SourcePointer);
 	void (*RedirectedStore8)(uint64_t DestAddress, const void* SourcePointer);
 	void (*RedirectedStore4)(uint64_t DestAddress, const void* SourcePointer);
diff --git a/Engine/Source/Runtime/CoreUObject/Private/UObject/RemoteExecutor.cpp b/Engine/Source/Runtime/CoreUObject/Private/UObject/RemoteExecutor.cpp
--- a/Engine/Source/Runtime/CoreUObject/Private/UObject/RemoteExecutor.cpp
+++ b/Engine/Source/Runtime/CoreUObject/Private/UObject/RemoteExecutor.cpp
@@ -756,7 +756,7 @@ private:
 				{
 					using namespace UE::RemoteExecutor;
 
-					UE_LOGF(LogRemoteExec, Verbose, "ExecutePendingWork[%ls] (%ls): Attempt %ux ABORTED. Reason: %ls (%ls)",
+					UE_LOGF(LogRemoteExec, Verbose, "ExecutePendingWork[%ls] (%ls): Attempt %u ABORTED. Reason: %ls (%ls)",
 							*ExecutingWork->RequestId.ToString(),
 							*ExecutingWork->Name.ToString(),
 							ExecutingWork->ExecutionAttempts,
diff --git a/Engine/Source/Runtime/CoreUObject/Private/UObject/RemoteHeap.cpp b/Engine/Source/Runtime/CoreUObject/Private/UObject/RemoteHeap.cpp
--- a/Engine/Source/Runtime/CoreUObject/Private/UObject/RemoteHeap.cpp
+++ b/Engine/Source/Runtime/CoreUObject/Private/UObject/RemoteHeap.cpp
@@ -12,7 +12,6 @@ DECLARE_LOG_CATEGORY_EXTERN(LogRemoteHeap, Display, All);
 DEFINE_LOG_CATEGORY(LogRemoteHeap);
 
 const uint64 RemoteHeapBaseAddress   = 0xAA00000000000000ull; // as signed, must be negative
-const uint64 RemoteHeapSizePerServer = 0x0000000400000000ull; // 16GiB
 
 const uint64 RemoteHeapReservedBitCount = 8;
 const uint64 RemoteHeapServerIdBitCount = (REMOTE_OBJECT_SERVER_ID_BIT_SIZE + 3ull) & ~3ull;
@@ -24,20 +23,49 @@ const uint64 RemoteHeapLineMask = ~(RemoteHeapLineSize-1ull);
 // Prevent us from having to match this definition to the declaration in the header just to initialize the delegate
 #define DEFINE_REMOTE_HEAP_DELEGATE(x) decltype(x) x
 	DEFINE_REMOTE_HEAP_DELEGATE(UE::RemoteHeapSendLineDelegate);
-	DEFINE_REMOTE_HEAP_DELEGATE(UE::RemoteHeapRequestLineDelegate);
+	DEFINE_REMOTE_HEAP_DELEGATE(UE::RemoteHeapRequestLineReadDelegate);
+	DEFINE_REMOTE_HEAP_DELEGATE(UE::RemoteHeapRequestLineExclusiveDelegate);
+	DEFINE_REMOTE_HEAP_DELEGATE(UE::RemoteHeapExclusiveAckDelegate);
 	DEFINE_REMOTE_HEAP_DELEGATE(UE::RemoteHeapDeallocateDelegate);
+	DEFINE_REMOTE_HEAP_DELEGATE(UE::RemoteHeapNotifyLineDeallocatedDelegate);
 #undef DEFINE_REMOTE_HEAP_DELEGATE
 
-static uint64 RemoteHeapGetBaseAddressForServerId(FRemoteServerId ServerId)
+#if WITH_DEV_AUTOMATION_TESTS
+
+static uint64 GRemoteHeapCoverage = 0;
+
+namespace UE
 {
-	uint64 ServerIdNum = ServerId.GetIdNumber();
+void FRemoteHeapCoverage::Reset()
+{
+	GRemoteHeapCoverage = 0;
+}
 
-	uint64 BaseAddress = RemoteHeapBaseAddress;
-	BaseAddress |= (ServerIdNum << (64 - RemoteHeapReservedBitCount - RemoteHeapServerIdBitCount));
+void FRemoteHeapCoverage::RecordFlow(ERemoteHeapCoverageFlow Flow)
+{
+	GRemoteHeapCoverage |= 1ull << (uint32)Flow;
+}
 
-	return BaseAddress;
+uint64 FRemoteHeapCoverage::GetCoveredFlows()
+{
+	return GRemoteHeapCoverage;
 }
 
+bool FRemoteHeapCoverage::IsFlowCovered(ERemoteHeapCoverageFlow Flow)
+{
+	return (GRemoteHeapCoverage & (1ull << (uint32)Flow)) != 0;
+}
+
+uint64 FRemoteHeapCoverage::GetFullCoverageMask()
+{
+	static_assert((uint32)ERemoteHeapCoverageFlow::Count <= 64,
+		"ERemoteHeapCoverageFlow has more than 64 entries; switch to a wider type");
+	return (1ull << (uint32)ERemoteHeapCoverageFlow::Count) - 1ull;
+}
+} // namespace UE
+
+#endif // WITH_DEV_AUTOMATION_TESTS
+
 static FRemoteServerId RemoteHeapGetOwnerServerId(uint64 Address)
 {
 	uint32 ServerIdNum = (uint32)((Address >> (64 - RemoteHeapReservedBitCount - RemoteHeapServerIdBitCount)) & RemoteHeapServerIdMask);
@@ -46,36 +74,47 @@ static FRemoteServerId RemoteHeapGetOwnerServerId(uint64 Address)
 
 namespace UE
 {
-void RemoteHeapRedirectedLoad(void* DestPointer, uint64_t Size, uint64_t SourceAddress)
+
+uint64 FRemoteHeap::GetRemoteHeapBaseAddressForServerId(FRemoteServerId ServerId)
+{
+	uint64 ServerIdNum = ServerId.GetIdNumber();
+
+	uint64 BaseAddress = RemoteHeapBaseAddress;
+	BaseAddress |= (ServerIdNum << (64 - RemoteHeapReservedBitCount - RemoteHeapServerIdBitCount));
+
+	return BaseAddress;
+}
+
+void RemoteHeapRedirectedLoad(void* DestPointer, uint64_t Size, uint64_t SourceAddress, bool bWillWriteHint)
 {
 	switch(Size)
 	{
-	case 8: *(uint64*)DestPointer = FRemoteHeap::Instance->RemoteHeapRead8(SourceAddress); break;
-	case 4: *(uint32*)DestPointer = FRemoteHeap::Instance->RemoteHeapRead4(SourceAddress); break;
-	case 2: *(uint16*)DestPointer = FRemoteHeap::Instance->RemoteHeapRead2(SourceAddress); break;
-	case 1: *(uint8*)DestPointer = FRemoteHeap::Instance->RemoteHeapRead1(SourceAddress); break;
-	default: FRemoteHeap::Instance->RemoteHeapRead(DestPointer, Size, SourceAddress); break;
+	case 8: *(uint64*)DestPointer = FRemoteHeap::Instance->RemoteHeapRead8(SourceAddress, bWillWriteHint); break;
+	case 4: *(uint32*)DestPointer = FRemoteHeap::Instance->RemoteHeapRead4(SourceAddress, bWillWriteHint); break;
+	case 2: *(uint16*)DestPointer = FRemoteHeap::Instance->RemoteHeapRead2(SourceAddress, bWillWriteHint); break;
+	case 1: *(uint8*)DestPointer = FRemoteHeap::Instance->RemoteHeapRead1(SourceAddress, bWillWriteHint); break;
+	default: FRemoteHeap::Instance->RemoteHeapRead(DestPointer, Size, SourceAddress, bWillWriteHint); break;
 	}
 }
 
-void RemoteHeapRedirectedLoad8(void* DestPointer, uint64_t SourceAddress)
+void RemoteHeapRedirectedLoad8(void* DestPointer, uint64_t SourceAddress, bool bWillWriteHint)
 {
-	*(uint64*)DestPointer = FRemoteHeap::Instance->RemoteHeapRead8(SourceAddress);
+	*(uint64*)DestPointer = FRemoteHeap::Instance->RemoteHeapRead8(SourceAddress, bWillWriteHint);
 }
 
-void RemoteHeapRedirectedLoad4(void* DestPointer, uint64_t SourceAddress)
+void RemoteHeapRedirectedLoad4(void* DestPointer, uint64_t SourceAddress, bool bWillWriteHint)
 {
-	*(uint32*)DestPointer = FRemoteHeap::Instance->RemoteHeapRead4(SourceAddress);
+	*(uint32*)DestPointer = FRemoteHeap::Instance->RemoteHeapRead4(SourceAddress, bWillWriteHint);
 }
 
-void RemoteHeapRedirectedLoad2(void* DestPointer, uint64_t SourceAddress)
+void RemoteHeapRedirectedLoad2(void* DestPointer, uint64_t SourceAddress, bool bWillWriteHint)
 {
-	*(uint16*)DestPointer = FRemoteHeap::Instance->RemoteHeapRead2(SourceAddress);
+	*(uint16*)DestPointer = FRemoteHeap::Instance->RemoteHeapRead2(SourceAddress, bWillWriteHint);
 }
 
-void RemoteHeapRedirectedLoad1(void* DestPointer, uint64_t SourceAddress)
+void RemoteHeapRedirectedLoad1(void* DestPointer, uint64_t SourceAddress, bool bWillWriteHint)
 {
-	*(uint8*)DestPointer = FRemoteHeap::Instance->RemoteHeapRead1(SourceAddress);
+	*(uint8*)DestPointer = FRemoteHeap::Instance->RemoteHeapRead1(SourceAddress, bWillWriteHint);
 }
 
 void RemoteHeapRedirectedStore(uint64_t DestAddress, uint64_t Size, const void* SourcePointer)
@@ -110,18 +149,22 @@ void RemoteHeapRedirectedStore1(uint64_t DestAddress, const void* SourcePointer)
 	FRemoteHeap::Instance->RemoteHeapWrite1(DestAddress, *(uint8*)SourcePointer);
 }
 
-void FRemoteHeap::CreateInstance(FRemoteExecutor* Executor, FRemoteServerId LocalServerId)
+void FRemoteHeap::CreateInstance(FRemoteExecutor* Executor, FRemoteServerId LocalServerId, FRemoteHeapAllocatorOverride* AllocatorOverride)
 {
 	FRemoteHeap* Heap = new FRemoteHeap(LocalServerId);
+	Heap->AllocatorOverride = AllocatorOverride;
 	Instance = Heap;
 
 #if UE_WITH_REMOTE_OBJECT_HANDLE
 	RemoteExecutor::RegisterRemoteSubsystem(Executor, Heap);
 
-	uint64 LocalServerHeapBaseAddress = RemoteHeapGetBaseAddressForServerId(LocalServerId);
+	uint64 LocalServerHeapBaseAddress = GetRemoteHeapBaseAddressForServerId(LocalServerId);
 
-	pas_range heap_range = pas_range_create(LocalServerHeapBaseAddress, LocalServerHeapBaseAddress + RemoteHeapSizePerServer);
-	jit_heap_add_fresh_memory(heap_range);
+	if (!AllocatorOverride)
+	{
+		pas_range heap_range = pas_range_create(LocalServerHeapBaseAddress, LocalServerHeapBaseAddress + RemoteHeapSizePerServer);
+		jit_heap_add_fresh_memory(heap_range);
+	}
 
 	// Register the heap redirect callbacks once (they route through FRemoteHeap::Instance)
 	static bool bCallbacksRegistered = false;
@@ -150,142 +193,443 @@ void FRemoteHeap::CreateInstance(FRemoteExecutor* Executor, FRemoteServerId Loca
 
 void FRemoteHeap::BeginRequest()
 {
+#if UE_WITH_REMOTE_OBJECT_HANDLE
+	UE_LOG(LogRemoteHeap, VeryVerbose, TEXT("BeginRequest: request %s %s"),
+		*ActiveRequest->RequestId.ToString(),
+		*ActiveRequest->Priority.ToString());
+#endif
 }
 
-void FRemoteHeap::TickSubsystem()
+void FRemoteHeap::ForwardIndividualRequests(FRemoteServerId DestinationServer, const TArray<FRemoteHeapOutgoingLineRequest>& IndividualRequests, bool bForceNeedsData)
 {
-#if UE_WITH_REMOTE_OBJECT_HANDLE
-	// go through the pending line requests and see if we can satisfy any of them
-	for (int32 LineRequestsIndex = 0; LineRequestsIndex < OutgoingLineRequests.Num(); )
+	for (int32 LineRequestIndex = 0; LineRequestIndex < IndividualRequests.Num(); LineRequestIndex++)
 	{
-		FRemoteHeapOutgoingLineRequests& LineRequests = OutgoingLineRequests[LineRequestsIndex];
+		const FRemoteHeapOutgoingLineRequest& LineRequest = IndividualRequests[LineRequestIndex];
 
-		UE_LOG(LogRemoteHeap, VeryVerbose, TEXT("TickSubsystem processing (%d) requests for line %llx"),
-			LineRequests.IndividualRequests.Num(),
-			LineRequests.Address);
-
-		// do we have the line?
-		bool bFoundLineEntry = false;
+		if (DestinationServer != LineRequest.DestinationServerId)
+		{
+			if (LineRequest.RequestedAccess == ERemoteHeapLineState::Exclusive)
+			{
+				RemoteHeapRequestLineExclusiveDelegate.ExecuteIfBound(LineRequest.Address, DestinationServer, LineRequest.RequestPriority, LineRequest.DestinationServerId, bForceNeedsData ? true : LineRequest.NeedsData);
+			}
+			else
+			{
+				RemoteHeapRequestLineReadDelegate.ExecuteIfBound(LineRequest.Address, DestinationServer, LineRequest.RequestPriority, LineRequest.DestinationServerId);
+			}
 
-		for (int32 HeapLineIndex = 0; HeapLineIndex < HeapLines.Num(); HeapLineIndex++)
+			UE_LOG(LogRemoteHeap, VeryVerbose, TEXT("Forwarding request for line %llx to server %s for %s"),
+				LineRequest.Address,
+				*DestinationServer.ToString(),
+				*LineRequest.DestinationServerId.ToString());
+		}
+		else
 		{
-			SRemoteHeapLineEntry& LineEntry = HeapLines[HeapLineIndex];
+			UE_LOG(LogRemoteHeap, VeryVerbose, TEXT("Dropping request-to-itself for line %llx to server %s"),
+				LineRequest.Address,
+				*DestinationServer.ToString());
+		}
 
-			if (LineEntry.BaseAddress == LineRequests.Address)
-			{
-				bFoundLineEntry = true;
+	}
+}
 
-				if (LineEntry.State == ERemoteHeapLineState::Exclusive)
-				{
-					// yes, we have the line
-					FRemoteHeapOutgoingLineRequest* HighestPriorityRequest = &LineRequests.IndividualRequests[0];
+void FRemoteHeap::NotifyIndividualRequestsDeallocated(const TArray<FRemoteHeapOutgoingLineRequest>& IndividualRequests)
+{
+	for (int32 LineRequestIndex = 0; LineRequestIndex < IndividualRequests.Num(); LineRequestIndex++)
+	{
+		const FRemoteHeapOutgoingLineRequest& LineRequest = IndividualRequests[LineRequestIndex];
 
-					bool bLineLocked = false;
+		RemoteHeapNotifyLineDeallocatedDelegate.ExecuteIfBound(LineRequest.Address, LineRequest.DestinationServerId);
 
-					// search through all local transactions and see if the remote request is higher priority than all of them
-					for (int32 RequestIndex = 0; RequestIndex < GetRequestCount(); RequestIndex++)
-					{
-						FRemoteHeapRequest* ExistingRequest = GetRequestByIndex(RequestIndex);
+		UE_LOG(LogRemoteHeap, Verbose, TEXT("Line %llx deallocated: notifying requesting server %s"),
+			LineRequest.Address,
+			*LineRequest.DestinationServerId.ToString());
+	}
+}
 
-						if (IsHigherPriority(ExistingRequest->Priority, HighestPriorityRequest->RequestPriority))
-						{
-							// this ExistingRequest is higher priority than the remote request - does it access the same line?
-							if (ExistingRequest->RequestedLines.Contains(LineEntry.BaseAddress))
-							{
-								// yes, the line is locked locally
-								bLineLocked = true;
-
-								UE_LOG(LogRemoteHeap, VeryVerbose, TEXT("TickSubsystem line %llx locked: remote request %s is lower priority than local work %s"),
-									HighestPriorityRequest->Address,
-									*HighestPriorityRequest->RequestPriority.ToString(),
-									*ExistingRequest->Priority.ToString());
-								break;
-							}
-						}
-					}
 
-					if (bLineLocked)
-					{
-						LineRequestsIndex++;
-					}
-					else
-					{
-						// send the line away
+SRemoteHeapLineEntry* FRemoteHeap::FindHeapLine(uint64 Address)
+{
+	for (SRemoteHeapLineEntry& HeapLine : HeapLines)
+	{
+		if (HeapLine.Address == Address)
+		{
+			return &HeapLine;
+		}
+	}
+	return nullptr;
+}
 
-						if (UE_LOG_ACTIVE(LogRemoteHeap, VeryVerbose))
-						{
-							TCHAR LineText[RemoteHeapLineSize*3+1];
-							static const TCHAR HexDigits[] = TEXT("0123456789abcdef");
+const SRemoteHeapLineEntry* FRemoteHeap::FindHeapLine(uint64 Address) const
+{
+	for (const SRemoteHeapLineEntry& HeapLine : HeapLines)
+	{
+		if (HeapLine.Address == Address)
+		{
+			return &HeapLine;
+		}
+	}
+	return nullptr;
+}
 
-							TCHAR* LineTextCursor = &LineText[0];
-							for (int32 I = 0; I < LineEntry.Data.Num(); I++)
-							{
-								uint8 Byte = LineEntry.Data[I];
+FRemoteHeap::ELineRequestAction FRemoteHeap::ProcessExclusiveLineRequest(
+	FRemoteHeapOutgoingLineRequests& LineRequests,
+	SRemoteHeapLineEntry& LineEntry,
+	const FRemoteHeapOutgoingLineRequest& HighestPriorityRequest)
+{
+	//
+	// they are looking for exclusive access - are they higher priority than all of our requests for the line?
+	//
+	bool bLineLocked = false;
 
-								uint8 ByteHigh = (Byte / 16);
-								uint8 ByteLow = (Byte % 16);
+	for (int32 RequestIndex = 0; RequestIndex < GetRequestCount(); RequestIndex++)
+	{
+		FRemoteHeapRequest* LocalRequest = GetRequestByIndex(RequestIndex);
 
-								*LineTextCursor++ = HexDigits[ByteHigh];
-								*LineTextCursor++ = HexDigits[ByteLow];
-								*LineTextCursor++ = TEXT(' ');
-							}
+		if (IsHigherPriority(LocalRequest->Priority, HighestPriorityRequest.RequestPriority))
+		{
+			bLineLocked = LocalRequest->RequestedLines.Contains(LineEntry.Address);
 
-							*LineTextCursor++ = TEXT('\0');
+			if (bLineLocked)
+			{
+				UE_LOG(LogRemoteHeap, VeryVerbose, TEXT("TickSubsystem line %llx locked: remote exclusive request %s is lower priority than local request %s"),
+					HighestPriorityRequest.Address,
+					*HighestPriorityRequest.RequestPriority.ToString(),
+					*LocalRequest->Priority.ToString());
+				break;
+			}
+		}
+	}
 
-							UE_LOG(LogRemoteHeap, VeryVerbose, TEXT("TickSubsystem sending line %llx to %s %s | Line data: %s"),
-								HighestPriorityRequest->Address,
-								*HighestPriorityRequest->DestinationServerId.ToString(),
-								*HighestPriorityRequest->RequestPriority.ToString(),
-								LineText);
-						}
+	if (bLineLocked)
+	{
+#if WITH_DEV_AUTOMATION_TESTS
+		FRemoteHeapCoverage::RecordFlow(ERemoteHeapCoverageFlow::TickSubsystem_ExclusiveRequest_LineLocked);
+#endif
+		return ELineRequestAction::Advance;
+	}
 
-						RemoteHeapSendLineDelegate.ExecuteIfBound(LineEntry.BaseAddress, HighestPriorityRequest->DestinationServerId, LineEntry.Data);
+	//
+	// we have the line data (Shared or Exclusive) and we are giving it up to an Exclusive requestor
+	//
 
-						// for every other request, forward the request on to the new server
-						for (int32 LineRequestIndex = 1; LineRequestIndex < LineRequests.IndividualRequests.Num(); LineRequestIndex++)
-						{
-							FRemoteHeapOutgoingLineRequest& LineRequest = LineRequests.IndividualRequests[LineRequestIndex];
-							RemoteHeapRequestLineDelegate.ExecuteIfBound(LineEntry.BaseAddress, HighestPriorityRequest->DestinationServerId, LineRequest.RequestPriority, LineRequest.DestinationServerId);
+	// Send line data if requested — the destination is the new last mutator
+	if (HighestPriorityRequest.NeedsData)
+	{
+		RemoteHeapSendLineDelegate.ExecuteIfBound(LineEntry.Address,
+			HighestPriorityRequest.DestinationServerId, LineEntry.Data, HighestPriorityRequest.DestinationServerId,
+			LineEntry.LineVersion, LineEntry.AllocationBase, LineEntry.AllocationSize);
+	}
 
-							UE_LOG(LogRemoteHeap, VeryVerbose, TEXT("Forwarding request for line %llx to %s for %s"),
-								HighestPriorityRequest->Address,
-								*HighestPriorityRequest->DestinationServerId.ToString(),
-								*LineRequest.DestinationServerId.ToString());
-						}
+	// now build the ForwardedServers list for them (the set of servers that the requestor needs
+	// to hear an Ack from before they can consider the line exclusive to them)
+	//
+	// three situations:
+	//
+	// - we have the line Exclusive, ack it with empty ForwardedServers
+	//
+	// - we have the line Shared AND we have a write request in-progress:
+	//   - ForwardedServers is just the servers we're waiting on Acks from
+	//
+	// - we have the line Shared and do NOT have a write request in-progress:
+	//   - ForwardedServers is the LastMutator plus the servers we shared it with
+	TArray<FRemoteServerId> ForwardedServers;
+
+	if (LineEntry.State == ERemoteHeapLineState::Exclusive)
+	{
+#if WITH_DEV_AUTOMATION_TESTS
+		FRemoteHeapCoverage::RecordFlow(ERemoteHeapCoverageFlow::TickSubsystem_ExclusiveRequest_FromExclusive);
+#endif
+		UE_LOG(LogRemoteHeap, VeryVerbose, TEXT("TickSubsystem evicting line %llx (v%llu) for exclusive request from %s, forwarded to %d children"),
+			LineEntry.Address,
+			LineEntry.LineVersion,
+			*HighestPriorityRequest.DestinationServerId.ToString(),
+			ForwardedServers.Num());
+
+		RemoteHeapExclusiveAckDelegate.ExecuteIfBound(LineEntry.Address,
+			LineEntry.LineVersion, HighestPriorityRequest.DestinationServerId, ForwardedServers);
+	}
+	else
+	{
+		check(LineEntry.State == ERemoteHeapLineState::Shared);
 
-						// clear the data and keep track of where it went
-						LineEntry.State = ERemoteHeapLineState::Remote;
-						LineEntry.Data.Reset();
-						LineEntry.RemoteServers.Add(HighestPriorityRequest->DestinationServerId);
+		FRemoteHeapLineRequest* PendingWriteRequest = nullptr;
 
-						// all done, delete pending requests for this line
-						OutgoingLineRequests.RemoveAt(LineRequestsIndex, EAllowShrinking::No);
-						// do not increment LineRequestsIndex
-					}
-				}
-				else
+		for (int32 PendingIndex = 0; PendingIndex < PendingLineRequests.Num(); PendingIndex++)
+		{
+			FRemoteHeapLineRequest& PendingRequest = PendingLineRequests[PendingIndex];
+
+			if (PendingRequest.Address == LineEntry.Address)
+			{
+				check(PendingRequest.RequestedAccess == ERemoteHeapLineState::Exclusive);
+
+				PendingWriteRequest = &PendingRequest;
+
+				// we have an in-progress request to acquire Exclusive ownership of this line
+				// but we are now yielding to this requester
+				//
+				// we need to forward their request to all of the NeededAcks servers
+				for (const FRemoteServerId& NeededAckServer : PendingRequest.NeededAcks)
 				{
-					// we don't have the line, forward the requests
-					for (int32 LineRequestIndex = 0; LineRequestIndex < LineRequests.IndividualRequests.Num(); LineRequestIndex++)
+					if (!PendingRequest.ReceivedAcks.Contains(NeededAckServer) &&
+						(NeededAckServer != HighestPriorityRequest.DestinationServerId))
 					{
-						FRemoteHeapOutgoingLineRequest& LineRequest = LineRequests.IndividualRequests[LineRequestIndex];
-						RemoteHeapRequestLineDelegate.ExecuteIfBound(LineEntry.BaseAddress, LineEntry.RemoteServers[0], LineRequest.RequestPriority, LineRequest.DestinationServerId);
+						ForwardedServers.Add(NeededAckServer);
 
-						UE_LOG(LogRemoteHeap, VeryVerbose, TEXT("Forwarding request for line %llx to %s for %s"),
-							LineEntry.BaseAddress,
-							*LineEntry.RemoteServers[0].ToString(),
-							*LineRequest.DestinationServerId.ToString());
+						// pass NeedsData as false here because they can get the data from us
+						RemoteHeapRequestLineExclusiveDelegate.ExecuteIfBound(LineEntry.Address, NeededAckServer,
+							HighestPriorityRequest.RequestPriority, HighestPriorityRequest.DestinationServerId, false);
 					}
+				}
+			}
+		}
+
+		if (!PendingWriteRequest)
+		{
+#if WITH_DEV_AUTOMATION_TESTS
+			FRemoteHeapCoverage::RecordFlow(ERemoteHeapCoverageFlow::TickSubsystem_ExclusiveRequest_FromShared_NoPendingWrite);
+#endif
+			if ((LineEntry.LastMutatorServerId != HighestPriorityRequest.DestinationServerId) &&
+				(LineEntry.LastMutatorServerId != LocalServerId))
+			{
+				// if they are not themselves the last mutator, tell them who is
+				ForwardedServers.Add(LineEntry.LastMutatorServerId);
 
-					// all done, delete pending requests for this line
-					OutgoingLineRequests.RemoveAt(LineRequestsIndex, EAllowShrinking::No);
-					// do not increment LineRequestsIndex
+				// pass NeedsData as false here because they can get the data from us
+				RemoteHeapRequestLineExclusiveDelegate.ExecuteIfBound(LineEntry.Address, LineEntry.LastMutatorServerId,
+					HighestPriorityRequest.RequestPriority, HighestPriorityRequest.DestinationServerId, false);
+			}
+
+			// tell them about everyone we shared the line with and forward their request to them
+			for (const FRemoteServerId& ChildServer : LineEntry.RemoteServers)
+			{
+				if (ChildServer != HighestPriorityRequest.DestinationServerId)
+				{
+					ForwardedServers.Add(ChildServer);
+
+					// pass NeedsData as false here because they can get the data from us
+					RemoteHeapRequestLineExclusiveDelegate.ExecuteIfBound(LineEntry.Address, ChildServer,
+						HighestPriorityRequest.RequestPriority, HighestPriorityRequest.DestinationServerId, false);
 				}
+			}
+		}
+
+		// Send ack with forwarded servers list
+		RemoteHeapExclusiveAckDelegate.ExecuteIfBound(LineEntry.Address,
+			LineEntry.LineVersion, HighestPriorityRequest.DestinationServerId, ForwardedServers);
+
+		UE_LOG(LogRemoteHeap, VeryVerbose, TEXT("TickSubsystem evicting line %llx (v%llu) for exclusive request from %s, forwarded to %d children"),
+			LineEntry.Address,
+			LineEntry.LineVersion,
+			*HighestPriorityRequest.DestinationServerId.ToString(),
+			ForwardedServers.Num());
+
+		if (PendingWriteRequest)
+		{
+#if WITH_DEV_AUTOMATION_TESTS
+			FRemoteHeapCoverage::RecordFlow(ERemoteHeapCoverageFlow::TickSubsystem_ExclusiveRequest_FromShared_PendingWrite);
+#endif
+			// if we have a pending write request, clear our NeededAcks, replace it with
+			// this requestor, and forward the request to them
+			PendingWriteRequest->NeededAcks.Reset();
+
+			PendingWriteRequest->NeededAcks.Add(HighestPriorityRequest.DestinationServerId);
+
+			// we are giving the line up, so NeedsData = true
+			RemoteHeapRequestLineExclusiveDelegate.ExecuteIfBound(LineEntry.Address, HighestPriorityRequest.DestinationServerId,
+				PendingWriteRequest->RequestPriority, LocalServerId, true);
+		}
+	}
+
+	// Transition to Remote
+	LineEntry.State = ERemoteHeapLineState::Remote;
+	LineEntry.Data.Reset();
+	LineEntry.LastMutatorServerId = HighestPriorityRequest.DestinationServerId;
+	LineEntry.RemoteServers.Reset();
+
+	// Forward remaining requests to the winner:
+
+	// first remove the highest pri from the individual requests
+	LineRequests.IndividualRequests.RemoveAt(0);
+
+	// force NeedsData to true in this case because we expect the data
+	// to be mutated before it arrives at these servers
+	ForwardIndividualRequests(HighestPriorityRequest.DestinationServerId, LineRequests.IndividualRequests, true);
+
+	return ELineRequestAction::Consumed;
+}
+
+FRemoteHeap::ELineRequestAction FRemoteHeap::ProcessSharedLineRequest(
+	FRemoteHeapOutgoingLineRequests& LineRequests,
+	SRemoteHeapLineEntry& LineEntry,
+	const FRemoteHeapOutgoingLineRequest& HighestPriorityRequest)
+{
+	//
+	// they are looking for shared access - we can share it unless
+	// we have a higher priority Write request
+	//
+	bool bLineLocked = false;
+
+	for (int32 RequestIndex = 0; RequestIndex < GetRequestCount(); RequestIndex++)
+	{
+		FRemoteHeapRequest* ExistingRequest = GetRequestByIndex(RequestIndex);
+
+		if (IsHigherPriority(ExistingRequest->Priority, HighestPriorityRequest.RequestPriority))
+		{
+			const ERemoteHeapLineState* LocalRequestedAccess = ExistingRequest->RequestedLines.Find(LineEntry.Address);
+			if (LocalRequestedAccess && *LocalRequestedAccess == ERemoteHeapLineState::Exclusive)
+			{
+				bLineLocked = true;
+
+				UE_LOG(LogRemoteHeap, VeryVerbose, TEXT("TickSubsystem line %llx locked: remote read request %s is lower priority than local write work %s"),
+					HighestPriorityRequest.Address,
+					*HighestPriorityRequest.RequestPriority.ToString(),
+					*ExistingRequest->Priority.ToString());
 				break;
 			}
 		}
+	}
 
-		check(bFoundLineEntry);
+	if (bLineLocked)
+	{
+#if WITH_DEV_AUTOMATION_TESTS
+		FRemoteHeapCoverage::RecordFlow(ERemoteHeapCoverageFlow::TickSubsystem_SharedRequest_LineLocked);
+#endif
+		return ELineRequestAction::Advance;
+	}
+
+	UE_LOG(LogRemoteHeap, VeryVerbose, TEXT("TickSubsystem sending copy of line %llx to %s %s"),
+		HighestPriorityRequest.Address,
+		*HighestPriorityRequest.DestinationServerId.ToString(),
+		*HighestPriorityRequest.RequestPriority.ToString());
+
+	check(HighestPriorityRequest.NeedsData);
+
+#if WITH_DEV_AUTOMATION_TESTS
+	FRemoteHeapCoverage::RecordFlow(ERemoteHeapCoverageFlow::TickSubsystem_SharedRequest_Served);
+#endif
+
+	RemoteHeapSendLineDelegate.ExecuteIfBound(LineEntry.Address,
+		HighestPriorityRequest.DestinationServerId, LineEntry.Data, LineEntry.LastMutatorServerId,
+		LineEntry.LineVersion, LineEntry.AllocationBase, LineEntry.AllocationSize);
+
+	// Track the new child. If we have a pending exclusive upgrade,
+	// add to NeededAcks (which is the authoritative set during acquisition)
+	// and immediately tell them to evict. Otherwise add to RemoteServers.
+	for (FRemoteHeapLineRequest& PendingRequest : PendingLineRequests)
+	{
+		if ((PendingRequest.Address == LineEntry.Address) && (PendingRequest.RequestedAccess == ERemoteHeapLineState::Exclusive))
+		{
+#if WITH_DEV_AUTOMATION_TESTS
+			FRemoteHeapCoverage::RecordFlow(ERemoteHeapCoverageFlow::TickSubsystem_SharedRequest_HasPendingExclusiveUpgrade);
+#endif
+			PendingRequest.NeededAcks.Add(HighestPriorityRequest.DestinationServerId);
+
+			RemoteHeapRequestLineExclusiveDelegate.ExecuteIfBound(LineEntry.Address,
+				HighestPriorityRequest.DestinationServerId,
+				PendingRequest.RequestPriority, LocalServerId, false);
+			break;
+		}
+	}
+
+	LineEntry.RemoteServers.Add(HighestPriorityRequest.DestinationServerId);
+
+	// If we were Exclusive, downgrade to Shared
+	if (LineEntry.State == ERemoteHeapLineState::Exclusive)
+	{
+		LineEntry.State = ERemoteHeapLineState::Shared;
+	}
+
+	// satisfied this individual request; loop back if more remain
+	LineRequests.IndividualRequests.RemoveAt(0, EAllowShrinking::No);
+	return LineRequests.IndividualRequests.Num() > 0 ? ELineRequestAction::Retry : ELineRequestAction::Consumed;
+}
+
+void FRemoteHeap::TickSubsystem()
+{
+#if UE_WITH_REMOTE_OBJECT_HANDLE
+	//
+	// go through the pending outgoing line requests and see if we can satisfy any of them
+	//
+	for (int32 Index = 0; Index < OutgoingLineRequests.Num(); )
+	{
+		FRemoteHeapOutgoingLineRequests& LineRequests = OutgoingLineRequests[Index];
+
+		UE_LOG(LogRemoteHeap, VeryVerbose, TEXT("TickSubsystem processing (%d) requests for line %llx"),
+			LineRequests.IndividualRequests.Num(),
+			LineRequests.Address);
+
+		SRemoteHeapLineEntry* LineEntry = FindHeapLine(LineRequests.Address);
+
+		if (!LineEntry)
+		{
+			//
+			// we have no record of the line locally - we should be the allocating server
+			// (other servers would have this line as Remote, not missing, and have to ask us)
+			//
+			FRemoteServerId AllocatingServer = RemoteHeapGetOwnerServerId(LineRequests.Address);
+			check(AllocatingServer == LocalServerId);
+
+			// We are the allocating server and don't have this line — it must have been deallocated.
+			// Notify each requesting server so they can clean up.
+			// Note that we don't need to maintain the LineVersion across deallocation because
+			// all of the RPCs are ordered, and we are the only server able to allocate the memory
+			// at this address, and deallocation requires Exclusive access, so there is no
+			// chance of an in-flight message for this memory line at this point, and the
+			// next one will originate from us (the LineVersion will start back at 0)
+#if WITH_DEV_AUTOMATION_TESTS
+			FRemoteHeapCoverage::RecordFlow(ERemoteHeapCoverageFlow::TickSubsystem_LineNotFound_AllocatingServer);
+#endif
+			NotifyIndividualRequestsDeallocated(LineRequests.IndividualRequests);
+			OutgoingLineRequests.RemoveAt(Index, EAllowShrinking::No);
+			continue;
+		}
+
+		if (LineEntry->State == ERemoteHeapLineState::Remote)
+		{
+			//
+			// We have a record of the line but don't have the contents locally - forward
+			// all requests to the LastMutator or the AllocatingServer
+			//
+#if WITH_DEV_AUTOMATION_TESTS
+			FRemoteHeapCoverage::RecordFlow(ERemoteHeapCoverageFlow::TickSubsystem_LineRemote_Forward);
+#endif
+			ForwardIndividualRequests(LineEntry->LastMutatorServerId, LineRequests.IndividualRequests, false);
+			OutgoingLineRequests.RemoveAt(Index, EAllowShrinking::No);
+			continue;
+		}
+
+		//
+		// We have the line (Exclusive or Shared) - IndividualRequests is sorted by priority so let's
+		// consider the highest priority
+		//
+		// HighestPriorityRequest is taken by value here because the process functions may modify
+		// LineRequests.IndividualRequests (e.g. RemoveAt(0)) before they finish using it
+		//
+		FRemoteHeapOutgoingLineRequest HighestPriorityRequest = LineRequests.IndividualRequests[0];
+
+		ELineRequestAction Action;
+		if (HighestPriorityRequest.RequestedAccess == ERemoteHeapLineState::Exclusive)
+		{
+			Action = ProcessExclusiveLineRequest(LineRequests, *LineEntry, HighestPriorityRequest);
+		}
+		else
+		{
+			Action = ProcessSharedLineRequest(LineRequests, *LineEntry, HighestPriorityRequest);
+		}
+
+		switch (Action)
+		{
+			case ELineRequestAction::Consumed:
+				OutgoingLineRequests.RemoveAt(Index, EAllowShrinking::No);
+				break; // Index now points at the next entry
+
+			case ELineRequestAction::Advance:
+				Index++;
+				break;
+
+			case ELineRequestAction::Retry:
+				break; // same Index, next IndividualRequest is now at [0]
+		}
 	}
 #endif
 }
@@ -296,83 +640,262 @@ void FRemoteHeap::TickRequest()
 	ActiveRequest->LocalMutatedLines.Reset();
 
 	// go through the list of RequestedLines and see if we need to send off requests for them
-	for (uint64 Address : ActiveRequest->RequestedLines)
+	for (const auto& RequestedLinePair : ActiveRequest->RequestedLines)
 	{
-		bool bFound = false;
-		const SRemoteHeapLineEntry* HeapLineEntry = nullptr;
+		uint64 Address = RequestedLinePair.Key;
+		ERemoteHeapLineState RequestedAccess = RequestedLinePair.Value;
+		check((RequestedAccess == ERemoteHeapLineState::Exclusive) || (RequestedAccess == ERemoteHeapLineState::Shared));
 
-		for (const SRemoteHeapLineEntry& HeapLine : HeapLines)
+		const SRemoteHeapLineEntry* HeapLineEntry = FindHeapLine(Address);
+
+		if (HeapLineEntry)
 		{
-			if (HeapLine.BaseAddress == Address)
+			//
+			// we have a record of the line locally, check to see if the request
+			// is already satisfied by what we have
+			//
+			bool bSatisfied = false;
+
+			if (RequestedAccess == ERemoteHeapLineState::Exclusive)
+			{
+				bSatisfied = (HeapLineEntry->State == ERemoteHeapLineState::Exclusive);
+			}
+			else
+			{
+				bSatisfied = (HeapLineEntry->State == ERemoteHeapLineState::Exclusive) || (HeapLineEntry->State == ERemoteHeapLineState::Shared);
+			}
+
+			if (bSatisfied)
 			{
-				HeapLineEntry = &HeapLine;
+#if WITH_DEV_AUTOMATION_TESTS
+				FRemoteHeapCoverage::RecordFlow(ERemoteHeapCoverageFlow::TickRequest_AlreadySatisfied);
+#endif
+				UE_LOG(LogRemoteHeap, VeryVerbose, TEXT("TickRequest: line %llx already satisfied"), Address);
+				continue;
+			}
+		}
+
+		//
+		// Request is not satisfied - determine the server to request it from
+		//
+		FRemoteServerId ExpectedOwnerServerId;
+
+		if (HeapLineEntry)
+		{
+			// if we have a record of the line, then we might have a more accurate server to ask
+			// (the last server that we think mutated it)
+			ExpectedOwnerServerId = HeapLineEntry->LastMutatorServerId;
+		}
+		else
+		{
+			// we have no record of the line, the originating allocating server will
+			// always have some record of where it might be if it isn't deallocated
+			ExpectedOwnerServerId = RemoteHeapGetOwnerServerId(Address);
+		}
+
 
-				if (HeapLine.State == ERemoteHeapLineState::Exclusive)
+		//
+		// before we consider sending out a request for the line, let's find
+		// the highest priority local request for this line
+		//
+		FRemoteWorkPriority HighestRequestPriority = ActiveRequest->Priority;
+		for (int32 RequestIndex = 0; RequestIndex < GetRequestCount(); RequestIndex++)
+		{
+			FRemoteHeapRequest* ExistingRequest = GetRequestByIndex(RequestIndex);
+			const ERemoteHeapLineState* ExistingRequestedAccess = ExistingRequest->RequestedLines.Find(Address);
+			if (ExistingRequestedAccess && (*ExistingRequestedAccess == ERemoteHeapLineState::Exclusive) && (RequestedAccess == ERemoteHeapLineState::Exclusive))
+			{
+				if (IsHigherPriority(ExistingRequest->Priority, HighestRequestPriority))
 				{
-					bFound = true;
-					break;
+					HighestRequestPriority = ExistingRequest->Priority;
 				}
 			}
 		}
 
-		if (!bFound)
+		if (RequestedAccess == ERemoteHeapLineState::Exclusive)
 		{
-			// find the highest priority request that needs this line
-			FRemoteHeapRequest* HighestPriorityRequest = ActiveRequest;
+			bool bNeedsData = (!HeapLineEntry) || (HeapLineEntry->State == ERemoteHeapLineState::Remote);
 
-			for (int32 RequestIndex = 0; RequestIndex < GetRequestCount(); RequestIndex++)
+			//
+			// see if we have an existing (pending) line request for Exclusive
+			//
+			FRemoteHeapLineRequest* WriteRequest = nullptr;
+			for (FRemoteHeapLineRequest& ExistingRequest : PendingLineRequests)
 			{
-				FRemoteHeapRequest* ExistingRequest = GetRequestByIndex(RequestIndex);
-
-				if (ExistingRequest->RequestedLines.Contains(Address))
+				if ((ExistingRequest.Address == Address) && (ExistingRequest.RequestedAccess == ERemoteHeapLineState::Exclusive))
 				{
-					if (IsHigherPriority(ExistingRequest->Priority, HighestPriorityRequest->Priority))
-					{
-						HighestPriorityRequest = ExistingRequest;
-					}
+					WriteRequest = &ExistingRequest;
+					break;
 				}
 			}
 
-			FRemoteServerId ExpectedOwnerServerId;
-
-			if (HeapLineEntry)
+			if (WriteRequest)
 			{
-				ExpectedOwnerServerId = HeapLineEntry->RemoteServers[0];
-				UE_LOG(LogRemoteHeap, VeryVerbose, TEXT("Requesting %llx from %s from existing line entry"), Address, *ExpectedOwnerServerId.ToString());
+				if (IsHigherPriority(HighestRequestPriority, WriteRequest->RequestPriority))
+				{
+					//
+					// we have a new request for this line at a higher priority - we need to
+					// re-issue requests to all servers that we expect to hear from so
+					// they can reconsider our request
+					//
+#if WITH_DEV_AUTOMATION_TESTS
+					FRemoteHeapCoverage::RecordFlow(ERemoteHeapCoverageFlow::TickRequest_BoostExclusiveRequest);
+#endif
+					UE_LOG(LogRemoteHeap, VeryVerbose, TEXT("TickRequest: boosting exclusive request for %llx to %s"),
+						Address, *HighestRequestPriority.ToString());
+
+					WriteRequest->RequestPriority = HighestRequestPriority;
+
+					// Re-send at higher priority to all servers we're still waiting on
+					for (const FRemoteServerId& NeededServer : WriteRequest->NeededAcks)
+					{
+						if (!WriteRequest->ReceivedAcks.Contains(NeededServer))
+						{
+							RemoteHeapRequestLineExclusiveDelegate.Execute(Address, NeededServer,
+								WriteRequest->RequestPriority, LocalServerId, bNeedsData);
+						}
+					}
+				}
 			}
 			else
 			{
-				// we don't know, extract it from the address
-				ExpectedOwnerServerId = RemoteHeapGetOwnerServerId(Address);;
-				UE_LOG(LogRemoteHeap, VeryVerbose, TEXT("Requesting %llx from %s via deduction"), Address, *ExpectedOwnerServerId.ToString());
-			}
+				//
+				// we don't have an existing Write request, so we need to send it out
+				//
+#if WITH_DEV_AUTOMATION_TESTS
+				FRemoteHeapCoverage::RecordFlow(ERemoteHeapCoverageFlow::TickRequest_NewExclusiveRequest);
+#endif
 
-			bool bFoundExistingRequest = false;
+				//
+				// we first check if there is an existing Read request - if it
+				// is lower priority , we can discard it
+				//
+				for (int32 PendingIndex = 0; PendingIndex < PendingLineRequests.Num(); PendingIndex++)
+				{
+					FRemoteHeapLineRequest& PendingRequest = PendingLineRequests[PendingIndex];
 
-			for (int32 PendingRequestIndex = 0; PendingRequestIndex < PendingLineRequests.Num(); PendingRequestIndex++)
-			{
-				FRemoteHeapLineRequest* PendingRequest = &PendingLineRequests[PendingRequestIndex];
+					if ((PendingRequest.Address == Address) && (PendingRequest.RequestedAccess == ERemoteHeapLineState::Shared))
+					{
+						if (!IsHigherPriority(PendingRequest.RequestPriority, HighestRequestPriority))
+						{
+							PendingLineRequests.RemoveAtSwap(PendingIndex, EAllowShrinking::No);
+						}
+						break;
+					}
+				}
+
+				UE_LOG(LogRemoteHeap, VeryVerbose, TEXT("TickRequest: requesting exclusive for %llx from %s at %s"),
+					Address, *ExpectedOwnerServerId.ToString(), *HighestRequestPriority.ToString());
+
+				WriteRequest = &PendingLineRequests.Emplace_GetRef();
+				WriteRequest->Address = Address;
+				WriteRequest->RequestPriority = HighestRequestPriority;
+				WriteRequest->RequestedAccess = ERemoteHeapLineState::Exclusive;
 
-				if (PendingRequest->Address == Address)
+				//
+				// seed the NeededAcks list with the list of known servers we need to
+				// receive Acks from before we can consider the line Exclusive to us.
+				// each Ack that we receive may add more NeededAcks
+				//
+				if (ExpectedOwnerServerId != LocalServerId)
 				{
-					bFoundExistingRequest = true;
+					WriteRequest->NeededAcks.Add(ExpectedOwnerServerId);
+					RemoteHeapRequestLineExclusiveDelegate.Execute(Address, ExpectedOwnerServerId,
+						WriteRequest->RequestPriority, LocalServerId, bNeedsData);
+				}
+
+				if (HeapLineEntry)
+				{
+					if (HeapLineEntry->RemoteServers.Num() > 0)
+					{
+						// if we have RemoteServers that should mean that the 
+						// line is currently Shared and we don't need the data ourselves
+						check (HeapLineEntry->State == ERemoteHeapLineState::Shared);
+						check (bNeedsData == false);
+					}
 
-					if (IsHigherPriority(HighestPriorityRequest->Priority, PendingRequest->RequestPriority))
+					for (FRemoteServerId ChildServer : HeapLineEntry->RemoteServers)
 					{
-						// send a request and bump the priority
-						PendingRequest->RequestPriority = HighestPriorityRequest->Priority;
-						RemoteHeapRequestLineDelegate.Execute(Address, ExpectedOwnerServerId, PendingRequest->RequestPriority, LocalServerId);
+						WriteRequest->NeededAcks.Add(ChildServer);
+						RemoteHeapRequestLineExclusiveDelegate.Execute(Address, ChildServer,
+							WriteRequest->RequestPriority, LocalServerId, bNeedsData);
 					}
 				}
 			}
+		}
+		else
+		{
+			//
+			// Read-only access
+			//
 
-			if (!bFoundExistingRequest)
+			// Find existing Read pending request
+			FRemoteHeapLineRequest* ReadRequest = nullptr;
+			for (FRemoteHeapLineRequest& ExistingRequest : PendingLineRequests)
 			{
-				FRemoteHeapLineRequest* PendingRequest = &PendingLineRequests.Emplace_GetRef();
-				PendingRequest->Address = Address;
-				PendingRequest->RequestPriority = HighestPriorityRequest->Priority;
+				if ((ExistingRequest.Address == Address) && (ExistingRequest.RequestedAccess == ERemoteHeapLineState::Shared))
+				{
+					ReadRequest = &ExistingRequest;
+					break;
+				}
+			}
 
-				RemoteHeapRequestLineDelegate.Execute(Address, ExpectedOwnerServerId, PendingRequest->RequestPriority, LocalServerId);
+			if (ReadRequest)
+			{
+				//
+				// we haven an existing request - do we need to bump the priority?
+				//
+				if (IsHigherPriority(HighestRequestPriority, ReadRequest->RequestPriority))
+				{
+#if WITH_DEV_AUTOMATION_TESTS
+					FRemoteHeapCoverage::RecordFlow(ERemoteHeapCoverageFlow::TickRequest_BoostReadRequest);
+#endif
+					UE_LOG(LogRemoteHeap, VeryVerbose, TEXT("TickRequest: boosting read request for %llx to %s"),
+						Address, *HighestRequestPriority.ToString());
+
+					ReadRequest->RequestPriority = HighestRequestPriority;
+					RemoteHeapRequestLineReadDelegate.Execute(Address, ExpectedOwnerServerId,
+						ReadRequest->RequestPriority, LocalServerId);
+				}
+			}
+			else
+			{
+				//
+				// we don't have an existing read request, but we should only create one
+				// if there isn't already a higher priority write request
+				//
+				bool bHigherPriorityWriteRequestExists = false;
+				for (const FRemoteHeapLineRequest& ExistingRequest : PendingLineRequests)
+				{
+					if ((ExistingRequest.Address == Address) && (ExistingRequest.RequestedAccess == ERemoteHeapLineState::Exclusive))
+					{
+						bHigherPriorityWriteRequestExists = IsHigherPriority(ExistingRequest.RequestPriority, HighestRequestPriority);
+						break;
+					}
+				}
+
+				if (bHigherPriorityWriteRequestExists)
+				{
+#if WITH_DEV_AUTOMATION_TESTS
+					FRemoteHeapCoverage::RecordFlow(ERemoteHeapCoverageFlow::TickRequest_ReadSkipped_HigherPriorityWrite);
+#endif
+				}
+				else
+				{
+#if WITH_DEV_AUTOMATION_TESTS
+					FRemoteHeapCoverage::RecordFlow(ERemoteHeapCoverageFlow::TickRequest_NewReadRequest);
+#endif
+					UE_LOG(LogRemoteHeap, VeryVerbose, TEXT("TickRequest: requesting read for %llx from %s at %s"),
+						Address, *ExpectedOwnerServerId.ToString(), *HighestRequestPriority.ToString());
+
+					ReadRequest = &PendingLineRequests.Emplace_GetRef();
+					ReadRequest->Address = Address;
+					ReadRequest->RequestPriority = HighestRequestPriority;
+					ReadRequest->RequestedAccess = ERemoteHeapLineState::Shared;
+
+					RemoteHeapRequestLineReadDelegate.Execute(Address, ExpectedOwnerServerId,
+						ReadRequest->RequestPriority, LocalServerId);
+				}
 			}
 		}
 	}
@@ -386,21 +909,18 @@ void FRemoteHeap::TickAbortedRequest()
 bool FRemoteHeap::AreDependenciesSatisfied() const
 {
 #if UE_WITH_REMOTE_OBJECT_HANDLE
-	for (uint64 Address : ActiveRequest->RequestedLines)
+	for (const auto& RequestedLinePair : ActiveRequest->RequestedLines)
 	{
-		bool bFound = false;
+		uint64 Address = RequestedLinePair.Key;
+		ERemoteHeapLineState NeededState = RequestedLinePair.Value;
 
-		for (const SRemoteHeapLineEntry& HeapLine : HeapLines)
-		{
-			if (HeapLine.BaseAddress == Address)
-			{
-				if (HeapLine.State == ERemoteHeapLineState::Exclusive)
-				{
-					bFound = true;
-					break;
-				}
-			}
-		}
+		const SRemoteHeapLineEntry* HeapLine = FindHeapLine(Address);
+
+		// Exclusive satisfies both Exclusive and Shared needs.
+		// Shared satisfies only Shared needs.
+		bool bFound = HeapLine &&
+			((HeapLine->State == ERemoteHeapLineState::Exclusive) ||
+			 (HeapLine->State == ERemoteHeapLineState::Shared && NeededState == ERemoteHeapLineState::Shared));
 
 		if (!bFound)
 		{
@@ -420,13 +940,9 @@ void FRemoteHeap::EndRequest(bool bTransactionCommitted)
 		// commit the mutated lines into the line cache
 		for (SRemoteHeapLineData& MutatedLine : ActiveRequest->LocalMutatedLines)
 		{
-			for (SRemoteHeapLineEntry& HeapLine : HeapLines)
+			if (SRemoteHeapLineEntry* HeapLine = FindHeapLine(MutatedLine.Address))
 			{
-				if (HeapLine.BaseAddress == MutatedLine.BaseAddress)
-				{
-					HeapLine.Data = MoveTempIfPossible(MutatedLine.Data);
-					break;
-				}
+				HeapLine->Data = MoveTempIfPossible(MutatedLine.Data);
 			}
 		}
 
@@ -451,73 +967,90 @@ void FRemoteHeap::CommitMultiServerCommit()
 {
 }
 
-void FRemoteHeap::OnLineReceived(uint64 Address, const TArray<uint8>& Data)
+void FRemoteHeap::OnLineReceived(uint64 Address, const TArray<uint8>& Data,
+	FRemoteServerId InLastMutatorServerId, uint64 LineVersion,
+	FRemoteServerId SenderServerId, uint64 AllocationBase, uint64 AllocationSize)
 {
 #if UE_WITH_REMOTE_OBJECT_HANDLE
+	UE_LOG(LogRemoteHeap, VeryVerbose, TEXT("OnLineReceived: %llx (v%llu) from LastMutator %s"), Address, LineVersion, *InLastMutatorServerId.ToString());
 
-	if (UE_LOG_ACTIVE(LogRemoteHeap, VeryVerbose))
-	{
-		// dump the actual data when VeryVerbose
-		TCHAR LineText[RemoteHeapLineSize*3+1] = {};
-		static const TCHAR HexDigits[] = TEXT("0123456789abcdef");
-
-		TCHAR* LineTextCursor = &LineText[0];
-		for (int32 I = 0; I < Data.Num(); I++)
-		{
-			uint8 Byte = Data[I];
-
-			uint8 ByteHigh = (Byte / 16);
-			uint8 ByteLow = (Byte % 16);
-
-			*LineTextCursor++ = HexDigits[ByteHigh];
-			*LineTextCursor++ = HexDigits[ByteLow];
-			*LineTextCursor++ = TEXT(' ');
-		}
-
-		*LineTextCursor++ = TEXT('\0');
-
-		UE_LOG(LogRemoteHeap, VeryVerbose, TEXT("OnLineReceived: %llx : %s"), Address, LineText);
-	}
-
-	SRemoteHeapLineEntry* LineEntry = nullptr;
+	SRemoteHeapLineEntry* LineEntry = FindHeapLine(Address);
 
-	for (SRemoteHeapLineEntry& ExistingLineEntry : HeapLines)
+	if (LineEntry != nullptr)
 	{
-		if (ExistingLineEntry.BaseAddress == Address)
-		{
-			LineEntry = &ExistingLineEntry;
-		}
+		UE_LOG(LogRemoteHeap, VeryVerbose, TEXT("OnLineReceived: %llx : found existing line entry"), Address);
 	}
-
-	if (LineEntry == nullptr)
+	else
 	{
 		LineEntry = &HeapLines.Emplace_GetRef();
+		LineEntry->State = ERemoteHeapLineState::Shared;
+		LineEntry->Address = Address;
+		UE_LOG(LogRemoteHeap, VeryVerbose, TEXT("OnLineReceived: %llx : created new line entry"), Address);
 	}
 
-	LineEntry->State = ERemoteHeapLineState::Exclusive;
 	ensure(Data.Num() == RemoteHeapLineSize);
 	LineEntry->Data = Data;
-	LineEntry->BaseAddress = Address;
-	LineEntry->RemoteServers.Reset();
+	LineEntry->LastMutatorServerId = InLastMutatorServerId;
+	LineEntry->LineVersion = LineVersion;
+	LineEntry->AllocationBase = AllocationBase;
+	LineEntry->AllocationSize = AllocationSize;
 
-	for (int32 PendingRequestIndex = 0; PendingRequestIndex < PendingLineRequests.Num(); PendingRequestIndex++)
-	{
-		FRemoteHeapLineRequest* PendingRequest = &PendingLineRequests[PendingRequestIndex];
+	//
+	// check our pending requests
+	//
+
+	// the line is at least Shared in both the read-satisfied and write-waiting cases
+	LineEntry->State = ERemoteHeapLineState::Shared;
 
-		if (PendingRequest->Address == Address)
+	// If a pending read request exists for this line it is now satisfied — remove it
+	for (int32 PendingIndex = 0; PendingIndex < PendingLineRequests.Num(); PendingIndex++)
+	{
+		if (PendingLineRequests[PendingIndex].Address == Address &&
+			PendingLineRequests[PendingIndex].RequestedAccess == ERemoteHeapLineState::Shared)
 		{
-			PendingLineRequests.RemoveAtSwap(PendingRequestIndex, EAllowShrinking::No);
+#if WITH_DEV_AUTOMATION_TESTS
+			FRemoteHeapCoverage::RecordFlow(ERemoteHeapCoverageFlow::OnLineReceived_SharedRequestSatisfied);
+#endif
+			PendingLineRequests.RemoveAtSwap(PendingIndex, EAllowShrinking::No);
 			break;
 		}
 	}
+
+	// If a pending write request exists for this line, update its ack tracking
+	for (FRemoteHeapLineRequest& PendingRequest : PendingLineRequests)
+	{
+		if ((PendingRequest.Address != Address) || (PendingRequest.RequestedAccess != ERemoteHeapLineState::Exclusive))
+			continue;
+
+		// If this line version is newer, reset ack tracking
+		if (LineVersion > PendingRequest.AckVersion)
+		{
+			PendingRequest.NeededAcks.Reset();
+			PendingRequest.ReceivedAcks.Reset();
+			PendingRequest.AckVersion = LineVersion;
+		}
+
+		// The server that sent us this data must also ack (confirming it evicted its copy)
+		// before we can claim exclusive ownership
+		PendingRequest.NeededAcks.Add(SenderServerId);
+
+		// I don't believe it's possible to reach this point and have all acks already
+		// satisfied, because the Ack from the server sending the line always comes after
+		// the sending of the line, so we definitely don't have an ack from this server yet
+		check(!PendingRequest.AreAcksSatisfied());
+
+#if WITH_DEV_AUTOMATION_TESTS
+		FRemoteHeapCoverage::RecordFlow(ERemoteHeapCoverageFlow::OnLineReceived_ExclusiveRequest_WaitingForAcks);
+#endif
+		break; // at most one Exclusive request per address
+	}
 #endif
 }
 
-void FRemoteHeap::OnLineRequested(uint64 Address, FRemoteWorkPriority Priority, FRemoteServerId Destination)
+static void AddOutgoingLineRequest(TArray<FRemoteHeapOutgoingLineRequests>& OutgoingLineRequests,
+	uint64 Address, FRemoteWorkPriority Priority, FRemoteServerId Destination,
+	ERemoteHeapLineState RequestedAccess, bool NeedsData)
 {
-#if UE_WITH_REMOTE_OBJECT_HANDLE
-	UE_LOG(LogRemoteHeap, Verbose, TEXT("OnLineRequested: %llx %s %s"), Address, *Priority.ToString(), *Destination.ToString());
-
 	FRemoteHeapOutgoingLineRequests* LineRequests = nullptr;
 
 	for (FRemoteHeapOutgoingLineRequests& ExistingLineRequests : OutgoingLineRequests)
@@ -542,8 +1075,16 @@ void FRemoteHeap::OnLineRequested(uint64 Address, FRemoteWorkPriority Priority,
 		{
 			LineRequest = &ExistingLineRequest;
 
-			// update the priority
-			LineRequest->RequestPriority = Priority;
+			// only upgrade priority, access, and NeedsData — never downgrade
+			if (IsHigherPriority(Priority, LineRequest->RequestPriority))
+			{
+				LineRequest->RequestPriority = Priority;
+			}
+			if (RequestedAccess == ERemoteHeapLineState::Exclusive)
+			{
+				LineRequest->RequestedAccess = ERemoteHeapLineState::Exclusive;
+			}
+			LineRequest->NeedsData = LineRequest->NeedsData || NeedsData;
 			break;
 		}
 	}
@@ -554,24 +1095,225 @@ void FRemoteHeap::OnLineRequested(uint64 Address, FRemoteWorkPriority Priority,
 		LineRequest->Address = Address;
 		LineRequest->DestinationServerId = Destination;
 		LineRequest->RequestPriority = Priority;
+		LineRequest->RequestedAccess = RequestedAccess;
+		LineRequest->NeedsData = NeedsData;
 	}
 
 	// re-sort the list by priority (either this is a new request,
-	// or we may have updated the priority of an existing request
+	// or we may have updated the priority of an existing request)
+	// Sort exclusive requests before read requests at the same priority
 	LineRequests->IndividualRequests.Sort([](const FRemoteHeapOutgoingLineRequest& Lhs, const FRemoteHeapOutgoingLineRequest& Rhs) -> bool
 		{
-			return IsHigherPriority(Lhs.RequestPriority, Rhs.RequestPriority);
+			if (IsHigherPriority(Lhs.RequestPriority, Rhs.RequestPriority))
+				return true;
+			if (IsHigherPriority(Rhs.RequestPriority, Lhs.RequestPriority))
+				return false;
+
+			// Same priority, sort exclusive requests first
+			return (Lhs.RequestedAccess == ERemoteHeapLineState::Exclusive);
 		});
+}
+
+void FRemoteHeap::OnReadLineRequested(uint64 Address, FRemoteWorkPriority Priority, FRemoteServerId Destination)
+{
+#if UE_WITH_REMOTE_OBJECT_HANDLE
+	UE_LOG(LogRemoteHeap, Verbose, TEXT("OnReadLineRequested: %llx %s %s"), Address, *Priority.ToString(), *Destination.ToString());
+
+	AddOutgoingLineRequest(OutgoingLineRequests, Address, Priority, Destination, ERemoteHeapLineState::Shared, true);
+#endif
+}
+
+void FRemoteHeap::OnExclusiveRequested(uint64 Address, FRemoteWorkPriority Priority, FRemoteServerId Destination, bool NeedsData)
+{
+#if UE_WITH_REMOTE_OBJECT_HANDLE
+	UE_LOG(LogRemoteHeap, Verbose, TEXT("OnExclusiveRequested: %llx %s %s NeedsData=%d"), Address, *Priority.ToString(), *Destination.ToString(), NeedsData);
+
+	AddOutgoingLineRequest(OutgoingLineRequests, Address, Priority, Destination, ERemoteHeapLineState::Exclusive, NeedsData);
+#endif
+}
+
+void FRemoteHeap::OnExclusiveAck(uint64 Address, uint64 LineVersion,
+	FRemoteServerId AckingServer,
+	const TArray<FRemoteServerId>& ForwardedServers)
+{
+#if UE_WITH_REMOTE_OBJECT_HANDLE
+	{
+		TArray<FString> ForwardedStrings;
+		ForwardedStrings.Reserve(ForwardedServers.Num());
+		for (const FRemoteServerId& Server : ForwardedServers)
+		{
+			ForwardedStrings.Add(Server.ToString());
+		}
+		UE_LOG(LogRemoteHeap, VeryVerbose, TEXT("OnExclusiveAck: %llx from %s, forwarded %d servers [%s]"),
+			Address, *AckingServer.ToString(), ForwardedServers.Num(), *FString::Join(ForwardedStrings, TEXT(", ")));
+	}
+
+	const SRemoteHeapLineEntry* HeapLineEntry = FindHeapLine(Address);
+
+	// Find the pending exclusive request for this address
+	FRemoteHeapLineRequest* MatchingRequest = nullptr;
+	int32 MatchingIndex = -1;
+
+	for (int32 PendingIndex = 0; PendingIndex < PendingLineRequests.Num(); PendingIndex++)
+	{
+		FRemoteHeapLineRequest& PendingRequest = PendingLineRequests[PendingIndex];
+		if ((PendingRequest.Address == Address) && (PendingRequest.RequestedAccess == ERemoteHeapLineState::Exclusive))
+		{
+			MatchingRequest = &PendingRequest;
+			MatchingIndex = PendingIndex;
+			break;
+		}
+	}
+
+	if (MatchingIndex == -1)
+	{
+		// this server has yielded the line to us and we expect more acks
+		// from the list of ForwardedServers, but we are no longer requesting
+		// write access to the line, so we need to forward this ack on to
+		// either the last mutator or the allocating server
+#if WITH_DEV_AUTOMATION_TESTS
+		FRemoteHeapCoverage::RecordFlow(ERemoteHeapCoverageFlow::OnExclusiveAck_Abandoned);
+#endif
+
+		FRemoteServerId ForwardServerId = HeapLineEntry ? HeapLineEntry->LastMutatorServerId : RemoteHeapGetOwnerServerId(Address);
+
+		if (LocalServerId != ForwardServerId)
+		{
+			TArray<FRemoteServerId> CombinedForwarded = ForwardedServers;
+			CombinedForwarded.Add(AckingServer);
+
+			RemoteHeapExclusiveAckDelegate.ExecuteIfBound(Address, LineVersion,
+				ForwardServerId, CombinedForwarded);
+
+			UE_LOG(LogRemoteHeap, VeryVerbose, TEXT("OnExclusiveAck: forwarding abandoned ack for %llx (v%llu) to %s"),
+				Address, LineVersion, *ForwardServerId.ToString());
+		}
+	}
+	else
+	{
+		if (LineVersion > MatchingRequest->AckVersion)
+		{
+			// Newer version: reset tracking and start fresh
+#if WITH_DEV_AUTOMATION_TESTS
+			FRemoteHeapCoverage::RecordFlow(ERemoteHeapCoverageFlow::OnExclusiveAck_NewerVersion);
+#endif
+			UE_LOG(LogRemoteHeap, VeryVerbose, TEXT("OnExclusiveAck: %llx newer version reset (v%llu -> v%llu) from %s"),
+				Address, MatchingRequest->AckVersion, LineVersion, *AckingServer.ToString());
+
+			MatchingRequest->NeededAcks.Reset();
+			MatchingRequest->ReceivedAcks.Reset();
+			MatchingRequest->AckVersion = LineVersion;
+			MatchingRequest->NeededAcks = TSet<FRemoteServerId>(ForwardedServers);
+			MatchingRequest->ReceivedAcks.Add(AckingServer);
+		}
+		else if (LineVersion == MatchingRequest->AckVersion)
+		{
+			// Same version: accumulate
+			MatchingRequest->ReceivedAcks.Add(AckingServer);
+			for (const FRemoteServerId& ForwardedServer : ForwardedServers)
+			{
+				MatchingRequest->NeededAcks.Add(ForwardedServer);
+			}
+		}
+		else
+		{
+			// it is believed to be impossible to validly get a 'stale' ack because if
+			// a server sends an ack, then it understands it has yielded the line to
+			// an exclusive requester, so in order for anyone to full acquire that
+			// line, all acks need to resolve. Only after all acks resolve can the
+			// version be bumped, so it should not be possible to see an old ack version
+			check(false); 
+
+			UE_LOG(LogRemoteHeap, VeryVerbose, TEXT("OnExclusiveAck: ignoring stale ack for %llx (v%llu < v%llu)"),
+				Address, LineVersion, MatchingRequest->AckVersion);
+		}
+
+		{
+			TArray<FString> NeededStrings;
+			for (const FRemoteServerId& Server : MatchingRequest->NeededAcks)
+			{
+				NeededStrings.Add(Server.ToString());
+			}
+			TArray<FString> ReceivedStrings;
+			for (const FRemoteServerId& Server : MatchingRequest->ReceivedAcks)
+			{
+				ReceivedStrings.Add(Server.ToString());
+			}
+
+			UE_LOG(LogRemoteHeap, VeryVerbose, TEXT("OnExclusiveAck: %llx NeededAcks=[%s] ReceivedAcks=[%s]"),
+				Address, *FString::Join(NeededStrings, TEXT(", ")), *FString::Join(ReceivedStrings, TEXT(", ")));
+		}
+
+		// Check if all acks collected and we have the data
+		if (MatchingRequest->AreAcksSatisfied())
+		{
+			SRemoteHeapLineEntry* LineEntry = FindHeapLine(Address);
+
+			// all acks are satisfied - the line data arrives before the ack from any particular
+			// server, so we should already have the data
+			check(LineEntry && LineEntry->Data.Num() > 0);
+
+			// Promote to Exclusive
+#if WITH_DEV_AUTOMATION_TESTS
+			FRemoteHeapCoverage::RecordFlow(ERemoteHeapCoverageFlow::OnExclusiveAck_SameVersion_Promoted);
+#endif
+			LineEntry->State = ERemoteHeapLineState::Exclusive;
+			LineEntry->LineVersion++;
+			LineEntry->LastMutatorServerId = LocalServerId;
+			LineEntry->RemoteServers.Reset();
+
+			PendingLineRequests.RemoveAtSwap(MatchingIndex, EAllowShrinking::No);
+
+			UE_LOG(LogRemoteHeap, VeryVerbose, TEXT("OnExclusiveAck: line %llx promoted to Exclusive"), Address);
+		}
+	}
+#endif
+}
+
+void FRemoteHeap::OnLineDeallocated(uint64 Address)
+{
+#if UE_WITH_REMOTE_OBJECT_HANDLE
+	UE_LOG(LogRemoteHeap, Verbose, TEXT("OnLineDeallocated: %llx"), Address);
+#if WITH_DEV_AUTOMATION_TESTS
+	FRemoteHeapCoverage::RecordFlow(ERemoteHeapCoverageFlow::OnLineDeallocated);
+#endif
+
+	// Remove from HeapLines
+	for (int32 LineIndex = 0; LineIndex < HeapLines.Num(); LineIndex++)
+	{
+		if (HeapLines[LineIndex].Address == Address)
+		{
+			HeapLines.RemoveAtSwap(LineIndex, EAllowShrinking::No);
+			break;
+		}
+	}
+
+	// Remove all pending requests for this address (both Read and Write)
+	for (int32 PendingIndex = 0; PendingIndex < PendingLineRequests.Num(); PendingIndex++)
+	{
+		int32 PendingIndexReversed = PendingLineRequests.Num() - PendingIndex - 1;
+		if (PendingLineRequests[PendingIndexReversed].Address == Address)
+		{
+			PendingLineRequests.RemoveAtSwap(PendingIndexReversed, EAllowShrinking::No);
+		}
+	}
+
+	// Remove from all active requests' RequestedLines maps so AreDependenciesSatisfied() no longer blocks
+	for (int32 RequestIndex = 0; RequestIndex < GetRequestCount(); RequestIndex++)
+	{
+		FRemoteHeapRequest* Request = GetRequestByIndex(RequestIndex);
+		Request->RequestedLines.Remove(Address);
+	}
 #endif
 }
 
-void FRemoteHeap::ExtractHeapRangeReadonly(uint64 StartAddress, uint64 Size, void* DstData)
+void FRemoteHeap::ExtractHeapRangeReadonly(uint64 StartAddress, uint64 Size, void* DstData, bool bWillWriteHint)
 {
 #if UE_WITH_REMOTE_OBJECT_HANDLE
 	uint64 LineAddressCursor = StartAddress & RemoteHeapLineMask;
 	FRemoteServerId OwnerServerId = RemoteHeapGetOwnerServerId(StartAddress); 
 
-	UE_LOG(LogRemoteHeap, VeryVerbose, TEXT("ExtractHeapRangeReadonly StartAddress: %llx Size %llu | Owner %s"), StartAddress, Size, *OwnerServerId.ToString());
+	UE_LOG(LogRemoteHeap, VeryVerbose, TEXT("ExtractHeapRangeReadonly StartAddress: %llx Size %llu | Owner %s | bWillWriteHint %d"), StartAddress, Size, *OwnerServerId.ToString(), bWillWriteHint);
 
 	uint8* DstCursor = (uint8*)DstData;
 	uint64 DstAddressCursor = StartAddress;
@@ -591,7 +1333,7 @@ void FRemoteHeap::ExtractHeapRangeReadonly(uint64 StartAddress, uint64 Size, voi
 		// is this line locally mutable?
 		for(SRemoteHeapLineData& HeapLine : ActiveRequest->LocalMutatedLines)
 		{
-			if (HeapLine.BaseAddress == LineAddressCursor)
+			if (HeapLine.Address == LineAddressCursor)
 			{
 				UE_LOG(LogRemoteHeap, VeryVerbose, TEXT("ExtractHeapRangeReadonly: line %llx found in mutable section"), LineAddressCursor);
 					
@@ -606,19 +1348,15 @@ void FRemoteHeap::ExtractHeapRangeReadonly(uint64 StartAddress, uint64 Size, voi
 		if (!bFound)
 		{
 			// check the global cache
-			for(SRemoteHeapLineEntry& HeapLine : HeapLines)
+			if (SRemoteHeapLineEntry* HeapLine = FindHeapLine(LineAddressCursor))
 			{
-				if (HeapLine.BaseAddress == LineAddressCursor)
+				if (HeapLine->State == ERemoteHeapLineState::Exclusive ||
+					HeapLine->State == ERemoteHeapLineState::Shared)
 				{
-					if (HeapLine.State == ERemoteHeapLineState::Exclusive)
-					{
-						UE_LOG(LogRemoteHeap, VeryVerbose, TEXT("ExtractHeapRangeReadonly: line %llx found in readonly section"), LineAddressCursor);
-						FMemory::Memcpy(DstCursor, &HeapLine.Data[(uint32)LineOffset], SizeToCopy);
+					UE_LOG(LogRemoteHeap, VeryVerbose, TEXT("ExtractHeapRangeReadonly: line %llx found in readonly section"), LineAddressCursor);
+					FMemory::Memcpy(DstCursor, &HeapLine->Data[(uint32)LineOffset], SizeToCopy);
 
-						bFound = true;
-					}
-
-					break;
+					bFound = true;
 				}
 			}
 		}
@@ -641,11 +1379,32 @@ void FRemoteHeap::ExtractHeapRangeReadonly(uint64 StartAddress, uint64 Size, voi
 		{
 			for (uint64 UnavailableLine : PendingUnavailable)
 			{
-				ActiveRequest->RequestedLines.Add(UnavailableLine);
+				if (bWillWriteHint)
+				{
+					// Write subsumes Read — set Write directly
+					ActiveRequest->RequestedLines.FindOrAdd(UnavailableLine) = ERemoteHeapLineState::Exclusive;
+				}
+				else
+				{
+					// if the line is already requested, don't overwrite it (it might be requested for
+					// write, which we don't want to downgrade)
+					ERemoteHeapLineState* Existing = ActiveRequest->RequestedLines.Find(UnavailableLine);
+					if (!Existing)
+					{
+						ActiveRequest->RequestedLines.Add(UnavailableLine, ERemoteHeapLineState::Shared);
+					}
+				}
 			}
 		};
 
-		RemoteExecutor::AbortTransactionRequiresDependencies(TEXT("remote heap read"));
+		if (bWillWriteHint)
+		{
+			RemoteExecutor::AbortTransactionRequiresDependencies(TEXT("remote heap read (with write hint)"));
+		}
+		else
+		{
+			RemoteExecutor::AbortTransactionRequiresDependencies(TEXT("remote heap read"));
+		}
 	}
 #endif
 }
@@ -676,7 +1435,7 @@ void FRemoteHeap::InsertHeapRangeMutable(uint64 StartAddress, uint64 Size, const
 		// is this line locally mutable?
 		for(SRemoteHeapLineData& MutableHeapLine : ActiveRequest->LocalMutatedLines)
 		{
-			if (MutableHeapLine.BaseAddress == LineAddressCursor)
+			if (MutableHeapLine.Address == LineAddressCursor)
 			{
 				// found the data in the mutable section
 
@@ -691,22 +1450,17 @@ void FRemoteHeap::InsertHeapRangeMutable(uint64 StartAddress, uint64 Size, const
 		if (!bFound)
 		{
 			// can we pull a line into the mutable section?
-			for(SRemoteHeapLineEntry& HeapLine : HeapLines)
+			if (SRemoteHeapLineEntry* HeapLine = FindHeapLine(LineAddressCursor))
 			{
-				if (HeapLine.BaseAddress == LineAddressCursor)
+				if (HeapLine->State == ERemoteHeapLineState::Exclusive)
 				{
-					if (HeapLine.State == ERemoteHeapLineState::Exclusive)
-					{
-						UE_LOG(LogRemoteHeap, VeryVerbose, TEXT("InsertHeapRangeMutable: line %llx found in readonly section, moving to mutable"), LineAddressCursor);
-
-						// make a copy of the cached line into the mutable section
-						SRemoteHeapLineData& MutableHeapLine = ActiveRequest->LocalMutatedLines.Emplace_GetRef(HeapLine);
-						FMemory::Memcpy(&MutableHeapLine.Data[(uint32)LineOffset], SrcCursor, SizeToCopy);
+					UE_LOG(LogRemoteHeap, VeryVerbose, TEXT("InsertHeapRangeMutable: line %llx found in readonly section, moving to mutable"), LineAddressCursor);
 
-						bFound = true;
-					}
+					// make a copy of the cached line into the mutable section
+					SRemoteHeapLineData& MutableHeapLine = ActiveRequest->LocalMutatedLines.Emplace_GetRef(*HeapLine);
+					FMemory::Memcpy(&MutableHeapLine.Data[(uint32)LineOffset], SrcCursor, SizeToCopy);
 
-					break;
+					bFound = true;
 				}
 			}
 		}
@@ -729,7 +1483,8 @@ void FRemoteHeap::InsertHeapRangeMutable(uint64 StartAddress, uint64 Size, const
 		{
 			for (uint64 UnavailableLine : PendingUnavailable)
 			{
-				ActiveRequest->RequestedLines.Add(UnavailableLine);
+				// Write subsumes Read — set Write directly
+				ActiveRequest->RequestedLines.FindOrAdd(UnavailableLine) = ERemoteHeapLineState::Exclusive;
 			}
 		};
 
@@ -738,10 +1493,10 @@ void FRemoteHeap::InsertHeapRangeMutable(uint64 StartAddress, uint64 Size, const
 #endif
 }
 
-uint8 FRemoteHeap::RemoteHeapRead1(uint64 Addr)
+uint8 FRemoteHeap::RemoteHeapRead1(uint64 Addr, bool bWillWriteHint)
 {
 	uint8 Result = 0;
-	ExtractHeapRangeReadonly(Addr, sizeof(Result), &Result);
+	ExtractHeapRangeReadonly(Addr, sizeof(Result), &Result, bWillWriteHint);
 
 	UE_LOG(LogRemoteHeap, VeryVerbose, TEXT("RemoteHeapRead1: %llx %02x"), Addr, (uint32)Result);
 	return Result;
@@ -753,10 +1508,10 @@ void FRemoteHeap::RemoteHeapWrite1(uint64 Addr, uint8 Value)
 	UE_LOG(LogRemoteHeap, VeryVerbose, TEXT("RemoteHeapWrite1: %llx %02x"), Addr, (uint32)Value);
 }
 
-uint16 FRemoteHeap::RemoteHeapRead2(uint64 Addr)
+uint16 FRemoteHeap::RemoteHeapRead2(uint64 Addr, bool bWillWriteHint)
 {
 	uint16 Result = 0;
-	ExtractHeapRangeReadonly(Addr, sizeof(Result), &Result);
+	ExtractHeapRangeReadonly(Addr, sizeof(Result), &Result, bWillWriteHint);
 
 	UE_LOG(LogRemoteHeap, VeryVerbose, TEXT("RemoteHeapRead2: %llx %04x"), Addr, Result);
 	return Result;
@@ -768,10 +1523,10 @@ void FRemoteHeap::RemoteHeapWrite2(uint64 Addr, uint16 Value)
 	UE_LOG(LogRemoteHeap, VeryVerbose, TEXT("RemoteHeapWrite2: %llx %04x"), Addr, Value);
 }
 
-uint32 FRemoteHeap::RemoteHeapRead4(uint64 Addr)
+uint32 FRemoteHeap::RemoteHeapRead4(uint64 Addr, bool bWillWriteHint)
 {
 	uint32 Result = 0;
-	ExtractHeapRangeReadonly(Addr, sizeof(Result), &Result);
+	ExtractHeapRangeReadonly(Addr, sizeof(Result), &Result, bWillWriteHint);
 
 	UE_LOG(LogRemoteHeap, VeryVerbose, TEXT("RemoteHeapRead4: %llx %08x"), Addr, Result);
 	return Result;
@@ -783,10 +1538,10 @@ void FRemoteHeap::RemoteHeapWrite4(uint64 Addr, uint32 Value)
 	UE_LOG(LogRemoteHeap, VeryVerbose, TEXT("RemoteHeapWrite4: %llx %08x"), Addr, Value);
 }
 
-uint64 FRemoteHeap::RemoteHeapRead8(uint64 Addr)
+uint64 FRemoteHeap::RemoteHeapRead8(uint64 Addr, bool bWillWriteHint)
 {
 	uint64 Result = 0;
-	ExtractHeapRangeReadonly(Addr, sizeof(Result), &Result);
+	ExtractHeapRangeReadonly(Addr, sizeof(Result), &Result, bWillWriteHint);
 
 	UE_LOG(LogRemoteHeap, VeryVerbose, TEXT("RemoteHeapRead8: %llx %016llx"), Addr, Result);
 	return Result;
@@ -798,9 +1553,9 @@ void FRemoteHeap::RemoteHeapWrite8(uint64 Addr, uint64 Value)
 	UE_LOG(LogRemoteHeap, VeryVerbose, TEXT("RemoteHeapWrite8: %llx %016llx"), Addr, Value);
 }
 
-void FRemoteHeap::RemoteHeapRead(void* DestPointer, uint64 Size, uint64 SourceAddress)
+void FRemoteHeap::RemoteHeapRead(void* DestPointer, uint64 Size, uint64 SourceAddress, bool bWillWriteHint)
 {
-	ExtractHeapRangeReadonly(SourceAddress, Size, DestPointer);
+	ExtractHeapRangeReadonly(SourceAddress, Size, DestPointer, bWillWriteHint);
 	UE_LOG(LogRemoteHeap, VeryVerbose, TEXT("RemoteHeapRead: %llx (size %llu)"), SourceAddress, Size);
 }
 
@@ -822,23 +1577,34 @@ void* FRemoteHeap::RemoteHeapAllocate(uint64 Size)
 
 	UE_AUTORTFM_OPEN
 	{
-		Result = jit_heap_try_allocate(Size);
+		if (AllocatorOverride)
+			Result = AllocatorOverride->Allocate(AllocatorOverride->Context, Size);
+		else
+			Result = jit_heap_try_allocate(Size);
 	};
 
-	UE_AUTORTFM_ONABORT(Result)
+	UE_AUTORTFM_ONABORT(this, Result)
 	{
-		jit_heap_deallocate(Result);
+		if (AllocatorOverride)
+			AllocatorOverride->Deallocate(AllocatorOverride->Context, Result);
+		else
+			jit_heap_deallocate(Result);
 	};
 
-	uint64 LineAddressCursor = (uint64)Result;
+	uint64 AllocationBase = (uint64)Result;
+	uint64 LineAddressCursor = AllocationBase;
 	uint64 LineAddressEnd = LineAddressCursor + Size;
 
 	while(LineAddressCursor < LineAddressEnd)
 	{
 		SRemoteHeapLineEntry& NewHeapLine = FRemoteHeap::Instance->HeapLines.Emplace_GetRef();
-		NewHeapLine.BaseAddress = LineAddressCursor;
+		NewHeapLine.Address = LineAddressCursor;
 		NewHeapLine.State = ERemoteHeapLineState::Exclusive;
+		NewHeapLine.LineVersion = 1;
+		NewHeapLine.LastMutatorServerId = FRemoteHeap::Instance->LocalServerId;
 		NewHeapLine.Data.SetNum(RemoteHeapLineSize);
+		NewHeapLine.AllocationBase = AllocationBase;
+		NewHeapLine.AllocationSize = Size;
 
 		LineAddressCursor += RemoteHeapLineSize;
 	}
@@ -856,37 +1622,123 @@ void FRemoteHeap::RemoteHeapDeallocate(void* Ptr)
 		return;
 
 #if UE_WITH_REMOTE_OBJECT_HANDLE
-	// did we allocate this memory?
 	uint64 Address = 0;
 	FMemory::Memcpy(&Address, &Ptr, sizeof(uint64));
 
+	uint64 LineAddress = Address & RemoteHeapLineMask;
+	check(Address == LineAddress);
 	FRemoteServerId OwnerServer = RemoteHeapGetOwnerServerId(Address);
 
-	UE_AUTORTFM_ONCOMMIT(this, Ptr, Address, OwnerServer)
+	if (OwnerServer != LocalServerId)
 	{
-		if (OwnerServer == LocalServerId)
+		// We are not the allocating server - forward the deallocation to the owner.
+		// The allocating server is responsible for acquiring Exclusive access to all
+		// lines before freeing, so we don't need to do that here. Our HeapLines
+		// entries will be cleaned up via OnLineDeallocated once the allocating server
+		// completes the deallocation.
+		UE_AUTORTFM_ONCOMMIT(this, Address, OwnerServer)
 		{
-			UE_LOG(LogRemoteHeap, Verbose, TEXT("RemoteHeapDeallocate (local): %llx"), Address);
-			jit_heap_deallocate((void*)Ptr);
+			UE_LOG(LogRemoteHeap, Verbose, TEXT("RemoteHeapDeallocate (remote %s): %llx"), *OwnerServer.ToString(), Address);
+			RemoteHeapDeallocateDelegate.ExecuteIfBound(Address, OwnerServer);
+		};
+		return;
+	}
+
+
+	// Find the first line to get the allocation metadata
+	const SRemoteHeapLineEntry* FirstLine = FindHeapLine(LineAddress);
+	check(FirstLine != nullptr);
+	uint64 AllocBase = FirstLine->AllocationBase;
+	uint64 AllocSize = FirstLine->AllocationSize;
+
+	check(AllocBase != 0 && AllocSize != 0);
+	check(Address == AllocBase);
+
+	// We are the allocating server - require Exclusive access on ALL lines before
+	// freeing to ensure all caches are evicted across all servers, otherwise we risk
+	// stale data if the same address range is allocated again later
+	bool bAllExclusive = true;
+
+	for (uint64 LineCursor = AllocBase; LineCursor < AllocBase + AllocSize; LineCursor += RemoteHeapLineSize)
+	{
+		const SRemoteHeapLineEntry* HeapLine = FindHeapLine(LineCursor);
+		if (!HeapLine || HeapLine->State != ERemoteHeapLineState::Exclusive)
+		{
+			bAllExclusive = false;
+			break;
+		}
+	}
+
+	if (!bAllExclusive)
+	{
+		// Request Write access for ALL lines in the allocation, not just the ones
+		// currently non-Exclusive, because Exclusive lines could be evicted before retry
+		UE_AUTORTFM_OPEN
+		{
+			for (uint64 LineCursor = AllocBase; LineCursor < AllocBase + AllocSize; LineCursor += RemoteHeapLineSize)
+			{
+				ActiveRequest->RequestedLines.FindOrAdd(LineCursor) = ERemoteHeapLineState::Exclusive;
+			}
+		};
+		RemoteExecutor::AbortTransactionRequiresDependencies(TEXT("remote heap deallocate requires exclusive"));
+		return;
+	}
+
+	UE_AUTORTFM_ONCOMMIT(this, Ptr, Address, AllocBase, AllocSize)
+	{
+		UE_LOG(LogRemoteHeap, Verbose, TEXT("RemoteHeapDeallocate (local): %llx"), Address);
+		if (AllocatorOverride)
+		{
+			AllocatorOverride->Deallocate(AllocatorOverride->Context, (void*)Ptr);
 		}
 		else
 		{
-			UE_LOG(LogRemoteHeap, Verbose, TEXT("RemoteHeapDeallocate (remote %s): %llx"), *OwnerServer.ToString(), Address);
-			RemoteHeapDeallocateDelegate.ExecuteIfBound(Address, OwnerServer);
+			jit_heap_deallocate((void*)Ptr);
+		}
+
+		// Remove HeapLines entries for all lines in this allocation
+		for (int32 LineIndex = 0; LineIndex < HeapLines.Num(); LineIndex++)
+		{
+			int32 LineIndexReversed = HeapLines.Num() - LineIndex - 1;
+			if (HeapLines[LineIndexReversed].AllocationBase == AllocBase)
+			{
+				HeapLines.RemoveAtSwap(LineIndexReversed, EAllowShrinking::No);
+			}
 		}
 	};
 
 #endif
 }
 
-void RemoteHeapOnLineReceived(uint64 Address, const TArray<uint8>& Data)
+void RemoteHeapOnLineReceived(uint64 Address, const TArray<uint8>& Data,
+	FRemoteServerId LastMutatorServerId, uint64 LineVersion,
+	FRemoteServerId SenderServerId,
+	uint64 AllocationBase, uint64 AllocationSize)
 {
-	FRemoteHeap::Instance->OnLineReceived(Address, Data);
+	FRemoteHeap::Instance->OnLineReceived(Address, Data, LastMutatorServerId, LineVersion, SenderServerId, AllocationBase, AllocationSize);
 }
 
 void RemoteHeapOnLineRequested(uint64 Address, FRemoteWorkPriority Priority, FRemoteServerId Destination)
 {
-	FRemoteHeap::Instance->OnLineRequested(Address, Priority, Destination);
+	FRemoteHeap::Instance->OnReadLineRequested(Address, Priority, Destination);
+}
+
+void RemoteHeapOnExclusiveRequested(uint64 Address, FRemoteWorkPriority Priority,
+	FRemoteServerId Destination, bool NeedsData)
+{
+	FRemoteHeap::Instance->OnExclusiveRequested(Address, Priority, Destination, NeedsData);
+}
+
+void RemoteHeapOnExclusiveAck(uint64 Address, uint64 LineVersion,
+	FRemoteServerId AckingServer,
+	const TArray<FRemoteServerId>& ForwardedServers)
+{
+	FRemoteHeap::Instance->OnExclusiveAck(Address, LineVersion, AckingServer, ForwardedServers);
+}
+
+void RemoteHeapOnLineDeallocated(uint64 Address)
+{
+	FRemoteHeap::Instance->OnLineDeallocated(Address);
 }
 
 void RemoteHeapDeallocate(uint64 Address)
diff --git a/Engine/Source/Runtime/CoreUObject/Public/UObject/RemoteHeap.h b/Engine/Source/Runtime/CoreUObject/Public/UObject/RemoteHeap.h
--- a/Engine/Source/Runtime/CoreUObject/Public/UObject/RemoteHeap.h
+++ b/Engine/Source/Runtime/CoreUObject/Public/UObject/RemoteHeap.h
@@ -3,30 +3,73 @@
 #include "CoreTypes.h"
 #include "AutoRTFM.h"
 #include "UObject/RemoteExecutor.h"
+#include "RemoteHeap.generated.h"
 
 namespace UE
 {
-extern TDelegate<void(uint64 /*Address*/, FRemoteServerId /*Destination*/, const TArray<uint8>& /*Data*/)> RemoteHeapSendLineDelegate;
-extern TDelegate<void(uint64 /*Address*/, FRemoteServerId /*ExpectedOwner*/, FRemoteWorkPriority /*Priority*/, FRemoteServerId /*Destination*/)> RemoteHeapRequestLineDelegate;
+// Sends a cache line to another server (for read or exclusive response)
+// Includes LastMutatorServerId so the receiver knows the sharing tree root
+extern TDelegate<void(uint64 /*Address*/, FRemoteServerId /*Destination*/,
+	const TArray<uint8>& /*Data*/, FRemoteServerId /*LastMutatorServerId*/,
+	uint64 /*LineVersion*/, uint64 /*AllocationBase*/, uint64 /*AllocationSize*/)>
+	RemoteHeapSendLineDelegate;
+
+// Requests a cache line for read (Shared) access
+extern TDelegate<void(uint64 /*Address*/, FRemoteServerId /*ExpectedOwner*/,
+	FRemoteWorkPriority /*Priority*/, FRemoteServerId /*Destination*/)>
+	RemoteHeapRequestLineReadDelegate;
+
+// Requests a cache line for exclusive (write) access
+// NeedsData: true if the requester needs the line data (Remote -> Exclusive),
+// false if the requester already has a Shared copy (Shared -> Exclusive)
+extern TDelegate<void(uint64 /*Address*/, FRemoteServerId /*ExpectedOwner*/,
+	FRemoteWorkPriority /*Priority*/, FRemoteServerId /*Destination*/,
+	bool /*NeedsData*/)>
+	RemoteHeapRequestLineExclusiveDelegate;
+
+// Acknowledge eviction of a cache line for exclusive request
+// Includes the list of servers the evicting server forwarded the invalidation to
+extern TDelegate<void(uint64 /*Address*/, uint64 /*LineVersion*/,
+	FRemoteServerId /*Destination*/,
+	const TArray<FRemoteServerId>& /*ForwardedServers*/)>
+	RemoteHeapExclusiveAckDelegate;
+
 extern TDelegate<void(uint64 /*Address*/, FRemoteServerId /*Owner*/)> RemoteHeapDeallocateDelegate;
 
+// Notifies a server that a line it requested has been deallocated by the allocating server
+extern TDelegate<void(uint64 /*Address*/, FRemoteServerId /*Destination*/)>
+	RemoteHeapNotifyLineDeallocatedDelegate;
+
+const uint64 RemoteHeapSizePerServer = 0x0000000400000000ull; // 16GiB
+
+struct FRemoteHeapAllocatorOverride
+{
+	void* (*Allocate)(void* Context, size_t Size) = nullptr;
+	void  (*Deallocate)(void* Context, void* Ptr) = nullptr;
+	void* Context = nullptr;
+};
+
 enum class ERemoteHeapLineState : uint8
 {
 	Exclusive,
-	//Shared, // TODO
+	Shared,
 	Remote,
 };
 
 struct SRemoteHeapLineData
 {
-	uint64 BaseAddress = 0;
+	uint64 Address = 0;
 	TArray<uint8> Data;
 };
 
 struct SRemoteHeapLineEntry : SRemoteHeapLineData
 {
 	ERemoteHeapLineState State = ERemoteHeapLineState::Exclusive;
+	FRemoteServerId LastMutatorServerId;
 	TArray<FRemoteServerId> RemoteServers;
+	uint64 LineVersion = 0;
+	uint64 AllocationBase = 0;
+	uint64 AllocationSize = 0;
 };
 
 class FRemoteHeapRequest
@@ -35,7 +78,7 @@ public:
 	FRemoteTransactionId RequestId;
 	FRemoteWorkPriority Priority;
 
-	TSet<uint64> RequestedLines;
+	TMap<uint64, ERemoteHeapLineState> RequestedLines;
 
 	// can only copy Exclusive heap lines here
 	TArray<SRemoteHeapLineData> LocalMutatedLines;
@@ -47,13 +90,17 @@ public:
 	uint64 Address = 0;
 	FRemoteServerId DestinationServerId;
 	FRemoteWorkPriority RequestPriority;
+	ERemoteHeapLineState RequestedAccess;
+	bool NeedsData = false;
 
 	FString ToString() const
 	{
-		return FString::Format(TEXT("[address %llx dest %s %s]"),
+		return FString::Format(TEXT("[address %llx dest %s %s %s%s]"),
 			{Address,
 			*DestinationServerId.ToString(),
-			*RequestPriority.ToString()});
+			*RequestPriority.ToString(),
+			RequestedAccess == ERemoteHeapLineState::Exclusive ? TEXT("exclusive") : TEXT("read"),
+			NeedsData ? TEXT(" needs-data") : TEXT("")});
 	}
 };
 
@@ -71,12 +118,34 @@ class FRemoteHeapLineRequest
 public:
 	uint64 Address = 0;
 	FRemoteWorkPriority RequestPriority;
+	ERemoteHeapLineState RequestedAccess;
+	uint64 AckVersion = 0;
+	TSet<FRemoteServerId> NeededAcks;
+	TSet<FRemoteServerId> ReceivedAcks;
+
+	bool AreAcksSatisfied() const
+	{
+		if (ReceivedAcks.Num() < NeededAcks.Num())
+		{
+			return false;
+		}
+		for (const FRemoteServerId& Needed : NeededAcks)
+		{
+			if (!ReceivedAcks.Contains(Needed))
+			{
+				return false;
+			}
+		}
+		return true;
+	}
 };
 
 class FRemoteHeap : public FRemoteSubsystem<FRemoteHeap, FRemoteHeapRequest>
 {
 	using FRemoteSubsystem::FRemoteSubsystem;
 
+	FRemoteHeapAllocatorOverride* AllocatorOverride = nullptr; // non-owning, must outlive this heap
+
 	TArray<SRemoteHeapLineEntry> HeapLines;
 
 	// this is a list of heap lines that we have locally that 
@@ -87,8 +156,21 @@ class FRemoteHeap : public FRemoteSubsystem<FRemoteHeap, FRemoteHeapRequest>
 	// for and are awaiting responses
 	TArray<FRemoteHeapLineRequest> PendingLineRequests;
 
+	void NotifyIndividualRequestsDeallocated(const TArray<FRemoteHeapOutgoingLineRequest>& IndividualRequests);
+	void ForwardIndividualRequests(FRemoteServerId DestinationServer, const TArray<FRemoteHeapOutgoingLineRequest>& IndividualRequests, bool bForceNeedsData);
+
+	enum class ELineRequestAction { Consumed, Advance, Retry };
+	SRemoteHeapLineEntry* FindHeapLine(uint64 Address);
+	const SRemoteHeapLineEntry* FindHeapLine(uint64 Address) const;
+	ELineRequestAction ProcessExclusiveLineRequest(FRemoteHeapOutgoingLineRequests& LineRequests, SRemoteHeapLineEntry& LineEntry, const FRemoteHeapOutgoingLineRequest& HighestPriorityRequest);
+	ELineRequestAction ProcessSharedLineRequest(FRemoteHeapOutgoingLineRequests& LineRequests, SRemoteHeapLineEntry& LineEntry, const FRemoteHeapOutgoingLineRequest& HighestPriorityRequest);
+
+
 public:
-	static void CreateInstance(class FRemoteExecutor* Executor, FRemoteServerId LocalServerId);
+	static void CreateInstance(class FRemoteExecutor* Executor, FRemoteServerId LocalServerId, FRemoteHeapAllocatorOverride* AllocatorOverride = nullptr);
+	static uint64 GetRemoteHeapBaseAddressForServerId(FRemoteServerId ServerId);
+
+	const TArray<SRemoteHeapLineEntry>& GetHeapLinesForTesting() const { return HeapLines; }
 
 	const TCHAR* NameForDebug() final { return TEXT("RemoteHeap"); }
 	void BeginRequest() final;
@@ -103,35 +185,124 @@ public:
 	void AbortMultiServerCommit() final;
 	void CommitMultiServerCommit() final;
 	
-	void OnLineReceived(uint64 Address, const TArray<uint8>& Data);
-	void OnLineRequested(uint64 Address, FRemoteWorkPriority Priority, FRemoteServerId Destination);
+	void OnLineReceived(uint64 Address, const TArray<uint8>& Data,
+		FRemoteServerId LastMutatorServerId, uint64 LineVersion,
+		FRemoteServerId SenderServerId, uint64 AllocationBase, uint64 AllocationSize);
+
+	void OnReadLineRequested(uint64 Address, FRemoteWorkPriority Priority,
+		FRemoteServerId Destination);
 
-	void ExtractHeapRangeReadonly(uint64 BaseAddress, uint64 Size, void* DstData);
+	void OnExclusiveRequested(uint64 Address, FRemoteWorkPriority Priority,
+		FRemoteServerId Destination, bool NeedsData);
+
+	void OnExclusiveAck(uint64 Address, uint64 LineVersion,
+		FRemoteServerId AckingServer,
+		const TArray<FRemoteServerId>& ForwardedServers);
+
+	void OnLineDeallocated(uint64 Address);
+
+	void ExtractHeapRangeReadonly(uint64 StartAddress, uint64 Size, void* DstData, bool bWillWriteHint);
 	void InsertHeapRangeMutable(uint64 StartAddress, uint64 Size, const void* SrcData);
 
-	uint8 RemoteHeapRead1(uint64 Addr);
+	uint8 RemoteHeapRead1(uint64 Addr, bool bWillWriteHint);
 	void RemoteHeapWrite1(uint64 Addr, uint8 Value);
 
-	uint16 RemoteHeapRead2(uint64 Addr);
+	uint16 RemoteHeapRead2(uint64 Addr, bool bWillWriteHint);
 	void RemoteHeapWrite2(uint64 Addr, uint16 Value);
 
-	uint32 RemoteHeapRead4(uint64 Addr);
+	uint32 RemoteHeapRead4(uint64 Addr, bool bWillWriteHint);
 	void RemoteHeapWrite4(uint64 Addr, uint32 Value);
 
-	uint64 RemoteHeapRead8(uint64 Addr);
+	uint64 RemoteHeapRead8(uint64 Addr, bool bWillWriteHint);
 	void RemoteHeapWrite8(uint64 Addr, uint64 Value);
 
-	void RemoteHeapRead(void* DestPointer, uint64 Size, uint64 SourceAddress);
+	void RemoteHeapRead(void* DestPointer, uint64 Size, uint64 SourceAddress, bool bWillWriteHint);
 	void RemoteHeapWrite(uint64 DestAddress, uint64 Size, const void* SourcePointer);
 
 	void* RemoteHeapAllocate(uint64 Size);
 	void RemoteHeapDeallocate(void* Ptr);
 };
 
-void RemoteHeapOnLineReceived(uint64 Address, const TArray<uint8>& Data);
-void RemoteHeapOnLineRequested(uint64 Address, FRemoteWorkPriority Priority, FRemoteServerId Destination);
+void RemoteHeapOnLineReceived(uint64 Address, const TArray<uint8>& Data,
+	FRemoteServerId LastMutatorServerId, uint64 LineVersion,
+	FRemoteServerId SenderServerId, uint64 AllocationBase, uint64 AllocationSize);
+void RemoteHeapOnLineRequested(uint64 Address, FRemoteWorkPriority Priority,
+	FRemoteServerId Destination);
+void RemoteHeapOnExclusiveRequested(uint64 Address, FRemoteWorkPriority Priority,
+	FRemoteServerId Destination, bool NeedsData);
+void RemoteHeapOnExclusiveAck(uint64 Address, uint64 LineVersion,
+	FRemoteServerId AckingServer,
+	const TArray<FRemoteServerId>& ForwardedServers);
+void RemoteHeapOnLineDeallocated(uint64 Address);
 void RemoteHeapDeallocate(uint64 Address);
 
+} // namespace UE
+
+// ERemoteHeapCoverageFlow is declared at global scope (outside namespace UE) so that
+// UHT can generate reflection data for it via UENUM(). The enum is always present;
+// only the mutable FRemoteHeapCoverage tracking struct is compiled out in non-test builds.
+UENUM()
+enum class ERemoteHeapCoverageFlow : uint32
+{
+	// TickSubsystem flows
+	TickSubsystem_LineNotFound_AllocatingServer         UMETA(DisplayName="TickSubsystem: line not found, we are allocator -> notify deallocated"),
+	TickSubsystem_LineRemote_Forward                    UMETA(DisplayName="TickSubsystem: line is Remote -> forward to LastMutatorServerId"),
+	TickSubsystem_ExclusiveRequest_LineLocked           UMETA(DisplayName="TickSubsystem: lower-priority remote exclusive -> deferred"),
+	TickSubsystem_ExclusiveRequest_FromExclusive        UMETA(DisplayName="TickSubsystem: evict Exclusive line, empty ForwardedServers"),
+	TickSubsystem_ExclusiveRequest_FromShared_NoPendingWrite UMETA(DisplayName="TickSubsystem: evict Shared line, build ForwardedServers"),
+	TickSubsystem_ExclusiveRequest_FromShared_PendingWrite   UMETA(DisplayName="TickSubsystem: yield pending write to higher-priority requester"),
+	TickSubsystem_SharedRequest_LineLocked              UMETA(DisplayName="TickSubsystem: lower-priority remote read -> deferred"),
+	TickSubsystem_SharedRequest_Served                  UMETA(DisplayName="TickSubsystem: send copy, downgrade Exclusive->Shared if needed"),
+	TickSubsystem_SharedRequest_HasPendingExclusiveUpgrade   UMETA(DisplayName="TickSubsystem: add new child to NeededAcks, request eviction"),
+
+	// TickRequest flows
+	TickRequest_AlreadySatisfied                        UMETA(DisplayName="TickRequest: line present at required access level, skip"),
+	TickRequest_NewExclusiveRequest                     UMETA(DisplayName="TickRequest: no pending write request -> send new exclusive RPC"),
+	TickRequest_BoostExclusiveRequest                   UMETA(DisplayName="TickRequest: pending write exists at lower priority -> re-send at higher"),
+	TickRequest_NewReadRequest                          UMETA(DisplayName="TickRequest: no pending read request -> send new read RPC"),
+	TickRequest_BoostReadRequest                        UMETA(DisplayName="TickRequest: pending read exists at lower priority -> re-send at higher"),
+	TickRequest_ReadSkipped_HigherPriorityWrite         UMETA(DisplayName="TickRequest: higher-priority write pending -> skip read request"),
+
+	// OnLineReceived flows
+	OnLineReceived_SharedRequestSatisfied               UMETA(DisplayName="OnLineReceived: pending Shared request satisfied, removed"),
+	OnLineReceived_ExclusiveRequest_WaitingForAcks      UMETA(DisplayName="OnLineReceived: pending Exclusive request, data arrived, still waiting on acks"),
+
+	// OnExclusiveAck flows
+	OnExclusiveAck_Abandoned                            UMETA(DisplayName="OnExclusiveAck: no matching write request -> forward ack onward"),
+	OnExclusiveAck_NewerVersion                         UMETA(DisplayName="OnExclusiveAck: ack version newer than tracked -> reset ack tracking"),
+	OnExclusiveAck_SameVersion_Promoted                 UMETA(DisplayName="OnExclusiveAck: all acks collected and data present -> promote to Exclusive"),
+
+	// OnLineDeallocated flow
+	OnLineDeallocated                                   UMETA(DisplayName="OnLineDeallocated: remove line and pending requests from local cache"),
+
+	Count                                               UMETA(Hidden)
+};
+
+namespace UE
+{
+
+#if WITH_DEV_AUTOMATION_TESTS
+
+struct FRemoteHeapCoverage
+{
+	// Clears coverage. Call at the start of each test.
+	COREUOBJECT_API static void Reset();
+
+	// Records that a flow was reached.
+	COREUOBJECT_API static void RecordFlow(ERemoteHeapCoverageFlow Flow);
+
+	// Returns the coverage bitfield (since last Reset).
+	COREUOBJECT_API static uint64 GetCoveredFlows();
+
+	// Returns true if the given flow was hit since the last Reset.
+	COREUOBJECT_API static bool IsFlowCovered(ERemoteHeapCoverageFlow Flow);
+
+	// Returns a mask with a bit set for every defined flow (i.e. full coverage).
+	COREUOBJECT_API static uint64 GetFullCoverageMask();
+};
+
+#endif // WITH_DEV_AUTOMATION_TESTS
+
 struct FRemoteHeapAllocator
 {
 	using SizeType = int32;
