[UObject] Remove loaders earlier than GC cycle when renaming packages

diff --git a/Engine/Source/Runtime/CoreUObject/Private/UObject/Obj.cpp b/Engine/Source/Runtime/CoreUObject/Private/UObject/Obj.cpp
--- a/Engine/Source/Runtime/CoreUObject/Private/UObject/Obj.cpp
+++ b/Engine/Source/Runtime/CoreUObject/Private/UObject/Obj.cpp
@@ -323,7 +323,16 @@ bool UObject::Rename(const TCHAR* InName, UObject* NewOuter, ERenameFlags Flags)
 	// Users who don't want such reload behaviour for the export should be calling InvalidateExport() on the object before renaming.
 	else if (!bIsTest && ((Flags & (REN_AllowPackageLinkerMismatch | REN_DontCreateRedirectors)) == REN_DontCreateRedirectors) && HasAnyFlags(RF_WasLoaded))
 	{
-		if (FLinkerLoad* Linker = GetLinker())
+		if (UPackage* Package = Cast<UPackage>(this))
+		{
+			// When renaming a package, we must ensure that the LinkerAnnotations for each of its objects are detached to avoid "Detaching from existing linker" errors on 
+			// objects loaded later on from the same package, because there is a small window of possible Object Index reuse until the next GC cycle.
+			if (FLinkerLoad* Linker = Package->GetLinker()) // Note: invoking UPackage::GetLinker, not UObjectBaseUtility::GetLinker
+			{
+				Linker->DetachExports();
+			}
+		}
+		else if (FLinkerLoad* Linker = GetLinker())
 		{
 			const int32 LinkerIndex = GetLinkerIndex();
 			if (LinkerIndex != INDEX_NONE)
