UE: Temporary code to track down progress callback crash

diff --git a/Engine/Source/Runtime/CoreUObject/Private/Serialization/AsyncLoading2.cpp b/Engine/Source/Runtime/CoreUObject/Private/Serialization/AsyncLoading2.cpp
--- a/Engine/Source/Runtime/CoreUObject/Private/Serialization/AsyncLoading2.cpp
+++ b/Engine/Source/Runtime/CoreUObject/Private/Serialization/AsyncLoading2.cpp
@@ -4393,6 +4393,7 @@ private:
 	EAsyncLoadingResult::Type	LoadStatus = EAsyncLoadingResult::Succeeded;
 	/** True if this package was created by this async package */
 	bool						bCreatedLinkerRoot = false;
+	bool						bProcessingProgressCallbacks = false;
 
 	UE::FMutex ConstructedObjectsMutex;
 	/** List of ConstructedObjects = Exports + UPackage + ObjectsCreatedFromExports */
@@ -8459,6 +8460,7 @@ void FAsyncPackage2::CallProgressCallbacks(EAsyncLoadingProgress ProgressType)
 {
 	if (ProgressCallbacks.Num() != 0)
 	{
+		bProcessingProgressCallbacks = true;
 		TRACE_CPUPROFILER_EVENT_SCOPE(FAsyncPackage2::CallProgressCallbacks);
 
 		FLoadPackageAsyncProgressParams Params
@@ -8472,6 +8474,7 @@ void FAsyncPackage2::CallProgressCallbacks(EAsyncLoadingProgress ProgressType)
 		{
 			ProgressCallback->Invoke(Params);
 		}
+		bProcessingProgressCallbacks = false;
 	}
 }
 
@@ -12517,6 +12520,7 @@ void FAsyncPackage2::AddCompletionCallback(TUniquePtr<FLoadPackageAsyncDelegate>
 
 void FAsyncPackage2::AddProgressCallback(TSharedPtr<FLoadPackageAsyncProgressDelegate> Callback)
 {
+	ensure(!bProcessingProgressCallbacks);
 	ProgressCallbacks.Emplace(MoveTemp(Callback));
 }
 
