Reset borrowed objects one-by-one to avoid issues with archetype <-> instance lookups.Don't create new instances of objects when resetting existing ones.

#Synced-CL 52058383

diff --git a/Engine/Source/Runtime/CoreUObject/Private/UObject/RemoteObjectSerialization.cpp b/Engine/Source/Runtime/CoreUObject/Private/UObject/RemoteObjectSerialization.cpp
--- a/Engine/Source/Runtime/CoreUObject/Private/UObject/RemoteObjectSerialization.cpp
+++ b/Engine/Source/Runtime/CoreUObject/Private/UObject/RemoteObjectSerialization.cpp
@@ -518,7 +518,7 @@ protected:
 	TMap<UObject*, FNameIndexType> PathNameMap;
 	TMap<FRemoteObjectId, FNameIndexType> RemoteIdMap;
 	TSet<UObject*>* ReferencedObjectsSet = nullptr;
-
+	ERemoteObjectSerializationFlags RemoteSerializationFlags = ERemoteObjectSerializationFlags::None;
 public:
 
 	static FRemoteObjectReferenceInfo GetReferenceInfo(const FObjectPtr& ObjPtr, UObject* InRootObject)
@@ -635,7 +635,7 @@ protected:
 				}
 			}
 
-			if (Info.Type == ERemoteReferenceType::IdOnly && Info.bIsSubobject)
+			if (Info.Type == ERemoteReferenceType::IdOnly && Info.bIsSubobject && !EnumHasAnyFlags(RemoteSerializationFlags, ERemoteObjectSerializationFlags::Resetting))
 			{
 				// Add subobjects of the root object to the list of objects to serialize
 				ensureMsgf(Object->IsIn(RootObject), TEXT("We're about to serialize a subobject %s which is not a subobject of a root object %s"), *Object->GetPathName(), *RootObject->GetPathName());
@@ -646,8 +646,14 @@ protected:
 
 public:
 
-	FArchiveRemoteObjectWriter(UObject* InRootObject, FRemoteObjectData& OutObjectData, const FUObjectMigrationContext* InMigrationContext, const TCHAR* ArchiveName = nullptr, TSet<UObject*>* OutReferencedObjectsSet = nullptr)
+	FArchiveRemoteObjectWriter(
+		UObject* InRootObject, 
+		FRemoteObjectData& OutObjectData, 
+		const FUObjectMigrationContext* InMigrationContext, 
+		const TCHAR* ArchiveName = nullptr, TSet<UObject*>* OutReferencedObjectsSet = nullptr,
+		ERemoteObjectSerializationFlags Flags = ERemoteObjectSerializationFlags::None)
 		: TArchiveRemoteObjectBase<FMemoryWriter>(OutObjectData, InMigrationContext, ArchiveName ? ArchiveName : TEXT("RemoteObjectWriter"))
+		, RemoteSerializationFlags(Flags)
 	{
 		static_assert(sizeof(FNameIndexType) == sizeof(FRemoteObjectTables::FNameIndexType), "Name Index type must match the type used for storing remote object path name indices");
 
@@ -816,16 +822,13 @@ using FRemoteObjectReaderScope = TRemoteObjectArchiveScope<FMemoryReader>;
 */
 class FArchetypeDeltaWriter : public FArchiveRemoteObjectWriter
 {
-	const TMap<UObject*, const UObject*>& InstanceToArchetypeMap;
-	const TMap<const UObject*, UObject*>& ArchetypeToInstanceMap;
 	UObject* InstanceScope = nullptr;
+	const UObject* ArchetypeScope = nullptr;
 
 public:
 
-	FArchetypeDeltaWriter(UObject* InRootObject, FRemoteObjectData& OutObjectData, const TMap<UObject*, const UObject*>& InInstanceToArchetypeMap, const TMap<const UObject*, UObject*>& InArchetypeToInstanceMap)
-		: FArchiveRemoteObjectWriter(InRootObject, OutObjectData, /*MigrationContext*/ nullptr, TEXT("RemoteArchetypeDeltaWriter"))
-		, InstanceToArchetypeMap(InInstanceToArchetypeMap)
-		, ArchetypeToInstanceMap(InArchetypeToInstanceMap)
+	FArchetypeDeltaWriter(UObject* InRootObject, FRemoteObjectData& OutObjectData)
+		: FArchiveRemoteObjectWriter(InRootObject, OutObjectData, /*MigrationContext*/ nullptr, TEXT("RemoteArchetypeDeltaWriter"), nullptr, ERemoteObjectSerializationFlags::Resetting)
 	{
 	}
 
@@ -837,30 +840,30 @@ public:
 		// this way we will delta serialize the difference between the archetype and its instance
 		if (InstanceScope)
 		{
-			// Prefer Instance -> Archetype lookup because it's never ambiguous
-			const UObject* const* ScopeArchetype = InstanceToArchetypeMap.Find(InstanceScope);
-			if (ScopeArchetype && *ScopeArchetype == InArchetype)
-			{
-				return InstanceScope;
-			}
+			checkf(InArchetype == ArchetypeScope, TEXT("Serializing an instance %s without a matching archetype (expected: %s)"), *GetPathNameSafe(InstanceScope), *GetPathNameSafe(ArchetypeScope));
+			return InstanceScope;
 		}
 
-		UE_LOG(LogRemoteSerialization, Warning, TEXT("FArchetypeDeltaWriter::ArchetypeToInstanceMap does not contain an archetype mapping for %s"), *GetPathNameSafe(InArchetype));
+		UE_LOG(LogRemoteSerialization, Warning, TEXT("FArchetypeDeltaWriter does not contain an archetype mapping for %s"), *GetPathNameSafe(InArchetype));
 		return FArchiveRemoteObjectWriter::GetArchetypeFromLoader(InArchetype);
 	}
 
 	virtual void SerializeRemoteObject(UObject* Object, const FRemoteObjectHeader& Header) override
 	{		
 		UObject* ObjectToSerialize = nullptr;
-		if (const UObject* const* Archetype = InstanceToArchetypeMap.Find(Object))
+		ArchetypeScope = nullptr;
+
+		if (const UObject* Archetype = FindArchetype(Object))
 		{
-			ObjectToSerialize = const_cast<UObject*>(*Archetype);			
+			ObjectToSerialize = const_cast<UObject*>(Archetype);
+			ArchetypeScope = Archetype;
 		}
-		checkf(ObjectToSerialize, TEXT("Serializing an instance without matching archetype: %s"), *GetPathNameSafe(Object));
+		checkf(ObjectToSerialize, TEXT("Serializing an instance %s without matching archetype"), *GetPathNameSafe(Object));
 
 		InstanceScope = Object;
 		ObjectToSerialize->Serialize(*this);
 		InstanceScope = nullptr;
+		ArchetypeScope = nullptr;
 	}
 
 protected:
@@ -871,15 +874,15 @@ protected:
 
 		FRemoteObjectReferenceInfo ReplacementInfo(RefInfo);
 
-		if (InstanceScope)
+		if (InstanceScope && ArchetypeScope)
 		{
 			// If we're serializing a reference to an archetype try to replace it with a reference to its instance
 			// This way the produced delta between the archetype and its instance will be correctly pointing to the instances of default subobjects
 			// we can then deserialize over the archetype instance.
-			if (UObject* const* FoundInstance = ArchetypeToInstanceMap.Find(RefInfo.Object))
+			if (RefInfo.Object == ArchetypeScope)
 			{
-				ReplacementInfo.Object = *FoundInstance;
-				ReplacementInfo.Id = FRemoteObjectId(*FoundInstance);
+				ReplacementInfo.Object = InstanceScope;
+				ReplacementInfo.Id = FRemoteObjectId(InstanceScope);
 			}
 		}
 
@@ -1057,99 +1060,6 @@ public:
 	}
 };
 
-/**
-* Helper archive that deserializes the difference between archetypes and their instances
-* FRemoteObjectData passed to this archive is the difference between the archetypes and their instances
-* The data is deserialized on top of existing instances
-* When deserializing references to UObjects, this archive will preserve references to the known existing instances instead of replacing them with refs to archetypes stored in the provided data (see SetObjectReference override)
-*/
-class FArchetypeDeltaReader : public FArchiveRemoteObjectReader
-{
-	const TMap<UObject*, const UObject*>& InstanceToArchetypeMap;
-	const TMap<const UObject*, UObject*>& ArchetypeToInstanceMap;
-	UObject* InstanceScope = nullptr;	
-
-public:
-
-	FArchetypeDeltaReader(FRemoteObjectData& InObjectData, ERemoteObjectSerializationFlags InDeserializeFlags, const TMap<UObject*, const UObject*>& InInstanceToArchetypeMap, const TMap<const UObject*, UObject*>& InArchetypeToInstanceMap)
-		: FArchiveRemoteObjectReader(InObjectData, /*MigrationContext*/ nullptr, InDeserializeFlags, TEXT("RemoteArchetypeDeltaReader"))
-		, InstanceToArchetypeMap(InInstanceToArchetypeMap)
-		, ArchetypeToInstanceMap(InArchetypeToInstanceMap)
-	{
-	}
-
-	virtual void SerializeRemoteObject(UObject* Object)
-	{
-		InstanceScope = Object;
-		Object->Serialize(*this);
-		InstanceScope = nullptr;
-	}
-
-protected:
-
-	virtual void SetObjectReference(FObjectPtr& CurrentValue, FObjectPtr NewValue) override
-	{
-		using namespace UE::CoreUObject::Private;
-		
-		// Helper lambda that extracts a raw object pointer from an FObjectPtr with resolving it.
-		auto ExtractObjectNoResolve = [](FObjectPtr& Ptr)
-			{
-				UObject* Result = nullptr;
-			#if UE_WITH_REMOTE_OBJECT_HANDLE
-				FObjectHandle Handle = Ptr.GetHandle();
-				if (Handle.ExtractHandleType() == ERemoteObjectHandleType::ObjectPointer)
-				{
-					Result = const_cast<UObject*>(Handle.ExtractObjectPointer());
-				}
-			#else
-				Result = Ptr.Get();
-			#endif
-				return Result;
-			};
-
-		UObject* NewObjectValue = ExtractObjectNoResolve(NewValue);
-		if (InstanceScope && NewObjectValue && NewObjectValue->IsTemplate())
-		{
-			// The reason why we use CurrentValueRemoteId to determine what the current value is that
-			// we don't really care about the current value pointer because only care if it's null (!IsValid())
-			// or not in which case we want to preserve the current value because it already points to an instance of its archetype
-			// which will be reset in its own Serialize call.
-			FRemoteObjectId CurrentValueRemoteId = CurrentValue.GetRemoteId();
-
-			// We're deserializing the delta on top of an existing instance and we're assigning a reference to a template object to a property in that instance
-			// The template reference is just a default subobject of the root object's archetype.
-			// There's a few options what the current reference (CurrentValue) is pointing to:			
-			if (CurrentValueRemoteId.IsValid())
-			{
-				// We have a valid, non-template (sub)object so preserve it. It will be reset in its own Serialize() call
-#if DO_CHECK
-				if (UObject* CurrentObjectValue = ExtractObjectNoResolve(CurrentValue))
-				{
-					checkf(!CurrentObjectValue->IsTemplate(), TEXT("Current value %s is a template when deserializing archetype delta. It should be instanced."), *CurrentObjectValue->GetPathName());
-				}
-#endif				
-				NewValue = CurrentValue;
-			}
-			else if (UObject* const* FoundInstance = ArchetypeToInstanceMap.Find(NewObjectValue))
-			{
-				// The current value is null (the value may have been nulled after construction by gameplay code) 
-				// but we have a unique mapping to an existing instance of the archetype reference so assign it to the current value
-				NewValue = *FoundInstance;
-			}
-			else
-			{
-				// The current value is null and we can't map the archetype to any of the known instances. This may be because
-				// the archetype -> value mapping was not unique (e.g. items in an array)
-				FStaticConstructObjectParameters Params(NewObjectValue->GetClass());
-				Params.Outer = InstanceScope;
-				Params.Template = NewObjectValue;
-				NewValue = StaticConstructObject_Internal(Params);
-			}
-		}
-		FArchiveRemoteObjectReader::SetObjectReference(CurrentValue, NewValue);
-	}
-};
-
 /**
 * Archive that replaces unresolved FObjectPtrs (TObjectPtrs<>) with actual pointers to deserialized objects
 */
@@ -1523,76 +1433,37 @@ int32 DeserializeObjectDataInternal(FArchiveRemoteObjectReader& Ar,
 	TArray<UObject*>& OutObjects,
 	ERemoteObjectSerializationFlags DeserializeFlags);
 
-void ResetRemoteObject(UObject* InObject)
+void ResetRemoteObjects(const TArray<UObject*>& ObjectHierarchy)
 {
-	UObject* Object = InObject;
-	if (GRemoteObjectsMigrateFullHierarchy)
-	{
-		Object = FindCanonicalRootObjectForSerialization(Object);
-	}
-
-	TArray<UObject*> ObjectHierarchy;
-	TMap<UObject*, const UObject*> ObjectToArchetypeMap; // Maps unique instances to archetypes
-	TMap<const UObject*, UObject*> ArchetypeToObjectMap; // Maps unique archetypes to their instances
-
-	// Gather all subobjects of the object to reset and map them to their respective archetypes
-	// This map will be used to replace serialized objects with the mapped archetypes (see FArchetypeDeltaWriter::SerializeRemoteObject())
-	// and to preserve references to instances when deserializing the delta between the archetypes and their respective instances (see FArchetypeDeltaReader::SetObjectReference())
+	for (UObject* ObjectToReset : ObjectHierarchy)
 	{
-		ObjectHierarchy.Add(Object);
-		GetObjectsWithOuter(Object, ObjectHierarchy, /*bIncludeNestedObjects*/ true);
-
-		TSet<const UObject*> NonUniqueArchetypes;
-		for (UObject* ObjectToReset : ObjectHierarchy)
-		{
-			const UObject* FoundArchetype = FindArchetype(ObjectToReset);
-			ObjectToArchetypeMap.Add(ObjectToReset, FoundArchetype);
-			if (ArchetypeToObjectMap.Find(FoundArchetype))
-			{
-				NonUniqueArchetypes.Add(FoundArchetype);
-			}
-			else
-			{
-				ArchetypeToObjectMap.Add(FoundArchetype, ObjectToReset);
-			}
-		}
-		// Remove all archetypes that don't have unique mapping to their respective instances
-		// FArchetypeDeltaReader will create new instances of archetypes with non-unique mapping
-		for (const UObject* NotUniqueArchetype : NonUniqueArchetypes)
+		// This code is effectively copying the delta between the archetype and its instance (ObjectToReset) to the instance so that the instance (ObjectToReset) is identical to the archetype.
+		// At the end of this operation ObjectToReset will be referencing the archetype's subobjects but that's ok because references to the instances of subobjects
+		// will be restored when deserializing the migrated data from another server.
+		FRemoteObjectData ArchetypeDelta;
 		{
-			ArchetypeToObjectMap.Remove(NotUniqueArchetype);
+			FRemoteObjectId RequestedObjectId(ObjectToReset);
+			FArchetypeDeltaWriter Ar(ObjectToReset, ArchetypeDelta);
+
+			// Even though we pass the Object to reset to SerializeObjectDataInternal, the FArchetypeDeltaWriter will actually serialize the differences between the archetypes' data and their respective instances
+			// This will produce a delta between the archetypes and their instances which will then be used to restore the instances' state to the archetypes' default state
+			TSet<UObject*> SerializedObjects;
+			SerializeObjectDataInternal(Ar, ObjectToReset, RequestedObjectId, SerializedObjects, ERemoteObjectSerializationFlags::Resetting | ERemoteObjectSerializationFlags::SkipCanonicalRootSearch);
 		}
-		
-	}
 
-	FRemoteObjectData* ArchetypeDelta = new FRemoteObjectData();
-	{
-		FRemoteObjectId RequestedObjectId(Object);		
-		FArchetypeDeltaWriter Ar(Object, *ArchetypeDelta, ObjectToArchetypeMap, ArchetypeToObjectMap);
-		for (UObject* ObjectToSerialize : ObjectHierarchy)
 		{
-			Ar.GetObjectsToSerialize().Add(ObjectToSerialize);
+			// Deserialize archetype data on top of the object to reset to restore its state to the archetype values
+			TArray<UObject*> DeserializedObjects;
+			TArray<FRemoteObjectId> DeserializedIds;
+			// When deserializing archetype delta we want to:
+			// Preserve any references to remote objects that have not been migrated yet (this is because we can't generate archetype delta for them because they don't exist on this server and we don't have their data)
+			// Additionally we don't want to recursively re-enter this function so let the deserialization process know we're already resetting migrated object(s)
+			ERemoteObjectSerializationFlags DeserializationFlags = ERemoteObjectSerializationFlags::PreserveRemoteReferences | ERemoteObjectSerializationFlags::UseExistingObjects | ERemoteObjectSerializationFlags::Resetting;
+			FArchiveRemoteObjectReader Ar(ArchetypeDelta, /*MigrationContext*/ nullptr, DeserializationFlags, TEXT("RemoteArchetypeDeltaReader"));
+			DeserializeObjectDataInternal(Ar, ArchetypeDelta, /*MigrationContext*/ nullptr, DeserializedIds, DeserializedObjects, DeserializationFlags);
+			check(DeserializedObjects.Num() == 1 && DeserializedObjects[0] ==ObjectToReset);
 		}
-
-		// Even though we pass the Object to reset to SerializeObjectDataInternal, the FArchetypeDeltaWriter will actually serialize the differences between the archetypes' data and their respective instances
-		// This will produce a delta between the archetypes and their instances which will then be used to restore the instances' state to the archetypes' default state
-		TSet<UObject*> SerializedObjects;
-		SerializeObjectDataInternal(Ar, Object, RequestedObjectId, SerializedObjects, ERemoteObjectSerializationFlags::Resetting | ERemoteObjectSerializationFlags::SkipCanonicalRootSearch);
-	}
-
-	{
-		// Deserialize archetype data on top of the object to reset (and its subobjects) to restore its state to the archetype values
-		TArray<UObject*> DeserializedObjects;
-		TArray<FRemoteObjectId> DeserializedIds;
-		// When deserializing archetype delta we want to:
-		// Preserve any references to remote objects that have not been migrated yet (this is because we can't generate archetype delta for them because they don't exist on this server and we don't have their data)
-		// Use any existing objects and deserialize archetype delta on top of them avoiding re-construction which may lead to undesirable side-effects
-		// Additionally we don't want to recursively re-enter this function so let the deserialization process know we're already resetting migrated object(s)
-		ERemoteObjectSerializationFlags DeserializationFlags = ERemoteObjectSerializationFlags::PreserveRemoteReferences | ERemoteObjectSerializationFlags::UseExistingObjects | ERemoteObjectSerializationFlags::Resetting;
-		FArchetypeDeltaReader Ar(*ArchetypeDelta, DeserializationFlags, ObjectToArchetypeMap, ArchetypeToObjectMap);
-		DeserializeObjectDataInternal(Ar, *ArchetypeDelta, /*MigrationContext*/ nullptr, DeserializedIds, DeserializedObjects, DeserializationFlags);
 	}
-	delete ArchetypeDelta;
 }
 
 int32 DeserializeObjectDataInternal(FArchiveRemoteObjectReader& Ar,
@@ -1735,7 +1606,8 @@ int32 DeserializeObjectDataInternal(FArchiveRemoteObjectReader& Ar,
 
 			if (bReturningBorrowedObject)
 			{
-				ResetRemoteObject(OutObjects[0]);
+				checkf(RootObject == OutObjects[0], TEXT("Wrong RootObject: %s, expected: %s"), *GetPathNameSafe(RootObject), *GetPathNameSafe(OutObjects[0]));
+				ResetRemoteObjects(CreatedObjects);
 			}
 		}
 
