Exclude unreachable objects in NotifyCompiledVersePackage

#rnx
#Synced-CL 52137917

diff --git a/Engine/Source/Runtime/CoreUObject/Private/Serialization/AsyncLoading2.cpp b/Engine/Source/Runtime/CoreUObject/Private/Serialization/AsyncLoading2.cpp
--- a/Engine/Source/Runtime/CoreUObject/Private/Serialization/AsyncLoading2.cpp
+++ b/Engine/Source/Runtime/CoreUObject/Private/Serialization/AsyncLoading2.cpp
@@ -4971,7 +4971,7 @@ public:
 
 	virtual void NotifyScriptVersePackage(Verse::VPackage* Package) override;
 
-	void NotifyCompiledVersePackage(Verse::VPackage* VersePackage, UPackage* Package);
+	void NotifyCompiledVersePackage(Verse::VPackage* VersePackage, UPackage* Package, TArrayView<UObject*> Exports);
 
 	virtual void NotifyRegistrationComplete() override;
 
@@ -5798,7 +5798,7 @@ public:
 static Verse::TLazyInitialized<FAsyncLoadingVerseRoot> AsyncLoadingVerseRoot;
 #endif
 
-COREUOBJECT_API TDelegate<void(Verse::VPackage*, UPackage*)> NotifyCompiledVersePackageDelegate;
+COREUOBJECT_API TDelegate<void(Verse::VPackage*, UPackage*, TArrayView<UObject*>)> NotifyCompiledVersePackageDelegate;
 
 /**
  * Updates FUObjectThreadContext with the current package when processing it.
@@ -11914,7 +11914,7 @@ void FAsyncLoadingThread2::NotifyScriptVersePackage(Verse::VPackage* Package)
 #endif
 }
 
-void FAsyncLoadingThread2::NotifyCompiledVersePackage(Verse::VPackage* VersePackage, UPackage* Package)
+void FAsyncLoadingThread2::NotifyCompiledVersePackage(Verse::VPackage* VersePackage, UPackage* Package, TArrayView<UObject*> Exports)
 {
 	check(IsInGameThread() && !IsMultithreaded());
 
@@ -11924,40 +11924,37 @@ void FAsyncLoadingThread2::NotifyCompiledVersePackage(Verse::VPackage* VersePack
 	// Construct an export map to reserve space in the package, rather than inserting objects directly.
 	// StoreGlobalObject does not amortize export map reallocations, so this avoids a quadratic slowdown.
 	TArray<FExportMapEntry> ExportMap;
-	TArray<UObject*> Exports;
-	ForEachObjectWithOuter(Package, [this, &ExportMap, &Exports](UObject* Object)
+	ExportMap.Reserve(Exports.Num());
+	for (UObject* Object : Exports)
 	{
-		if (Object->HasAllFlags(RF_Public))
-		{
-			Object->SetInternalFlags(EInternalObjectFlags::LoaderImport);
-			Object->SetFlags(RF_WasLoaded | RF_LoadCompleted);
+		ensure(Object->HasAllFlags(RF_Public));
+		Object->SetInternalFlags(EInternalObjectFlags::LoaderImport);
+		Object->SetFlags(RF_WasLoaded | RF_LoadCompleted);
 
-			TArray<FName, TInlineAllocator<64>> FullPath;
-			FullPath.Add(Object->GetFName());
-			UObject* Outer = Object->GetOuter();
-			while (Outer)
+		TArray<FName, TInlineAllocator<64>> FullPath;
+		FullPath.Add(Object->GetFName());
+		UObject* Outer = Object->GetOuter();
+		while (Outer)
+		{
+			FullPath.Add(Outer->GetFName());
+			Outer = Outer->GetOuter();
+		}
+		TStringBuilder<256> PackageRelativeExportPath;
+		for (int32 PathIndex = FullPath.Num() - 2; PathIndex >= 0; --PathIndex)
+		{
+			TCHAR NameStr[FName::StringBufferSize];
+			uint32 NameLen = FullPath[PathIndex].ToString(NameStr);
+			for (uint32 I = 0; I < NameLen; ++I)
 			{
-				FullPath.Add(Outer->GetFName());
-				Outer = Outer->GetOuter();
+				NameStr[I] = TChar<TCHAR>::ToLower(NameStr[I]);
 			}
-			TStringBuilder<256> PackageRelativeExportPath;
-			for (int32 PathIndex = FullPath.Num() - 2; PathIndex >= 0; --PathIndex)
-			{
-				TCHAR NameStr[FName::StringBufferSize];
-				uint32 NameLen = FullPath[PathIndex].ToString(NameStr);
-				for (uint32 I = 0; I < NameLen; ++I)
-				{
-					NameStr[I] = TChar<TCHAR>::ToLower(NameStr[I]);
-				}
-				PackageRelativeExportPath.AppendChar('/');
-				PackageRelativeExportPath.Append(FStringView(NameStr, NameLen));
-			}
-			uint64 ExportHash = CityHash64(reinterpret_cast<const char*>(PackageRelativeExportPath.GetData() + 1), (PackageRelativeExportPath.Len() - 1) * sizeof(TCHAR));
-
-			ExportMap.Add({ .PublicExportHash = ExportHash });
-			Exports.Add(Object);
+			PackageRelativeExportPath.AppendChar('/');
+			PackageRelativeExportPath.Append(FStringView(NameStr, NameLen));
 		}
-	}, /* bIncludeNestedObjects*/ true);
+		uint64 ExportHash = CityHash64(reinterpret_cast<const char*>(PackageRelativeExportPath.GetData() + 1), (PackageRelativeExportPath.Len() - 1) * sizeof(TCHAR));
+
+		ExportMap.Add({ .PublicExportHash = ExportHash });
+	}
 
 	{
 		FWriteLockedPackageRef PackageRef = GlobalImportStore.AddPackageRef(PackageId, PackageName, EPackageLoader::Unknown, EPackageExtension::Unspecified);
