Introduced a migration data id (stored as FRemoteObjectId) we can log to track/debug migrations between servers

#Synced-CL 52058383

diff --git a/Engine/Source/Runtime/CoreUObject/Private/UObject/RemoteObjectSerialization.cpp b/Engine/Source/Runtime/CoreUObject/Private/UObject/RemoteObjectSerialization.cpp
--- a/Engine/Source/Runtime/CoreUObject/Private/UObject/RemoteObjectSerialization.cpp
+++ b/Engine/Source/Runtime/CoreUObject/Private/UObject/RemoteObjectSerialization.cpp
@@ -23,6 +23,7 @@
 #include "Misc/Paths.h"
 #include "HAL/FileManager.h"
 #include "Serialization/ArchiveSerializedPropertyChain.h"
+#include <atomic>
 
 int32 GRemoteObjectsMigrateFullHierarchy = 1;
 static FAutoConsoleVariableRef CVarRemoteObjectsMigrateFullHierarchy(
@@ -79,6 +80,25 @@ FArchive& operator<<(FArchive& Ar, FPackedRemoteObjectHeader& Header)
 	return Ar;
 }
 
+namespace UE::RemoteObject::Serialization::Private
+{
+	std::atomic<uint64> MigrationSerialNumber(1);
+
+	FRemoteObjectId GetNextMigrationId()
+	{
+		uint64 SerialNumber = 0;
+		UE_AUTORTFM_OPEN
+		{
+			uint64 CurrentSerialNumber = MigrationSerialNumber.load(std::memory_order_relaxed);
+			do
+			{
+				SerialNumber = (CurrentSerialNumber <= MAX_REMOTE_OBJECT_SERIAL_NUMBER) ? CurrentSerialNumber : 1;
+			} while (!MigrationSerialNumber.compare_exchange_weak(CurrentSerialNumber, SerialNumber + 1, std::memory_order_relaxed));
+		};
+		return FRemoteObjectId(FRemoteServerId::GetLocalServerId(), SerialNumber);		
+	}
+}
+
 namespace UE::RemoteObject::Serialization::Disk
 {
 
@@ -117,6 +137,9 @@ void DeserializeFileHeader(FArchive& Ar, FRemoteObjectFileHeader& OutHeader)
 
 void SerializeNameTables(FArchive& Ar, const FRemoteObjectData& InObjectData)
 {
+	FRemoteObjectId MigrationId = InObjectData.MigrationId;
+	Ar << MigrationId;
+
 	FNameIndexType NumNames = IntCastChecked<FNameIndexType>(InObjectData.Tables.Names.Num());
 	Ar << NumNames;
 	for (const FName& Name : InObjectData.Tables.Names)
@@ -131,6 +154,7 @@ void SerializeNameTables(FArchive& Ar, const FRemoteObjectData& InObjectData)
 
 void DeserializeNameTables(FArchive& Ar, FRemoteObjectData& OutObjectData)
 {
+	Ar << OutObjectData.MigrationId;
 	OutObjectData.Tables.Names.Reset();
 	FNameIndexType NumNames = 0;
 	Ar << NumNames;
@@ -1400,6 +1424,7 @@ FRemoteObjectData SerializeObjectData(UObject* InObject, TSet<UObject*>& OutObje
 	}
 
 	FRemoteObjectData ObjectData;
+	ObjectData.MigrationId = UE::RemoteObject::Serialization::Private::GetNextMigrationId();
 	FRemoteObjectId RequestedObjectId(InObject);
 
 	{
diff --git a/Engine/Source/Runtime/CoreUObject/Private/UObject/RemoteObjectTransfer.cpp b/Engine/Source/Runtime/CoreUObject/Private/UObject/RemoteObjectTransfer.cpp
--- a/Engine/Source/Runtime/CoreUObject/Private/UObject/RemoteObjectTransfer.cpp
+++ b/Engine/Source/Runtime/CoreUObject/Private/UObject/RemoteObjectTransfer.cpp
@@ -899,6 +899,12 @@ public:
 	void SendRemoteObject(const FMigrateSendParams& Params)
 	{
 		check(!AutoRTFM::IsClosed());
+
+		UE_LOG(LogRemoteObject, Verbose, TEXT("FRemoteObjectTransferQueue(%s) SendRemoteObject %s (mid: %s)"),
+			*FRemoteServerId::GetLocalServerId().ToString(),
+			*Params.MigrationContext.ObjectId.ToString(),
+			*Params.ObjectData.MigrationId.ToString(ERemoteIdToStringVerbosity::Id));
+
 		RemoteObjectTransferDelegate.Execute(Params);
 	}
 
@@ -930,9 +936,12 @@ public:
 
 		const bool bFoundRequest = (FoundPendingRequestIndex != INDEX_NONE);
 
-		UE_LOG(LogRemoteObject, VeryVerbose, TEXT("FRemoteObjectTransferQueue(%s): FulfillReceiveRequest %s (was requested: %s, owner server id: %s, transferring ownership: %s)"),
+		checkf(ObjectData.MigrationId.IsValid(), TEXT("Received RemoteObjectData with an invalid migration id"));
+
+		UE_LOG(LogRemoteObject, Verbose, TEXT("FRemoteObjectTransferQueue(%s): FulfillReceiveRequest %s (mid: %s, was requested: %s, owner server id: %s, transferring ownership: %s)"),
 			*FRemoteServerId::GetLocalServerId().ToString(),
-			*MigrationContext.ObjectId.ToString(), 
+			*MigrationContext.ObjectId.ToString(),
+			*ObjectData.MigrationId.ToString(ERemoteIdToStringVerbosity::Id),
 			bFoundRequest ? TEXT("yes") : TEXT("no"),
 			*MigrationContext.OwnerServerId.ToString(),
 			(MigrationContext.GetObjectMigrationRecvType(MigrationContext.ObjectId) == EObjectMigrationRecvType::AssignedOwnership) ? TEXT("yes") : TEXT("no"));
@@ -941,20 +950,27 @@ public:
 		if (FRemoteObjectId::RemoteObjectSupportCompiledIn && GTransactionalMigration)
 		{
 			TRefCountPtr<FTransactionalMigrationData> StoredData = StoreMigratedObjectData(ObjectData, MigrationContext);
-	
-			if (!GTransactionalDeserialization)
+			if (StoredData)
 			{
-				DeserializeObjectsFromMigratedData(*StoredData);
-				// Deserialization sets the residence of the deserialized objects to Local as soon as they are created to avoid re-resolving the deserialized objects when accessing them as outers
-				// or in Serialize() function overrides. Since we queue the PostMigrate step to be transactional we need to mark the objects as LocalNotReady again so that if something else attempts
-				// to access the deserialized objects BEFORE the transactional PostMigrate runs the engine will attempt to resolve (PostMigrate) the deserialized objects
-				StoredData->SetReceivedObjectsResidence(EResidence::LocalNotReady);
-			}
-			static FName TransactionalWorkName("TransactionallyMigrateObjects");
-			UE::RemoteExecutor::EnqueueMigrationWork(TransactionalWorkName, [StoredData]()
+				if (!GTransactionalDeserialization)
 				{
-					TransactionallyMigrateObjects(StoredData);
-				});
+					DeserializeObjectsFromMigratedData(*StoredData);
+					// Deserialization sets the residence of the deserialized objects to Local as soon as they are created to avoid re-resolving the deserialized objects when accessing them as outers
+					// or in Serialize() function overrides. Since we queue the PostMigrate step to be transactional we need to mark the objects as LocalNotReady again so that if something else attempts
+					// to access the deserialized objects BEFORE the transactional PostMigrate runs the engine will attempt to resolve (PostMigrate) the deserialized objects
+					StoredData->SetReceivedObjectsResidence(EResidence::LocalNotReady);
+				}
+				static FName TransactionalWorkName("TransactionallyMigrateObjects");
+				UE::RemoteExecutor::EnqueueMigrationWork(TransactionalWorkName, [StoredData]()
+					{
+						TransactionallyMigrateObjects(StoredData);
+					});
+			}
+			else
+			{
+				// We received a dupicate data so exit early without removing the request (this is a re-entry and the request will be removed from the function that migrated this data first)
+				return;
+			}
 		}
 		else
 #endif
@@ -1014,18 +1030,29 @@ TRefCountPtr<FTransactionalMigrationData> StoreMigratedObjectData(FRemoteObjectD
 		FRemoteObjectId SerializedId = It.GetId();
 		FRemoteObjectStub* Stub = FindOrAddRemoteObjectStub(SerializedId, FRemoteServerId::GetLocalServerId());
 
-		checkf(!Stub->MigratedData, TEXT("Received remote object %s (%s) from server %s bundled with %s (%s) but the latter object's data already exists locally (received from server %s and bundled with %s %s, state: %s, data %s identical)"),
-			*MigratedData->Context.ObjectId.ToString(),
-			*MigratedData->Data.GetName(MigratedData->Context.ObjectId).ToString(),
-			*MigratedData->Context.RemoteServerId.ToString(),
-			*SerializedId.ToString(),
-			*MigratedData->Data.GetName(SerializedId).ToString(),
-			*Stub->MigratedData->Context.RemoteServerId.ToString(),
-			*Stub->MigratedData->Context.ObjectId.ToString(),
-			*Stub->MigratedData->Data.GetName(Stub->MigratedData->Context.ObjectId).ToString(),
-			EnumToString(Stub->MigratedData->State),
-			Stub->MigratedData->Data.Compare(MigratedData->Data) ? TEXT("is") : TEXT("is NOT")
-		);
+		if (Stub->MigratedData)
+		{			
+			// This can happen on re-entry from within the same callstack as the first store in which case both data structures should contain identical migrated object data
+			UE_CLOGF(Stub->MigratedData->Data.MigrationId != MigratedData->Data.MigrationId, LogRemoteObject, Fatal, 
+				"Received remote object %ls (%ls) from server %ls in mid: %ls bundled with %ls (%ls) but the latter object's data already exists locally (received from server %ls and bundled with %ls %ls, state: %ls, mid: %ls)",
+				*MigratedData->Context.ObjectId.ToString(),
+				*MigratedData->Data.GetName(MigratedData->Context.ObjectId).ToString(),
+				*MigratedData->Context.RemoteServerId.ToString(),
+				*MigratedData->Data.MigrationId.ToString(ERemoteIdToStringVerbosity::Id),
+				*SerializedId.ToString(),
+				*MigratedData->Data.GetName(SerializedId).ToString(),
+				*Stub->MigratedData->Context.RemoteServerId.ToString(),
+				*Stub->MigratedData->Context.ObjectId.ToString(),
+				*Stub->MigratedData->Data.GetName(Stub->MigratedData->Context.ObjectId).ToString(),
+				EnumToString(Stub->MigratedData->State),
+				*Stub->MigratedData->Data.MigrationId.ToString(ERemoteIdToStringVerbosity::Id)
+			);
+		
+			UE_LOGF(LogRemoteObject, Warning, "Received duplicate remote object data (mid: %ls) for the same object (%ls bundled with %ls).",
+				*MigratedData->Data.MigrationId.ToString(ERemoteIdToStringVerbosity::Id), *SerializedId.ToString(), *MigratedData->Context.ObjectId.ToString());
+
+			return nullptr;
+		}		
 
 		Stub->OwningServerId = MigratedData->Context.OwnerServerId;
 		Stub->MigratedData = MigratedData;
diff --git a/Engine/Source/Runtime/CoreUObject/Public/UObject/NoExportTypes.h b/Engine/Source/Runtime/CoreUObject/Public/UObject/NoExportTypes.h
--- a/Engine/Source/Runtime/CoreUObject/Public/UObject/NoExportTypes.h
+++ b/Engine/Source/Runtime/CoreUObject/Public/UObject/NoExportTypes.h
@@ -2868,6 +2868,9 @@ struct FRemoteObjectData
 {
 	GENERATED_BODY()
 
+	UPROPERTY()
+	FRemoteObjectId MigrationId;
+
 	UPROPERTY()
 	FRemoteObjectTables Tables;
 
diff --git a/Engine/Source/Runtime/CoreUObject/Public/UObject/RemoteObjectTransfer.h b/Engine/Source/Runtime/CoreUObject/Public/UObject/RemoteObjectTransfer.h
--- a/Engine/Source/Runtime/CoreUObject/Public/UObject/RemoteObjectTransfer.h
+++ b/Engine/Source/Runtime/CoreUObject/Public/UObject/RemoteObjectTransfer.h
@@ -90,6 +90,8 @@ public:
 */
 struct FRemoteObjectData
 {
+	/** Contains a unique id of the migration this data was created for. ServerId represents the sending side */
+	FRemoteObjectId MigrationId;
 	/** Tables containing lists of unique remote ids and FNames */
 	FRemoteObjectTables Tables;
 	/** List of unique pathnames (stored as a list of indices of unique FNames) */
