Assign a new data id for remote object data loaded from database

#Synced-CL 52058383

diff --git a/Engine/Source/Runtime/CoreUObject/Private/UObject/RemoteObjectSerialization.cpp b/Engine/Source/Runtime/CoreUObject/Private/UObject/RemoteObjectSerialization.cpp
--- a/Engine/Source/Runtime/CoreUObject/Private/UObject/RemoteObjectSerialization.cpp
+++ b/Engine/Source/Runtime/CoreUObject/Private/UObject/RemoteObjectSerialization.cpp
@@ -82,10 +82,16 @@ FArchive& operator<<(FArchive& Ar, FPackedRemoteObjectHeader& Header)
 
 namespace UE::RemoteObject::Serialization::Private
 {
+	/*
+	* The migration serial number is shared between 'Local' server id and 'Database'.
+	* The reason for that is that we can't have a globally unique database serial number as it's stored locally
+	* and effectively it's the current server that restores the object from the database.
+	*/
 	std::atomic<uint64> MigrationSerialNumber(1);
 
-	FRemoteObjectId GetNextMigrationId()
+	FRemoteObjectId GetNextMigrationId(FRemoteServerId ServerId = FRemoteServerId::GetLocalServerId())
 	{
+		checkf(ServerId.IsLocal() || ServerId.IsDatabase(), TEXT("Migration id can only be generated for the Local server or Database, requested for %s"), *ServerId.ToString());
 		uint64 SerialNumber = 0;
 		UE_AUTORTFM_OPEN
 		{
@@ -95,7 +101,7 @@ namespace UE::RemoteObject::Serialization::Private
 				SerialNumber = (CurrentSerialNumber <= MAX_REMOTE_OBJECT_SERIAL_NUMBER) ? CurrentSerialNumber : 1;
 			} while (!MigrationSerialNumber.compare_exchange_weak(CurrentSerialNumber, SerialNumber + 1, std::memory_order_relaxed));
 		};
-		return FRemoteObjectId(FRemoteServerId::GetLocalServerId(), SerialNumber);		
+		return FRemoteObjectId(ServerId, SerialNumber);
 	}
 }
 
@@ -205,6 +211,10 @@ void LoadObjectFromDisk(const FUObjectMigrationContext& MigrationContext)
 	FRemoteObjectData ObjectData;
 	LoadObjectFromDiskInternal(MigrationContext, MigrationContext.ObjectId, ObjectData, ERemoteObjectFileType::Unknown);
 
+	// Loading the object from disk will restore the MigrationId the data was serialized with but to make things consistent (and not confusing as to where the data came from) 
+	// we want a new migration id with the server id set to Database
+	ObjectData.MigrationId = UE::RemoteObject::Serialization::Private::GetNextMigrationId(FRemoteServerId(ERemoteServerIdConstants::Database));
+
 	FRemoteObjectId RootObjectId(ObjectData.GetRootSerializedObjectId());
 	for (FRemoteObjectId SerializedObjectId : ObjectData)
 	{
