[ZenLoader]
- Fix potential race and missing progress calls when registering and processing callbacks in different threads

#Synced-CL 52203528

diff --git a/Engine/Source/Runtime/CoreUObject/Private/Serialization/AsyncLoading2.cpp b/Engine/Source/Runtime/CoreUObject/Private/Serialization/AsyncLoading2.cpp
--- a/Engine/Source/Runtime/CoreUObject/Private/Serialization/AsyncLoading2.cpp
+++ b/Engine/Source/Runtime/CoreUObject/Private/Serialization/AsyncLoading2.cpp
@@ -4052,6 +4052,71 @@ private:
 };
 #endif // ALT2_ENABLE_NEW_ARCHIVE_FOR_LINKERLOAD
 
+/**
+* Node in a lock-free intrusive linked list of progress callbacks.
+* Each entry has its own InvocationLock to serialize callback invocations,
+* ensuring in-order delivery even when multiple threads call CallProgressCallbacks concurrently.
+* The Next pointer is immutable after the node is pushed onto the list.
+*/
+struct FProgressCallbackEntry
+{
+	TSharedPtr<FLoadPackageAsyncProgressDelegate> Delegate;
+	UE::FMutex InvocationLock;
+	EAsyncLoadingProgress LastNotifiedProgress = EAsyncLoadingProgress::Failed; // Failed = sentinel for "nothing notified yet"
+	std::atomic<FProgressCallbackEntry*> Next = nullptr;
+};
+	
+// Fire events up to the current progress cooperatively between threads in case some were missed (i.e. late load request registration).
+void CallProgressCallback(FProgressCallbackEntry& Entry, FName PackageName, UPackage* LoadedPackage, EAsyncLoadingProgress CurrentProgress)
+{
+	// Progress states in dispatch order
+	static constexpr EAsyncLoadingProgress ProgressOrder[] = {
+		EAsyncLoadingProgress::Started,
+		EAsyncLoadingProgress::Read,
+		EAsyncLoadingProgress::Serialized,
+		EAsyncLoadingProgress::FullyLoaded
+	};
+
+	// Per-entry lock serializes invocations for this callback, ensuring
+	// in-order delivery even when multiple threads call concurrently.
+	UE::TUniqueLock Lock(Entry.InvocationLock);
+	if (CurrentProgress > Entry.LastNotifiedProgress)
+	{
+		for (EAsyncLoadingProgress Progress : ProgressOrder)
+		{
+			if (Progress > Entry.LastNotifiedProgress && Progress <= CurrentProgress)
+			{
+				FLoadPackageAsyncProgressParams Params
+				{
+					.PackageName = PackageName,
+					.LoadedPackage = LoadedPackage,
+					.ProgressType = Progress
+				};
+				Entry.Delegate->Invoke(Params);
+			}
+		}
+		Entry.LastNotifiedProgress = CurrentProgress;
+	}
+}
+
+EAsyncLoadingProgress GetLoadingProgressFromState(EAsyncPackageLoadingState2 AsyncPackageLoadingState)
+{
+	if (AsyncPackageLoadingState >= EAsyncPackageLoadingState2::ExportsDone)
+	{
+		return EAsyncLoadingProgress::Serialized;
+	}
+	if (AsyncPackageLoadingState >= EAsyncPackageLoadingState2::ProcessPackageSummary)
+	{
+		return EAsyncLoadingProgress::Read;
+	}
+	if (AsyncPackageLoadingState >= EAsyncPackageLoadingState2::WaitingForIo)
+	{
+		return EAsyncLoadingProgress::Started;
+	}
+
+	return EAsyncLoadingProgress::Failed;
+}
+
 /**
 * Structure containing intermediate data required for async loading of all exports of a package.
 */
@@ -4412,8 +4477,8 @@ private:
 	FRequestIDs ImportRequestIDs;
 	/** Callbacks called when we finished loading this package */
 	TArray<TUniquePtr<FLoadPackageAsyncDelegate>, TInlineAllocator<2>> CompletionCallbacks;
-	/** Callbacks called for the different loading phase of this package */
-	TArray<TSharedPtr<FLoadPackageAsyncProgressDelegate>, TInlineAllocator<2>> ProgressCallbacks;
+	/** Lock-free intrusive linked list of progress callbacks. @see FProgressCallbackEntry */
+	std::atomic<FProgressCallbackEntry*> ProgressCallbacksHead = nullptr;
 
 public:
 
@@ -4463,8 +4528,7 @@ public:
 
 	/** [EDL] End Event driven loader specific stuff */
 
-	void CallProgressCallback(FLoadPackageAsyncProgressDelegate& ProgressCallback, EAsyncLoadingProgress ProgressType);
-	void CallProgressCallbacks(EAsyncLoadingProgress ProgressType);
+	void CallProgressCallbacks();
 private:
 	void InitializeExportArchive(FExportArchive& Ar, bool bIsOptionalSegment);
 	void CreatePackageNodes(const FAsyncLoadEventSpec* EventSpecs);
@@ -4727,7 +4791,7 @@ private:
 		FAsyncPackage2* AsyncPackage = nullptr;
 		FAsyncPackage2::FRequestIDs RequestIDs;
 		TArray<TUniquePtr<FLoadPackageAsyncDelegate>, TInlineAllocator<2>> CompletionCallbacks;
-		TArray<TSharedPtr<FLoadPackageAsyncProgressDelegate>, TInlineAllocator<2>> ProgressCallbacks;
+		FProgressCallbackEntry* ProgressCallbacksHead = nullptr;
 		FAsyncPackage2::FRequestIDs ImportRequestIDs;
 
 		static FCompletedPackageRequest FromUnreslovedPackage(
@@ -4756,7 +4820,7 @@ private:
 				.AsyncPackage = Package,
 				.RequestIDs = Package->RequestIDs, // must be a copy as this array is accessed after the package marked as completed
 				.CompletionCallbacks = MoveTemp(Package->CompletionCallbacks),
-				.ProgressCallbacks = MoveTemp(Package->ProgressCallbacks),
+				.ProgressCallbacksHead = Package->ProgressCallbacksHead.exchange(nullptr, std::memory_order_acquire),
 				.ImportRequestIDs = MoveTemp(Package->ImportRequestIDs)
 			};
 
@@ -4781,24 +4845,32 @@ private:
 				CompletionCallbacks.Empty();
 			}
 
-			if (ProgressCallbacks.Num() != 0)
+			if (ProgressCallbacksHead)
 			{
 				TRACE_CPUPROFILER_EVENT_SCOPE(PackageProgressCallbacks_Completion);
 
-				const EAsyncLoadingProgress ProgressType = Result == EAsyncLoadingResult::Succeeded ? EAsyncLoadingProgress::FullyLoaded : EAsyncLoadingProgress::Failed;
-
-				for (TSharedPtr<FLoadPackageAsyncProgressDelegate>& ProgressCallback : ProgressCallbacks)
+				FProgressCallbackEntry* Entry = ProgressCallbacksHead;
+				while (Entry)
 				{
-					FLoadPackageAsyncProgressParams Params
+					if (Result == EAsyncLoadingResult::Succeeded)
 					{
-						.PackageName = PackageName,
-						.LoadedPackage = UPackage,
-						.ProgressType = ProgressType
-					};
+						CallProgressCallback(*Entry, PackageName, UPackage, EAsyncLoadingProgress::FullyLoaded);
+					}
+					else
+					{
+						FLoadPackageAsyncProgressParams Params
+						{
+							.PackageName = PackageName,
+							.LoadedPackage = UPackage,
+							.ProgressType = EAsyncLoadingProgress::Failed
+						};
+						Entry->Delegate->Invoke(Params);
+					}
 
-					ProgressCallback->Invoke(Params);
+					FProgressCallbackEntry* Next = Entry->Next;
+					delete Entry;
+					Entry = Next;
 				}
-				ProgressCallbacks.Empty();
 			}
 		}
 	};
@@ -6030,23 +6102,6 @@ FAsyncPackage2* FAsyncLoadingThread2::FindOrInsertPackage(FAsyncLoadingThreadSta
 
 				if (PackageProgressDelegate.IsValid())
 				{
-					// To avoid missing progress callbacks for package already in-flight,
-					// call any state that the package might already have reached.
-					if (Package->AsyncPackageLoadingState >= EAsyncPackageLoadingState2::WaitingForIo)
-					{
-						Package->CallProgressCallback(*PackageProgressDelegate, EAsyncLoadingProgress::Started);
-					}
-
-					if (Package->AsyncPackageLoadingState >= EAsyncPackageLoadingState2::ProcessPackageSummary)
-					{
-						Package->CallProgressCallback(*PackageProgressDelegate, EAsyncLoadingProgress::Read);
-					}
-
-					if (Package->AsyncPackageLoadingState >= EAsyncPackageLoadingState2::ExportsDone)
-					{
-						Package->CallProgressCallback(*PackageProgressDelegate, EAsyncLoadingProgress::Serialized);
-					}
-
 					Package->AddProgressCallback(PackageProgressDelegate);
 				}
 			}
@@ -7508,10 +7563,12 @@ void FAsyncPackage2::StartLoading(FAsyncLoadingThreadState2& ThreadState, FIoBat
 
 	LoadStartTime = FPlatformTime::Seconds();
 
-	CallProgressCallbacks(EAsyncLoadingProgress::Started);
-
 	AsyncPackageLoadingState = EAsyncPackageLoadingState2::WaitingForIo;
 
+	// CallProgressCallbacks will read AsyncPackageLoadingState and dispatch progress so it needs
+	// to be called after we set the state to WaitingForIo to have an effect.
+	CallProgressCallbacks();
+
 	FIoReadOptions ReadOptions;
 	if (LinkerLoadState.IsSet())
 	{
@@ -8328,7 +8385,7 @@ EEventLoadNodeExecutionResult FAsyncPackage2::Event_ProcessPackageSummary(FAsync
 	UE::Core::Private::FPlayInEditorLoadingScope PlayInEditorIDScope(Package->Desc.PIEInstanceID);
 #endif
 
-	Package->CallProgressCallbacks(EAsyncLoadingProgress::Read);
+	Package->CallProgressCallbacks();
 
 	if (Package->LinkerLoadState.IsSet())
 	{
@@ -8443,35 +8500,26 @@ EEventLoadNodeExecutionResult FAsyncPackage2::Event_DependenciesReady(FAsyncLoad
 	return EEventLoadNodeExecutionResult::Complete;
 }
 
-void FAsyncPackage2::CallProgressCallback(FLoadPackageAsyncProgressDelegate& ProgressCallback, EAsyncLoadingProgress ProgressType)
+void FAsyncPackage2::CallProgressCallbacks()
 {
-	FLoadPackageAsyncProgressParams Params
+	FProgressCallbackEntry* Entry = ProgressCallbacksHead.load(std::memory_order_acquire);
+	if (!Entry)
 	{
-		.PackageName = Desc.UPackageName,
-		.LoadedPackage = GetLinkerRoot(),
-		.ProgressType = ProgressType
-	};
-
-	ProgressCallback.Invoke(Params);
-}
+		return;
+	}
 
-void FAsyncPackage2::CallProgressCallbacks(EAsyncLoadingProgress ProgressType)
-{
-	if (ProgressCallbacks.Num() != 0)
+	EAsyncLoadingProgress CurrentProgress = GetLoadingProgressFromState(AsyncPackageLoadingState);
+	if (CurrentProgress <= EAsyncLoadingProgress::Failed)
 	{
-		TRACE_CPUPROFILER_EVENT_SCOPE(FAsyncPackage2::CallProgressCallbacks);
+		return;
+	}
 
-		FLoadPackageAsyncProgressParams Params
-		{
-			.PackageName = Desc.UPackageName,
-			.LoadedPackage = GetLinkerRoot(),
-			.ProgressType = ProgressType
-		};
+	TRACE_CPUPROFILER_EVENT_SCOPE(FAsyncPackage2::CallProgressCallbacks);
 
-		for (TSharedPtr<FLoadPackageAsyncProgressDelegate>& ProgressCallback : ProgressCallbacks)
-		{
-			ProgressCallback->Invoke(Params);
-		}
+	while (Entry)
+	{
+		CallProgressCallback(*Entry, Desc.UPackageName, GetLinkerRoot(), CurrentProgress);
+		Entry = Entry->Next.load(std::memory_order_relaxed);
 	}
 }
 
@@ -9873,7 +9921,7 @@ EEventLoadNodeExecutionResult FAsyncPackage2::Event_ExportsDone(FAsyncLoadingThr
 		FCoreDelegates::ReleasePreloadedPackageShaderMaps.ExecuteIfBound(Package->Data.ShaderMapHashes);
 	}
 
-	Package->CallProgressCallbacks(EAsyncLoadingProgress::Serialized);
+	Package->CallProgressCallbacks();
 
 	Package->ConditionalReleasePartialRequests(ThreadState);
 	Package->ConditionalBeginPostLoad(ThreadState);
@@ -10838,7 +10886,7 @@ EAsyncPackageState::Type FAsyncLoadingThread2::ProcessLoadedPackagesFromGameThre
 			Package->AsyncPackageLoadingState = EAsyncPackageLoadingState2::DeferredClearImportedPackages;
 			DeferredCleanupPackages.Enqueue(Package);
 
-			if (Package->CompletionCallbacks.IsEmpty() && Package->ProgressCallbacks.IsEmpty())
+			if (Package->CompletionCallbacks.IsEmpty() && Package->ProgressCallbacksHead.load(std::memory_order_relaxed) == nullptr)
 			{
 				RemovePendingRequests(ThreadState, Package->RequestIDs);
 				Package->Release();
@@ -12069,6 +12117,14 @@ FAsyncPackage2::~FAsyncPackage2()
 	TRACE_LOADTIME_DESTROY_ASYNC_PACKAGE(this);
 	UE_ASYNC_PACKAGE_LOG(Verbose, Desc, TEXT("AsyncThread: Deleted"), TEXT("Package deleted."));
 
+	FProgressCallbackEntry* Head = ProgressCallbacksHead.exchange(nullptr, std::memory_order_relaxed);
+	while (Head)
+	{
+		FProgressCallbackEntry* Next = Head->Next.load(std::memory_order_relaxed);
+		delete Head;
+		Head = Next;
+	}
+
 	ImportStore.ReleaseImportedPackageReferences(Desc, HeaderData.ImportedPackageIds);
 #if WITH_EDITOR
 	if (OptionalSegmentHeaderData.IsSet())
@@ -12514,7 +12570,25 @@ void FAsyncPackage2::AddCompletionCallback(TUniquePtr<FLoadPackageAsyncDelegate>
 
 void FAsyncPackage2::AddProgressCallback(TSharedPtr<FLoadPackageAsyncProgressDelegate> Callback)
 {
-	ProgressCallbacks.Emplace(MoveTemp(Callback));
+	FProgressCallbackEntry* Entry = new FProgressCallbackEntry();
+	Entry->Delegate = MoveTemp(Callback);
+
+	// Lock-free push onto intrusive linked list
+	FProgressCallbackEntry* OldHead = ProgressCallbacksHead.load(std::memory_order_relaxed);
+	while (true)
+	{
+		Entry->Next.store(OldHead, std::memory_order_relaxed);
+		if (ProgressCallbacksHead.compare_exchange_weak(OldHead, Entry, std::memory_order_release, std::memory_order_relaxed))
+		{
+			break;
+		}
+	}
+
+	EAsyncLoadingProgress CurrentProgress = GetLoadingProgressFromState(AsyncPackageLoadingState);
+	if (CurrentProgress > EAsyncLoadingProgress::Failed)
+	{
+		CallProgressCallback(*Entry, Desc.UPackageName, GetLinkerRoot(), CurrentProgress);
+	}
 }
 
 UE_TRACE_EVENT_BEGIN(CUSTOM_LOADTIMER_LOG, LoadAsyncPackageInternal, NoSync)
