Fix FPersistentObjectRef to correctly handle breaks in the stable reference chain

#rnx
#Synced-CL 52168278

diff --git a/Engine/Source/Runtime/CoreUObject/Private/Misc/TransactionCommon.cpp b/Engine/Source/Runtime/CoreUObject/Private/Misc/TransactionCommon.cpp
--- a/Engine/Source/Runtime/CoreUObject/Private/Misc/TransactionCommon.cpp
+++ b/Engine/Source/Runtime/CoreUObject/Private/Misc/TransactionCommon.cpp
@@ -9,7 +9,11 @@ namespace UE::Transaction
 
 FPersistentObjectRef::FPersistentObjectRef(UObject* InObject)
 {
-	RootObject = InObject;
+	// Resolve the leaf-most object that is UseOuter() == false, without any UseOuter() == true objects above it
+	// This is needed as there may be a break in the stable reference chain, eg, Actor -> Component -> InstancedObject, 
+	// may have a UseOuter() == true Component despite the InstancedObject itself being UseOuter() == false, in which 
+	// case Actor needs to be the root object rather than InstancedObject
+	RootObject = nullptr;
 	{
 		auto UseOuter = [](const UObject* Obj)
 		{
@@ -43,10 +47,24 @@ FPersistentObjectRef::FPersistentObjectRef(UObject* InObject)
 			return false;
 		};
 
-		while (UseOuter(RootObject))
+		// First pass to find the correct root object
+		for (UObject* CurObject = InObject; CurObject; CurObject = CurObject->GetOuter())
 		{
-			SubObjectHierarchyIDs.Add(RootObject->GetFName());
-			RootObject = RootObject->GetOuter();
+			if (UseOuter(CurObject))
+			{
+				// This reference is unstable, so reset any current root object and keep going
+				RootObject = nullptr;
+			}
+			else if (!RootObject)
+			{
+				RootObject = CurObject;
+			}
+		}
+
+		// Second pass to fill SubObjectHierarchyIDs with the objects under the root object
+		for (UObject* CurObject = InObject; CurObject != RootObject; CurObject = CurObject->GetOuter())
+		{
+			SubObjectHierarchyIDs.Add(CurObject->GetFName());
 		}
 	}
 	check(RootObject);
