Appeased static analysis

#rnx
#Synced-CL 52168278

diff --git a/Engine/Source/Runtime/CoreUObject/Private/Misc/TransactionCommon.cpp b/Engine/Source/Runtime/CoreUObject/Private/Misc/TransactionCommon.cpp
--- a/Engine/Source/Runtime/CoreUObject/Private/Misc/TransactionCommon.cpp
+++ b/Engine/Source/Runtime/CoreUObject/Private/Misc/TransactionCommon.cpp
@@ -62,7 +62,7 @@ FPersistentObjectRef::FPersistentObjectRef(UObject* InObject)
 		}
 
 		// Second pass to fill SubObjectHierarchyIDs with the objects under the root object
-		for (UObject* CurObject = InObject; CurObject != RootObject; CurObject = CurObject->GetOuter())
+		for (UObject* CurObject = InObject; CurObject && CurObject != RootObject; CurObject = CurObject->GetOuter())
 		{
 			SubObjectHierarchyIDs.Add(CurObject->GetFName());
 		}
