Handle FVRestValueProperty in FReferenceCollector::AddPropertyReferences

This enables collectors without the serialization-based slow path to see them.

#rnx
#Synced-CL 52137917

diff --git a/Engine/Source/Runtime/CoreUObject/Private/UObject/UObjectGlobals.cpp b/Engine/Source/Runtime/CoreUObject/Private/UObject/UObjectGlobals.cpp
--- a/Engine/Source/Runtime/CoreUObject/Private/UObject/UObjectGlobals.cpp
+++ b/Engine/Source/Runtime/CoreUObject/Private/UObject/UObjectGlobals.cpp
@@ -5451,6 +5451,20 @@ void CollectPropertyReferences(FReferenceCollector& Collector, FProperty& Proper
 		{
 			CollectOptionalReference<CollectFlags>(Collector, static_cast<FOptionalProperty&>(Property), Value, Referencer);
 		}
+#if WITH_VERSE_VM || defined(__INTELLISENSE__)
+		else if (EnumHasAnyFlags(CastFlags, CASTCLASS_FVValueProperty))
+		{
+			Collector.AddReferencedVerseValue(*static_cast<Verse::TWriteBarrier<Verse::VValue>*>(Value), Referencer, &Property);
+		}
+		else if (EnumHasAnyFlags(CastFlags, CASTCLASS_FVRestValueProperty))
+		{
+			Collector.AddReferencedVerseValue(*static_cast<Verse::VRestValue*>(Value), Referencer, &Property);
+		}
+		else if (EnumHasAnyFlags(CastFlags, CASTCLASS_FVCellProperty))
+		{
+			Collector.AddReferencedVerseValue(*static_cast<Verse::TWriteBarrier<Verse::VCell>*>(Value), Referencer, &Property);
+		}
+#endif
 		else
 		{
 			// Fallback to virtual SerializeItem dispatch inside SerializeBin
