Acquire Verse heap access on async loading worker threads and during GC FlushAsyncLoading

This enables FVRestValueProperty serialization to allocate and mutate VCells, and GC-triggered PostLoads to run Verse code.

#rnx
#Synced-CL 52137917

diff --git a/Engine/Source/Runtime/CoreUObject/Private/Serialization/AsyncLoading2.cpp b/Engine/Source/Runtime/CoreUObject/Private/Serialization/AsyncLoading2.cpp
--- a/Engine/Source/Runtime/CoreUObject/Private/Serialization/AsyncLoading2.cpp
+++ b/Engine/Source/Runtime/CoreUObject/Private/Serialization/AsyncLoading2.cpp
@@ -6523,9 +6523,21 @@ void FEventLoadNode2::ParallelLoadingWork(FAsyncLoadingThreadState2* InParentSta
 	// to this package if they don't belong to any other package.
 	ThreadContext.AsyncPackage = Package;
 
+#if WITH_VERSE_VM || defined(__INTELLISENSE__)
+	Verse::FIOContext Context = Verse::FIOContext::CreateForManualStackScanning();
+	Context.SetIsInManuallyEmptyStack(false);
+	Verse::FRunningContext RunningContext = Context.AcquireAccessForManualStackScanning();
+#endif
+
 	EEventLoadNodeExecutionResult Result = Execute(LocalThreadState);
 	check(Result == EEventLoadNodeExecutionResult::Complete);
 
+#if WITH_VERSE_VM || defined(__INTELLISENSE__)
+	RunningContext.RelinquishAccessForManualStackScanning();
+	Context.SetIsInManuallyEmptyStack(true);
+	Context.ReleaseForManualStackScanning();
+#endif
+
 	ThreadContext.AsyncPackage = nullptr;
 
 	// We should never destroy a thread state that still has nodes to fire
@@ -11964,8 +11976,6 @@ void FAsyncLoadingThread2::NotifyScriptVersePackage(Verse::VPackage* Package)
 
 void FAsyncLoadingThread2::NotifyCompiledVersePackage(Verse::VPackage* VersePackage, UPackage* Package, TArrayView<UObject*> Exports)
 {
-	check(IsInGameThread() && !IsMultithreaded());
-
 	FName PackageName = Package->GetFName();
 	FPackageId PackageId = FPackageId::FromName(PackageName);
 
diff --git a/Engine/Source/Runtime/CoreUObject/Private/UObject/GarbageCollection.cpp b/Engine/Source/Runtime/CoreUObject/Private/UObject/GarbageCollection.cpp
--- a/Engine/Source/Runtime/CoreUObject/Private/UObject/GarbageCollection.cpp
+++ b/Engine/Source/Runtime/CoreUObject/Private/UObject/GarbageCollection.cpp
@@ -5994,11 +5994,6 @@ void FReachabilityAnalysisState::PerformReachabilityAnalysisAndConditionallyPurg
 	}
 	GGCStats.bFinishedAsFullPurge = bPerformFullPurge;
 
-#if WITH_VERSE_VM || defined(__INTELLISENSE__)
-	Verse::FRunningContext RunningContext = Verse::FRunningContextPromise{};
-	Verse::FIOContext Context = RelinquishVerseHeapAccess(RunningContext);
-#endif
-	
 	if (bPerformFullPurge)
 	{
 		UE::GC::PreCollectGarbageImpl<true>(ObjectKeepFlags);
@@ -6008,6 +6003,12 @@ void FReachabilityAnalysisState::PerformReachabilityAnalysisAndConditionallyPurg
 		UE::GC::PreCollectGarbageImpl<false>(ObjectKeepFlags);
 	}
 
+#if WITH_VERSE_VM || defined(__INTELLISENSE__)
+	// Wait to relinquish access until after PreCollectGarbageImpl calls FlushAsyncLoading.
+	Verse::FRunningContext RunningContext = Verse::FRunningContextPromise{};
+	Verse::FIOContext Context = RelinquishVerseHeapAccess(RunningContext);
+#endif
+	
 	const bool bForceNonIncrementalReachability =
 		!GIsIncrementalReachabilityPending &&
 		(bPerformFullPurge || !GAllowIncrementalReachability);
