[ZenLoader]
- Bring back an important log that was suppressed and didn't appear anywhere anymore because LogStreaming defaults to Display

#Synced-CL 52203528

diff --git a/Engine/Source/Runtime/CoreUObject/Private/Serialization/AsyncLoading2.cpp b/Engine/Source/Runtime/CoreUObject/Private/Serialization/AsyncLoading2.cpp
--- a/Engine/Source/Runtime/CoreUObject/Private/Serialization/AsyncLoading2.cpp
+++ b/Engine/Source/Runtime/CoreUObject/Private/Serialization/AsyncLoading2.cpp
@@ -11248,7 +11248,7 @@ FAsyncLoadingThread2::FAsyncLoadingThread2(FIoDispatcher& InIoDispatcher)
 	}
 #endif
 
-	UE_LOG(LogStreaming, Log, TEXT("AsyncLoading2 - Created: Async Loading Thread: %s, Async Post Load: %s, Multithreaded: %s"),
+	UE_LOG(LogStreaming, Display, TEXT("AsyncLoading2 - Created: Async Loading Thread: %s, Async Post Load: %s, Multithreaded: %s"),
 		FAsyncLoadingThreadSettings::Get().bAsyncLoadingThreadEnabled ? TEXT("true") : TEXT("false"),
 		FAsyncLoadingThreadSettings::Get().bAsyncPostLoadEnabled ? TEXT("true") : TEXT("false"),
 		GAllowMultithreadedLoading ? TEXT("true") : TEXT("false"));
