Silence static code analysis false positives.

#rnx

- If UE_WITH_REMOTE_OBJECT_HANDLE is defined as 0 then ::GetResidence will always return EResidence::Local.
- This has started to flag some lines of code as being redundant because we are calling ::GetResidence and checki the return value which will never change.
- We don't want to have to wrap all the checks with UE_WITH_REMOTE_OBJECT_HANDLE checks as it will bloat the code, so we are supressing the warnings instead.

#Synced-CL 52050999

diff --git a/Engine/Source/Runtime/CoreUObject/Private/UObject/RemoteObjectTransfer.cpp b/Engine/Source/Runtime/CoreUObject/Private/UObject/RemoteObjectTransfer.cpp
--- a/Engine/Source/Runtime/CoreUObject/Private/UObject/RemoteObjectTransfer.cpp
+++ b/Engine/Source/Runtime/CoreUObject/Private/UObject/RemoteObjectTransfer.cpp
@@ -1340,7 +1340,7 @@ void TransactionallyMigrateObjects(TRefCountPtr<FTransactionalMigrationData> Mig
 			checkf(Stub->MigratedData == nullptr, TEXT("Stub for serialized object %s still has its migrated data stored even though the data has been fully deserialized"), *SerializedId.ToString());
 			UObject* Object = StaticFindObjectFastInternal(SerializedId);
 			checkf(Object != nullptr, TEXT("Object %s has been migrated but its UObject does not exist locally"), *SerializedId.ToString());
-			checkf(GetResidence(Object) == EResidence::Local, TEXT("Object %s %s has been fully migrated but its residence is not Local (%s)"), *SerializedId.ToString(), *GetPathNameSafe(Object), EnumToString(GetResidence(Object)));
+			checkf(GetResidence(Object) == EResidence::Local, TEXT("Object %s %s has been fully migrated but its residence is not Local (%s)"), *SerializedId.ToString(), *GetPathNameSafe(Object), EnumToString(GetResidence(Object))); //-V547
 		}
 #endif
 	}
@@ -1700,7 +1700,7 @@ void StoreObjectToDatabase(UObject* Object, FRemoteObjectId ObjectId)
 	{
 		// We don't own this, and it's definitely remotely owned, just don't save it to the database
 		EResidence Residence = GetResidence(Object);
-		if (Residence == EResidence::Remote)
+		if (Residence == EResidence::Remote) //-V547
 		{
 			return;
 		}
@@ -1747,7 +1747,7 @@ void StoreObjectToDatabase(UObject* Object, FRemoteObjectId ObjectId)
 				FRemoteObjectId SentObjectId(SentObject);
 				if (FRemoteObjectStub* Stub = FindRemoteObjectStub(SentObjectId))
 				{
-					checkf(GetResidence(SentObject) != EResidence::Remote, TEXT("Double store to database for %s? Object is remote on server %s"),
+					checkf(GetResidence(SentObject) != EResidence::Remote, TEXT("Double store to database for %s? Object is remote on server %s"), //-V547
 						*SentObjectId.ToString(ERemoteIdToStringVerbosity::PathName),
 						*Stub->ResidentServerId.ToString());
 				}
