AssetRegistry: GathererCache: Add a check statement for invalid empty tag values rather than crashing on the check statement in operator [] when calling &String[0] on an empty string.
#rnx
#Synced-CL 52241492

diff --git a/Engine/Source/Runtime/CoreUObject/Private/AssetRegistry/AssetDataTagMap.cpp b/Engine/Source/Runtime/CoreUObject/Private/AssetRegistry/AssetDataTagMap.cpp
--- a/Engine/Source/Runtime/CoreUObject/Private/AssetRegistry/AssetDataTagMap.cpp
+++ b/Engine/Source/Runtime/CoreUObject/Private/AssetRegistry/AssetDataTagMap.cpp
@@ -984,6 +984,10 @@ namespace FixedTagPrivate
 		else
 		{	
 			FString String = Value.AsString();
+			// Zero-length Tag Values are invalid. They are supposed to be prevented by early exits in
+			// UE::AssetRegistry::ReadPackageDataMain and FAssetData::SetTagsAndAssetBundles.
+			checkf(String.Len() > 0, TEXT("Invalid Zero-length tag: Key: \"%s\", Value: \"%s\" IsFixed: %s."),
+				*Key.ToString(), *String, *LexToString(Value.IsFixed() != 0));
 			if (FCString::IsPureAnsi(*String, String.Len()))
 			{
 				return FValueId{EValueType::AnsiString,	 AnsiStrings.Index(MoveTemp(String))};
diff --git a/Engine/Source/Runtime/CoreUObject/Public/AssetRegistry/AssetDataTagMapSerializationDetails.h b/Engine/Source/Runtime/CoreUObject/Public/AssetRegistry/AssetDataTagMapSerializationDetails.h
--- a/Engine/Source/Runtime/CoreUObject/Public/AssetRegistry/AssetDataTagMapSerializationDetails.h
+++ b/Engine/Source/Runtime/CoreUObject/Public/AssetRegistry/AssetDataTagMapSerializationDetails.h
@@ -226,11 +226,11 @@ namespace FixedTagPrivate
 			}
 			static uint32 GetKeyHash(const FString& Key)
 			{
-				return HashCaseSensitive(&Key[0], Key.Len());
+				return HashCaseSensitive(GetData(Key), Key.Len());
 			}
 			static uint32 GetKeyHash(const FUtf8String& Key)
 			{
-				return HashCaseSensitive(&Key[0], Key.Len());
+				return HashCaseSensitive(GetData(Key), Key.Len());
 			}
 
 			static bool Matches(FNameEntryId A, FNameEntryId B)
