[Engine] Convert a PropertyClass check into checkfs so we can figure out (hopefully) which class properties are not uclass properties.  Seems to be triggered by the DAD system of patching.  My theory is the underlying property changed in code, so the default object has one property, and the patched has another, and it probably compares the two when serializing. to see if it should just use the default value.

diff --git a/Engine/Source/Runtime/CoreUObject/Private/UObject/PropertyClass.cpp b/Engine/Source/Runtime/CoreUObject/Private/UObject/PropertyClass.cpp
--- a/Engine/Source/Runtime/CoreUObject/Private/UObject/PropertyClass.cpp
+++ b/Engine/Source/Runtime/CoreUObject/Private/UObject/PropertyClass.cpp
@@ -186,8 +186,17 @@ bool FClassProperty::Identical(TNotNull<const void*> A, const void* B, uint32 Po
 	if (!(PortFlags & PPF_AvoidRemoteObjectMigration))
 #endif
 	{
-		check(ObjectA == nullptr || ObjectA.IsA<UClass>());
-		check(ObjectB == nullptr || ObjectB.IsA<UClass>());
+		// We need to see A and B in the same check.
+		checkf(
+			(ObjectA == nullptr || ObjectA.IsA<UClass>()) && (ObjectB == nullptr || ObjectB.IsA<UClass>()),
+			TEXT("FClassProperty::Identical - ObjectA: (Name: %s Path: %s Class: %s) ObjectB: (Name: %s Path: %s Class: %s)"),
+			ObjectA ? *ObjectA->GetName() : TEXT("nullptr"),
+			ObjectA ? *ObjectA->GetPathName() : TEXT("nullptr"),
+			ObjectA ? *ObjectA->GetClass()->GetName() : TEXT("nullptr"),
+			ObjectB ? *ObjectB->GetName() : TEXT("nullptr"),
+			ObjectB ? *ObjectB->GetPathName() : TEXT("nullptr"),
+			ObjectB ? *ObjectB->GetClass()->GetName() : TEXT("nullptr")
+		);
 	}
 
 	return (ObjectA == ObjectB);
