Adding more profiler events for remote object serialization

#Synced-CL 52058383

diff --git a/Engine/Source/Runtime/CoreUObject/Private/UObject/RemoteObjectSerialization.cpp b/Engine/Source/Runtime/CoreUObject/Private/UObject/RemoteObjectSerialization.cpp
--- a/Engine/Source/Runtime/CoreUObject/Private/UObject/RemoteObjectSerialization.cpp
+++ b/Engine/Source/Runtime/CoreUObject/Private/UObject/RemoteObjectSerialization.cpp
@@ -1346,6 +1346,8 @@ FORCENOINLINE void SortRemoteObjectHeadersByOuter(FArchiveRemoteObjectWriter& Ar
 
 void SerializeObjectDataInternal(FArchiveRemoteObjectWriter& Ar, UObject* RequestedObject, FRemoteObjectId RequestedObjectId, TSet<UObject*>& OutObjects, ERemoteObjectSerializationFlags SerializationFlags)
 {
+	TRACE_CPUPROFILER_EVENT_SCOPE(SerializeRemoteObjectDataInternal);
+
 	using namespace UE::RemoteObject::Handle;
 
 	TArray<FRemoteObjectHeader> ObjectHeaders;
@@ -1414,6 +1416,7 @@ void SerializeObjectDataInternal(FArchiveRemoteObjectWriter& Ar, UObject* Reques
 	// Skip sorting if there's just one object or we're resetting a borrowed object because CDO delta serialization substitutes serialized objects (archetypes) with their instances and the outers don't match
 	if (OutObjects.Num() > 1 && !EnumHasAnyFlags(SerializationFlags, ERemoteObjectSerializationFlags::Resetting))
 	{
+		TRACE_CPUPROFILER_EVENT_SCOPE(SerializeRemoteObjectDataSortHeaders);
 		SortRemoteObjectHeadersByOuter(Ar, ObjectHeaders);
 	}
 
@@ -1427,6 +1430,8 @@ void SerializeObjectDataInternal(FArchiveRemoteObjectWriter& Ar, UObject* Reques
 
 FRemoteObjectData SerializeObjectData(UObject* InObject, TSet<UObject*>& OutObjects, TSet<UObject*>& OutReferencedObjects, const FUObjectMigrationContext* MigrationContext, ERemoteObjectSerializationFlags InFlags)
 {
+	TRACE_CPUPROFILER_EVENT_SCOPE(SerializeRemoteObjectData);
+
 	UObject* Object = InObject;
 	if (GRemoteObjectsMigrateFullHierarchy && !EnumHasAnyFlags(InFlags, ERemoteObjectSerializationFlags::SkipCanonicalRootSearch))
 	{
@@ -1470,6 +1475,8 @@ int32 DeserializeObjectDataInternal(FArchiveRemoteObjectReader& Ar,
 
 void ResetRemoteObjects(const TArray<UObject*>& ObjectHierarchy)
 {
+	TRACE_CPUPROFILER_EVENT_SCOPE(ResetRemoteObjects)
+
 	for (UObject* ObjectToReset : ObjectHierarchy)
 	{
 		// This code is effectively copying the delta between the archetype and its instance (ObjectToReset) to the instance so that the instance (ObjectToReset) is identical to the archetype.
@@ -1691,6 +1698,8 @@ int32 DeserializeObjectData(FRemoteObjectData& ObjectData,
 	TArray<UObject*>& OutObjects,
 	ERemoteObjectSerializationFlags DeserializeFlags)
 {
+	TRACE_CPUPROFILER_EVENT_SCOPE(DeserializeRemoteObjectData);
+
 	FArchiveRemoteObjectReader Ar(ObjectData, MigrationContext, DeserializeFlags);
 	return DeserializeObjectDataInternal(Ar, ObjectData, MigrationContext, OutObjectRemoteIds, OutObjects, DeserializeFlags);
 }
