Resolving -disableunity -nopch linker issue for MHATest project

#Synced-CL 52291121

diff --git a/Engine/Source/Runtime/CoreUObject/Private/UObject/ObjectTrace.cpp b/Engine/Source/Runtime/CoreUObject/Private/UObject/ObjectTrace.cpp
--- a/Engine/Source/Runtime/CoreUObject/Private/UObject/ObjectTrace.cpp
+++ b/Engine/Source/Runtime/CoreUObject/Private/UObject/ObjectTrace.cpp
@@ -13,6 +13,7 @@
 #include "ProfilingDebugging/TraceAuxiliary.h"
 #include "Trace/Config.h"
 #include "Trace/Trace.inl"
+#include "Templates/Casts.h"
 #include "UObject/Class.h"
 #include "UObject/NameTypes.h"
 #include "UObject/Object.h"
