Add * and -> operators to TInstancedStruct for convenience

#Synced-CL 52296784

diff --git a/Engine/Source/Runtime/CoreUObject/Public/StructUtils/InstancedStruct.h b/Engine/Source/Runtime/CoreUObject/Public/StructUtils/InstancedStruct.h
--- a/Engine/Source/Runtime/CoreUObject/Public/StructUtils/InstancedStruct.h
+++ b/Engine/Source/Runtime/CoreUObject/Public/StructUtils/InstancedStruct.h
@@ -471,6 +471,30 @@ public:
 		return UE::StructUtils::GetStructPtr<T, BaseStructT>(InstancedStruct.GetScriptStruct(), InstancedStruct.GetMutableMemory());
 	}
 
+	/** Returns a mutable reference. */
+	[[nodiscard]] UE_REWRITE BaseStructT& operator*() requires (!std::is_const_v<BaseStructT>)
+	{
+		return GetMutablePtr();
+	}
+
+	/** Returns an immutable reference. */
+	[[nodiscard]] UE_REWRITE const std::decay_t<BaseStructT>& operator*() const
+	{
+		return GetPtr();
+	}
+
+	/** Returns a mutable pointer. */
+	[[nodiscard]] UE_REWRITE BaseStructT* operator->() requires (!std::is_const_v<BaseStructT>)
+	{
+		return GetMutablePtr();
+	}
+
+	/** Returns an immutable pointer. */
+	[[nodiscard]] UE_REWRITE const std::decay_t<BaseStructT>* operator->() const
+	{
+		return GetPtr();
+	}
+
 	/** Returns True if the struct is valid.*/
 	bool IsValid() const
 	{
