Made placeholders capable of keeping the super class of component or entity and reduced restrictions that hid them from component iterators

#Synced-CL 52110756

diff --git a/Engine/Source/Runtime/CoreUObject/Internal/UObject/InstanceDataObjectUtils.h b/Engine/Source/Runtime/CoreUObject/Internal/UObject/InstanceDataObjectUtils.h
--- a/Engine/Source/Runtime/CoreUObject/Internal/UObject/InstanceDataObjectUtils.h
+++ b/Engine/Source/Runtime/CoreUObject/Internal/UObject/InstanceDataObjectUtils.h
@@ -57,6 +57,7 @@ bool CanCreatePropertyBagPlaceholderForType(const UClass* ClassKind, const UStru
 
 /** Query if a given type can serve as a super-type for a placeholder type. */
 bool CanCreatePropertyBagPlaceholderSuperType(const UE::Serialization::Private::FTypeResource& SuperType);
+UE_API bool CanCreatePropertyBagPlaceholderSuperType(const UStruct* SuperType);
 
 /** Helper to check if a class is an IDO class. */
 UE_API bool IsClassOfInstanceDataObjectClass(const UStruct* Class);
diff --git a/Engine/Source/Runtime/CoreUObject/Private/UObject/InstanceDataObjectUtils.cpp b/Engine/Source/Runtime/CoreUObject/Private/UObject/InstanceDataObjectUtils.cpp
--- a/Engine/Source/Runtime/CoreUObject/Private/UObject/InstanceDataObjectUtils.cpp
+++ b/Engine/Source/Runtime/CoreUObject/Private/UObject/InstanceDataObjectUtils.cpp
@@ -886,6 +886,22 @@ namespace UE
 		return bResult;
 	}
 
+	bool CanCreatePropertyBagPlaceholderSuperType(const UStruct* SuperType)
+	{
+		bool bResult = false;
+
+		if ((SuperType->GetFName() == NAME_Component) && (SuperType->GetClass()->GetFName() == NAME_VerseClass))
+		{
+			bResult = true;
+		}
+		else if ((SuperType->GetFName() == NAME_Entity) && (SuperType->GetClass()->GetFName() == NAME_Class))
+		{
+			bResult = true;
+		}
+
+		return bResult;
+	}
+
 	bool IsClassOfInstanceDataObjectClass(const UStruct* Class)
 	{
 		return Class->IsA(UInstanceDataObjectClass::StaticClass()) || Class->IsA(UInstanceDataObjectStruct::StaticClass());
