Removing unnecessary includes from UObject/Class.h
Adding includes to other files that relied on the removed includes.

#Synced-CL 52214067

diff --git a/Engine/Plugins/Experimental/PythonScriptPlugin/Source/PythonScriptPlugin/Private/PyGenUtil.h b/Engine/Plugins/Experimental/PythonScriptPlugin/Source/PythonScriptPlugin/Private/PyGenUtil.h
--- a/Engine/Plugins/Experimental/PythonScriptPlugin/Source/PythonScriptPlugin/Private/PyGenUtil.h
+++ b/Engine/Plugins/Experimental/PythonScriptPlugin/Source/PythonScriptPlugin/Private/PyGenUtil.h
@@ -13,6 +13,7 @@
 #include "UObject/WeakObjectPtrTemplates.h"
 #include "UObject/PropertyAccessUtil.h"
 #include "UObject/WeakFieldPtr.h"
+#include "Async/Mutex.h"
 
 #if WITH_PYTHON
 
diff --git a/Engine/Plugins/FX/Niagara/Source/NiagaraEditor/Private/TraversalCache/Traversal.h b/Engine/Plugins/FX/Niagara/Source/NiagaraEditor/Private/TraversalCache/Traversal.h
--- a/Engine/Plugins/FX/Niagara/Source/NiagaraEditor/Private/TraversalCache/Traversal.h
+++ b/Engine/Plugins/FX/Niagara/Source/NiagaraEditor/Private/TraversalCache/Traversal.h
@@ -5,6 +5,7 @@
 #include "Async/RecursiveMutex.h"
 #include "Delegates/DelegateCombinations.h"
 #include "TraversalCache/TraversalShared.h"
+#include "Async/Mutex.h"
 
 namespace UE::Niagara::TraversalCache
 {
@@ -183,4 +184,5 @@ private:
 	mutable UE::FMutex CanWriteStaticAttributesCacheGuard;	
 };
 
-} // UE::Niagara::TraversalCache
\ No newline at end of file
+} // UE::Niagara::TraversalCache
+
diff --git a/Engine/Plugins/FX/Niagara/Source/NiagaraEditor/Private/TraversalCache/TraversalBuilder.h b/Engine/Plugins/FX/Niagara/Source/NiagaraEditor/Private/TraversalCache/TraversalBuilder.h
--- a/Engine/Plugins/FX/Niagara/Source/NiagaraEditor/Private/TraversalCache/TraversalBuilder.h
+++ b/Engine/Plugins/FX/Niagara/Source/NiagaraEditor/Private/TraversalCache/TraversalBuilder.h
@@ -8,6 +8,7 @@
 #include "Templates/SharedPointer.h"
 #include "Templates/SubclassOf.h"
 #include "TraversalCache/TraversalShared.h"
+#include "Async/Mutex.h"
 
 class UClass;
 class UEdGraphNode;
diff --git a/Engine/Plugins/FX/Niagara/Source/NiagaraEditor/Private/TraversalCache/TraversalCache.h b/Engine/Plugins/FX/Niagara/Source/NiagaraEditor/Private/TraversalCache/TraversalCache.h
--- a/Engine/Plugins/FX/Niagara/Source/NiagaraEditor/Private/TraversalCache/TraversalCache.h
+++ b/Engine/Plugins/FX/Niagara/Source/NiagaraEditor/Private/TraversalCache/TraversalCache.h
@@ -8,6 +8,7 @@
 #include "NiagaraTypes.h"
 #include "TraversalCache/TraversalShared.h"
 #include "UObject/ObjectKey.h"
+#include "Async/Mutex.h"
 
 class UEdGraphPin;
 class UNiagaraNode;
@@ -251,4 +252,5 @@ EStackFunctionReadFilterFlags operator&(
 EStackFunctionReadFilterFlags operator~(
 	EStackFunctionReadFilterFlags Flags);
 
-} // UE::Niagara::TraversalCache
\ No newline at end of file
+} // UE::Niagara::TraversalCache
+
diff --git a/Engine/Plugins/FX/Niagara/Source/NiagaraEditor/Private/ViewModels/NiagaraScriptStatsViewModel.h b/Engine/Plugins/FX/Niagara/Source/NiagaraEditor/Private/ViewModels/NiagaraScriptStatsViewModel.h
--- a/Engine/Plugins/FX/Niagara/Source/NiagaraEditor/Private/ViewModels/NiagaraScriptStatsViewModel.h
+++ b/Engine/Plugins/FX/Niagara/Source/NiagaraEditor/Private/ViewModels/NiagaraScriptStatsViewModel.h
@@ -4,6 +4,8 @@
 
 #include "CoreMinimal.h"
 #include "Templates/SharedPointer.h"
+#include "Async/Mutex.h"
+#include "Async/UniqueLock.h"
 #include "NiagaraScriptStatsViewModel.generated.h"
 
 class FNiagaraSystemViewModel;
diff --git a/Engine/Source/Developer/CollectionManager/Private/Collection.h b/Engine/Source/Developer/CollectionManager/Private/Collection.h
--- a/Engine/Source/Developer/CollectionManager/Private/Collection.h
+++ b/Engine/Source/Developer/CollectionManager/Private/Collection.h
@@ -6,6 +6,7 @@
 #include "Misc/Guid.h"
 #include "CollectionManagerTypes.h"
 #include "Templates/DontCopy.h"
+#include "Async/Mutex.h"
 
 class FTextFilterExpressionEvaluator;
 class ITextFilterExpressionContext;
diff --git a/Engine/Source/Developer/Windows/LiveCoding/Private/LiveCodingModule.h b/Engine/Source/Developer/Windows/LiveCoding/Private/LiveCodingModule.h
--- a/Engine/Source/Developer/Windows/LiveCoding/Private/LiveCodingModule.h
+++ b/Engine/Source/Developer/Windows/LiveCoding/Private/LiveCodingModule.h
@@ -12,6 +12,8 @@
 #include "Internationalization/Text.h"
 #include "Async/TaskGraphFwd.h"
 #include <atomic>
+#include "Async/Mutex.h"
+#include "Async/UniqueLock.h"
 
 struct IConsoleCommand;
 class IConsoleVariable;
diff --git a/Engine/Source/Runtime/CoreUObject/Private/Internationalization/PackageLocalizationManager.cpp b/Engine/Source/Runtime/CoreUObject/Private/Internationalization/PackageLocalizationManager.cpp
--- a/Engine/Source/Runtime/CoreUObject/Private/Internationalization/PackageLocalizationManager.cpp
+++ b/Engine/Source/Runtime/CoreUObject/Private/Internationalization/PackageLocalizationManager.cpp
@@ -8,6 +8,7 @@
 #include "Internationalization/Culture.h"
 #include "Internationalization/IPackageLocalizationCache.h"
 #include "Internationalization/PackageLocalizationCache.h"
+#include "Async/UniqueLock.h"
 
 DEFINE_LOG_CATEGORY_STATIC(LogPackageLocalizationManager, Log, All);
 
diff --git a/Engine/Source/Runtime/CoreUObject/Private/UObject/LinkerPlaceholderBase.cpp b/Engine/Source/Runtime/CoreUObject/Private/UObject/LinkerPlaceholderBase.cpp
--- a/Engine/Source/Runtime/CoreUObject/Private/UObject/LinkerPlaceholderBase.cpp
+++ b/Engine/Source/Runtime/CoreUObject/Private/UObject/LinkerPlaceholderBase.cpp
@@ -7,6 +7,7 @@
 #include "UObject/UnrealType.h"
 #include "UObject/UnrealTypePrivate.h"
 #include "Blueprint/BlueprintSupport.h"
+#include "Async/UniqueLock.h"
 
 #if USE_DEFERRED_DEPENDENCY_CHECK_VERIFICATION_TESTS
 	#define DEFERRED_DEPENDENCY_ENSURE(EnsueExpr) ensure(EnsueExpr)
diff --git a/Engine/Source/Runtime/CoreUObject/Private/UObject/LinkerPlaceholderBase.h b/Engine/Source/Runtime/CoreUObject/Private/UObject/LinkerPlaceholderBase.h
--- a/Engine/Source/Runtime/CoreUObject/Private/UObject/LinkerPlaceholderBase.h
+++ b/Engine/Source/Runtime/CoreUObject/Private/UObject/LinkerPlaceholderBase.h
@@ -8,6 +8,7 @@
 #include "UObject/ObjectResource.h"
 #include "UObject/LinkerLoad.h"
 #include "UObject/Field.h"
+#include "Async/UniqueLock.h"
 
 FORCEINLINE bool UseCircularDependencyLoadDeferring()
 {
diff --git a/Engine/Source/Runtime/CoreUObject/Private/UObject/LinkerPlaceholderClass.h b/Engine/Source/Runtime/CoreUObject/Private/UObject/LinkerPlaceholderClass.h
--- a/Engine/Source/Runtime/CoreUObject/Private/UObject/LinkerPlaceholderClass.h
+++ b/Engine/Source/Runtime/CoreUObject/Private/UObject/LinkerPlaceholderClass.h
@@ -7,6 +7,7 @@
 #include "UObject/UObjectGlobals.h"
 #include "UObject/Class.h"
 #include "UObject/LinkerPlaceholderBase.h"
+#include "Async/UniqueLock.h"
 
 #include "LinkerPlaceholderClass.generated.h"
 
diff --git a/Engine/Source/Runtime/CoreUObject/Private/UObject/LinkerPlaceholderFunction.h b/Engine/Source/Runtime/CoreUObject/Private/UObject/LinkerPlaceholderFunction.h
--- a/Engine/Source/Runtime/CoreUObject/Private/UObject/LinkerPlaceholderFunction.h
+++ b/Engine/Source/Runtime/CoreUObject/Private/UObject/LinkerPlaceholderFunction.h
@@ -6,6 +6,7 @@
 #include "UObject/ObjectMacros.h"
 #include "UObject/Class.h"
 #include "UObject/LinkerPlaceholderBase.h"
+#include "Async/UniqueLock.h"
 
 #include "LinkerPlaceholderFunction.generated.h"
 
diff --git a/Engine/Source/Runtime/CoreUObject/Private/UObject/Property.cpp b/Engine/Source/Runtime/CoreUObject/Private/UObject/Property.cpp
--- a/Engine/Source/Runtime/CoreUObject/Private/UObject/Property.cpp
+++ b/Engine/Source/Runtime/CoreUObject/Private/UObject/Property.cpp
@@ -10,6 +10,7 @@
 #include "Math/Ray.h"
 #include "Math/Sphere.h"
 #include "Misc/AsciiSet.h"
+#include "Misc/FallbackStruct.h"
 #include "Misc/Guid.h"
 #include "Misc/ScopeExit.h"
 #include "Misc/StringBuilder.h"
diff --git a/Engine/Source/Runtime/CoreUObject/Private/UObject/PropertyPrivate.cpp b/Engine/Source/Runtime/CoreUObject/Private/UObject/PropertyPrivate.cpp
--- a/Engine/Source/Runtime/CoreUObject/Private/UObject/PropertyPrivate.cpp
+++ b/Engine/Source/Runtime/CoreUObject/Private/UObject/PropertyPrivate.cpp
@@ -22,6 +22,7 @@
 #include "UObject/GarbageCollectionSchema.h"
 #include "UObject/LinkerPlaceholderFunction.h"
 #include "UObject/LinkerPlaceholderClass.h"
+#include "Async/UniqueLock.h"
 
 /*-----------------------------------------------------------------------------
 	UProperty implementation.
diff --git a/Engine/Source/Runtime/CoreUObject/Private/UObject/PropertyStruct.cpp b/Engine/Source/Runtime/CoreUObject/Private/UObject/PropertyStruct.cpp
--- a/Engine/Source/Runtime/CoreUObject/Private/UObject/PropertyStruct.cpp
+++ b/Engine/Source/Runtime/CoreUObject/Private/UObject/PropertyStruct.cpp
@@ -11,6 +11,7 @@
 #include "UObject/PropertyHelper.h"
 #include "UObject/UnrealTypePrivate.h"
 #include "UObject/UObjectThreadContext.h"
+#include "Async/UniqueLock.h"
 
 static inline void PreloadInnerStructMembers(FStructProperty* StructProperty)
 {
diff --git a/Engine/Source/Runtime/CoreUObject/Public/UObject/Class.h b/Engine/Source/Runtime/CoreUObject/Public/UObject/Class.h
--- a/Engine/Source/Runtime/CoreUObject/Public/UObject/Class.h
+++ b/Engine/Source/Runtime/CoreUObject/Public/UObject/Class.h
@@ -17,13 +17,7 @@
 #include "Containers/Map.h"
 #include "Containers/Set.h"
 #include "Containers/StringFwd.h"
-#include "Containers/UnrealString.h"
-#include "CoreTypes.h"
 #include "Algo/MaxElement.h"
-#include "Async/Mutex.h"
-#include "Async/SharedLock.h"
-#include "Async/SharedMutex.h"
-#include "Async/UniqueLock.h"
 #include "HAL/CriticalSection.h"
 #include "HAL/PlatformCrt.h"
 #include "HAL/UnrealMemory.h"
@@ -31,21 +25,11 @@
 #include "Logging/LogCategory.h"
 #include "Logging/LogMacros.h"
 #include "Logging/LogVerbosity.h"
-#include "Math/Box2D.h"
 #include "Math/InterpCurvePoint.h"
 #include "Math/MathFwd.h"
-#include "Math/Matrix.h"
-#include "Math/Plane.h"
-#include "Math/Quat.h"
-#include "Math/RandomStream.h"
-#include "Math/Rotator.h"
-#include "Math/Transform.h"
-#include "Math/Vector2D.h"
-#include "Math/Vector4.h"
 #include "Misc/AssertionMacros.h"
 #include "Misc/CString.h"
 #include "Misc/EnumClassFlags.h"
-#include "Misc/FallbackStruct.h"
 #include "Misc/Guid.h"
 #include "Misc/NotNull.h"
 #include "Misc/Optional.h"
@@ -53,19 +37,16 @@
 #include "Serialization/StructuredArchive.h"
 #include "Serialization/StructuredArchiveAdapters.h"
 #include "Templates/AlignmentTemplates.h"
-#include "Templates/EnableIf.h"
 #include "Templates/Identity.h"
 #include "Templates/IsAbstract.h"
 #include "Templates/IsEnum.h"
 #include "Templates/IsPODType.h"
 #include "Templates/IsUECoreType.h"
 #include "Templates/Models.h"
-#include "Templates/NoDestroy.h"
 #include "Templates/Tuple.h"
 #include "Templates/TypeCompatibleBytes.h"
 #include "Templates/UnrealTemplate.h"
 #include "Templates/UnrealTypeTraits.h"
-#include "Trace/Detail/Channel.h"
 #include "UObject/CoreNative.h"
 #include "UObject/Field.h"
 #include "UObject/FieldPath.h"
@@ -75,18 +56,25 @@
 #include "UObject/ObjectMacros.h"
 #include "UObject/PropertyTag.h"
 #include "UObject/PropertyVisitor.h"
-#include "UObject/ReflectedTypeAccessors.h"
 #include "UObject/Script.h"
 #include "UObject/StructOpsTypeTraits.h"
 #include "UObject/TopLevelAssetPath.h"
 #include "UObject/UObjectGlobals.h"
-#include "UObject/UnrealNames.h"
 #include "UObject/ObjectPtr.h"
 #include "UObject/ObjectPtrLinkedListBuilder.h"
 #include "UObject/ObjectPtrLinkedListIterator.h"
 #include <type_traits>
 #include <atomic>
 
+#if UE_WITH_CONSTINIT_UOBJECT
+#include "Async/SharedLock.h"
+#include "Async/SharedMutex.h"
+#include "Async/UniqueLock.h"
+#elif PLATFORM_FORCE_UE_LOCK_USAGE
+#include "Async/Mutex.h"
+#include "Async/UniqueLock.h"
+#endif
+
 #include "Class.generated.h"
 
 class FArchive;
@@ -5770,3 +5758,27 @@
 UE_DECLARE_CORE_VARIANT_TYPE(Sphere3, Sphere);
 
 #undef UE_DECLARE_CORE_VARIANT_TYPE
+
+#if UE_ENABLE_INCLUDE_ORDER_DEPRECATED_IN_5_8
+#include "Containers/UnrealString.h"
+#include "CoreTypes.h"
+#include "Async/Mutex.h"
+#include "Async/SharedLock.h"
+#include "Async/SharedMutex.h"
+#include "Async/UniqueLock.h"
+#include "Math/Box2D.h"
+#include "Math/Matrix.h"
+#include "Math/Plane.h"
+#include "Math/Quat.h"
+#include "Math/RandomStream.h"
+#include "Math/Rotator.h"
+#include "Math/Transform.h"
+#include "Math/Vector2D.h"
+#include "Math/Vector4.h"
+#include "Misc/FallbackStruct.h"
+#include "Templates/EnableIf.h"
+#include "Templates/NoDestroy.h"
+#include "Trace/Detail/Channel.h"
+#include "UObject/ReflectedTypeAccessors.h"
+#include "UObject/UnrealNames.h"
+#endif
diff --git a/Engine/Source/Runtime/Engine/Classes/Animation/AnimSequence.h b/Engine/Source/Runtime/Engine/Classes/Animation/AnimSequence.h
--- a/Engine/Source/Runtime/Engine/Classes/Animation/AnimSequence.h
+++ b/Engine/Source/Runtime/Engine/Classes/Animation/AnimSequence.h
@@ -26,6 +26,7 @@
 
 #if WITH_EDITOR
 #include "AnimData/IAnimationDataModel.h"
+#include "Async/Mutex.h"
 #endif // WITH_EDITOR
 
 #include "AnimSequence.generated.h"
diff --git a/Engine/Source/Runtime/Engine/Classes/Engine/SkinnedAssetCommon.h b/Engine/Source/Runtime/Engine/Classes/Engine/SkinnedAssetCommon.h
--- a/Engine/Source/Runtime/Engine/Classes/Engine/SkinnedAssetCommon.h
+++ b/Engine/Source/Runtime/Engine/Classes/Engine/SkinnedAssetCommon.h
@@ -13,6 +13,7 @@
 #include "UObject/PerPlatformProperties.h"
 #include "SkeletalMeshReductionSettings.h"
 #include "Animation/SkeletalMeshVertexAttribute.h"
+#include "Async/Mutex.h"
 
 #include "SkinnedAssetCommon.generated.h"
 
diff --git a/Engine/Source/Runtime/Engine/Public/WorldPartition/WorldPartitionStreamingSource.h b/Engine/Source/Runtime/Engine/Public/WorldPartition/WorldPartitionStreamingSource.h
--- a/Engine/Source/Runtime/Engine/Public/WorldPartition/WorldPartitionStreamingSource.h
+++ b/Engine/Source/Runtime/Engine/Public/WorldPartition/WorldPartitionStreamingSource.h
@@ -4,6 +4,7 @@
 #include "Math/Color.h"
 #include "Math/RandomStream.h"
 #include "HAL/IConsoleManager.h"
+#include "Math/Box2D.h"
 #include "WorldPartitionStreamingSource.generated.h"
 
 /** See https://en.wikipedia.org/wiki/Spherical_sector. */
diff --git a/Engine/Source/Runtime/MassEntity/Private/MassEntityManagerStorage.cpp b/Engine/Source/Runtime/MassEntity/Private/MassEntityManagerStorage.cpp
--- a/Engine/Source/Runtime/MassEntity/Private/MassEntityManagerStorage.cpp
+++ b/Engine/Source/Runtime/MassEntity/Private/MassEntityManagerStorage.cpp
@@ -5,6 +5,7 @@
 #include "MassEntityHandle.h"
 #include "MassEntityTypes.h"
 #include "Templates/SharedPointer.h"
+#include "Async/UniqueLock.h"
 
 namespace UE::Mass
 {
diff --git a/Engine/Source/Runtime/SlateCore/Public/Types/SlateBox2.h b/Engine/Source/Runtime/SlateCore/Public/Types/SlateBox2.h
--- a/Engine/Source/Runtime/SlateCore/Public/Types/SlateBox2.h
+++ b/Engine/Source/Runtime/SlateCore/Public/Types/SlateBox2.h
@@ -3,6 +3,7 @@
 #pragma once
 
 #include "Types/SlateVector2.h"
+#include "Math/Box2D.h"
 
 namespace UE::Slate
 {
