Blueprint Actor Construction Scripts: CookDeterminism: Add FCookRand to provide a seeded RandomStream to supply deterministic RandomGuids when NewGuid is requested from blueprint code running Actor Construction Scripts.
#rnx

#Synced-CL 52239865

diff --git a/Engine/Source/Runtime/CoreUObject/Private/Cooker/CookRand.cpp b/Engine/Source/Runtime/CoreUObject/Private/Cooker/CookRand.cpp
new file mode 100644
--- /dev/null
+++ b/Engine/Source/Runtime/CoreUObject/Private/Cooker/CookRand.cpp
@@ -0,0 +1,120 @@
+// Copyright Epic Games, Inc. All Rights Reserved.
+
+#if WITH_EDITOR
+
+#include "Cooker/CookRand.h"
+
+#include "Containers/UnrealString.h"
+#include "Hash/Blake3.h"
+#include "Misc/Char.h"
+#include "Misc/Guid.h"
+#include "Misc/StringBuilder.h"
+
+namespace UE::Cook
+{
+
+thread_local TRefCountPtr<FCookRand> GCookRand;
+
+FCookRand::FCookRand(int32 InSeed)
+	: RandomStreamVar(InSeed)
+{
+}
+
+FGuid FCookRand::RandomGuid() const
+{
+	FBlake3 Builder;
+
+	// Use a number of random bits equal in size to FGuid, so that we do not unnecessarily restrict our random
+	// output to some subset of possible guids.
+	for (int32 WordIndex = 0; WordIndex < sizeof(FGuid) / sizeof(int32); ++WordIndex)
+	{
+		uint32 RandWord = RandomStreamVar.GetUnsignedInt();
+		Builder.Update(&RandWord, sizeof(RandWord));
+	}
+	return FGuid::NewGuidFromHash(Builder.Finalize());
+}
+
+TRefCountPtr<FCookRand> FCookRand::GetThreadScope()
+{
+	return GCookRand;
+}
+
+int32 FCookRand::ConstructSeed(FStringView String)
+{
+	// Lower-case the string so that filepath capitalization differences do not alter the hash.
+	// Convert the string to utf8 so that whether TCHAR is UTF8 or UTF16 does not alter the hash.
+	TUtf8StringBuilder<1024> Normalized;
+	for (TCHAR C : String)
+	{
+		C = FChar::ToLower(C);
+		// Conversion operator<<(TCHAR) is not supported because it is unclear what it would do. So instead cast
+		// our TCHAR to an FStringView and append that.
+		Normalized << FStringView(&C, 1);
+	}
+
+	// Don't use GetTypeHash since GetTypeHash documentation explicitly warns that the values are not
+	// guaranteed persistent between processes.
+	FBlake3 Builder;
+	Builder.Update(*Normalized, Normalized.Len() * sizeof(**Normalized));
+	FBlake3Hash FullHash = Builder.Finalize();
+
+	// Arbitrary shorter-length prefixes of Blake3 hashes are documented as the equivalent of running
+	// the Blake3 hashing algorithm for the shorter length. So just take a 4-byte prefix of the Blake3Hash
+	// to get our int32.
+	int32 Result = 0;
+	static_assert(sizeof(FBlake3Hash::ByteArray) >= sizeof(Result));
+	// Use memcpy instead of reinterpret_cast to avoid technically undefined behavior due to misalignment.
+	FMemory::Memcpy(&Result, FullHash.GetBytes(), sizeof(Result));
+	return Result;
+}
+
+FCookRandScope::FCookRandScope(int32 Seed, ECookRandScope RandScope)
+{
+	if (ShouldAddNewInCurrentScope(RandScope))
+	{
+		bRestorePrevious = true;
+		Previous = MoveTemp(GCookRand);
+		GCookRand = MakeRefCount<FCookRand>(Seed);
+	}
+}
+
+FCookRandScope::FCookRandScope(TRefCountPtr<FCookRand> CookRand, ECookRandScope RandScope)
+{
+	if (ShouldAddNewInCurrentScope(RandScope))
+	{
+		bRestorePrevious = true;
+		Previous = MoveTemp(GCookRand);
+		GCookRand = MoveTemp(CookRand);
+	}
+}
+
+FCookRandScope::~FCookRandScope()
+{
+	if (bRestorePrevious)
+	{
+		GCookRand = MoveTemp(Previous);
+	}
+}
+
+bool FCookRandScope::ShouldAddNewInCurrentScope(ECookRandScope RandScope)
+{
+	switch (RandScope)
+	{
+	case ECookRandScope::FindOrAdd:
+		if (GCookRand)
+		{
+			return false;
+		}
+		break;
+	case ECookRandScope::Overwrite:
+		break;
+	default:
+		checkNoEntry();
+		break;
+	}
+	return true;
+}
+
+}; // namespace UE::Cook
+
+#endif // WITH_EDITOR
diff --git a/Engine/Source/Runtime/CoreUObject/Public/Cooker/CookRand.h b/Engine/Source/Runtime/CoreUObject/Public/Cooker/CookRand.h
new file mode 100644
--- /dev/null
+++ b/Engine/Source/Runtime/CoreUObject/Public/Cooker/CookRand.h
@@ -0,0 +1,108 @@
+// Copyright Epic Games, Inc. All Rights Reserved.
+
+#pragma once
+
+#if WITH_EDITOR
+
+#include "Containers/StringFwd.h"
+#include "HAL/CriticalSection.h"
+#include "Math/RandomStream.h"
+#include "Templates/RefCounting.h"
+#include "Templates/UnrealTemplate.h"
+
+struct FGuid;
+
+namespace UE::Cook
+{
+
+/**
+ * A RefCounted RandomStream and a threadlocal global accessor for accessing a currently scoped instance.
+ * Intended use is for cook operations to set a scope that can be used for deterministic replacement of random or
+ * other fundamentally non-deterministic calculations such as NewGuid.
+ * 
+ * Like RandomStream, const functions on FCookRand can have the side effect of mutating the RandomSeed to the next
+ * random value.
+ * 
+ * An FCookRand instance is not thread-safe because call order in a threaded environment is in general
+ * non-deterministic. Callers that run multithreaded operations using FCookRand should maintain their own FCookRand
+ * and should pass exclusive ownership of it to each task thread in a deterministic order.
+ * 
+ * Example usage:
+ * 
+ * void HighLevelFunctionWithUserCode()
+ * {
+ *     int32 MySeed = ComputeMyDeterministicSeed();
+ *     UE::Cook::FCookRandScope Scope(MySeed, UE::Cook::ECookRandScope::Overwrite);
+ *
+ *    // ... Call UserCode that uses MyAPI
+ * }
+ *
+ * FGuid MyAPI::NewGuid()
+ * {
+ *     TRefCountPtr<FCookRand> Rand = FCookRand::GetThreadScope();
+ *     return Rand ? Rand->RandomGuid() : FGuid::NewGuid();
+ * }
+ * float MyAPI::FRand()
+ * {
+ *    TRefCountPtr<FCookRand> Rand = FCookRand::GetThreadScope();
+ *    return Rand ? Rand->RandomStream().FRand() : FMath::FRand();
+ * }
+ */
+class FCookRand : public FRefCountedObject
+{
+public:
+	COREUOBJECT_API explicit FCookRand(int32 InSeed);
+
+
+	COREUOBJECT_API FGuid RandomGuid() const;
+	const FRandomStream& RandomStream() const;
+
+	/** Construct a seed to initialize a CookRand, from a StringView. Normalizes to UTF8 and LowerCase. */
+	COREUOBJECT_API static int32 ConstructSeed(FStringView String);
+
+	/**
+	 * Get the FCookRand that was declared in a containing scope on the current thread, or null if none has been
+	 * declared.
+	 */
+	COREUOBJECT_API static TRefCountPtr<FCookRand> GetThreadScope();
+
+
+private:
+	FRandomStream RandomStreamVar;
+};
+
+enum class ECookRandScope
+{
+	/** If an FCookRand is already in scope, use it, otherwise create a new one and use that. */
+	FindOrAdd,
+	/** Create a new FCookRand and put it in scope, pushing to idle any existing one. */
+	Overwrite,
+};
+
+class FCookRandScope : public FNoncopyable
+{
+public:
+	COREUOBJECT_API FCookRandScope(int32 Seed, ECookRandScope RandScope = ECookRandScope::Overwrite);
+	COREUOBJECT_API FCookRandScope(TRefCountPtr<FCookRand> CookRand, ECookRandScope RandScope = ECookRandScope::FindOrAdd);
+	COREUOBJECT_API ~FCookRandScope();
+
+private:
+	bool ShouldAddNewInCurrentScope(ECookRandScope RandScope);
+
+	TRefCountPtr<FCookRand> Previous;
+	bool bRestorePrevious = false;
+};
+
+
+///////////////////////////////////////////////////////
+// Inline implementations
+///////////////////////////////////////////////////////
+
+inline const FRandomStream& FCookRand::RandomStream() const
+{
+	return RandomStreamVar;
+}
+
+}; // namespace UE::Cook
+
+#endif // WITH_EDITOR
diff --git a/Engine/Source/Runtime/Engine/Private/ActorConstruction.cpp b/Engine/Source/Runtime/Engine/Private/ActorConstruction.cpp
--- a/Engine/Source/Runtime/Engine/Private/ActorConstruction.cpp
+++ b/Engine/Source/Runtime/Engine/Private/ActorConstruction.cpp
@@ -19,6 +19,7 @@
 #include "ScopedSuspendRerunConstructionScripts.h"
 
 #if WITH_EDITOR
+#include "Cooker/CookRand.h"
 #include "Editor.h"
 #else
 #include "Engine/World.h"
@@ -1010,6 +1011,17 @@ bool AActor::ExecuteConstruction(const FTransform& Transform, const FRotationCon
 
 void AActor::ProcessUserConstructionScript()
 {
+#if WITH_EDITOR
+	TOptional<UE::Cook::FCookRandScope> CookRandScope;
+	if (IsRunningCookCommandlet())
+	{
+		// Notify blueprint functions called by the UserConstructScript that support CookRand to return deterministic
+		// results based on our seed. Overwrite an existing CookRand if it exists so that we always get the same
+		// guids/rands even if called from other UserConstruction scripts or CookRand environments.
+		int32 Seed = UE::Cook::FCookRand::ConstructSeed(GetPathName());
+		CookRandScope.Emplace(Seed, UE::Cook::ECookRandScope::Overwrite);
+	}
+#endif
 	// Set a flag that this actor is currently running UserConstructionScript.
 	bRunningUserConstructionScript = true;
 	UserConstructionScript();
diff --git a/Engine/Source/Runtime/Engine/Private/KismetGuidLibrary.cpp b/Engine/Source/Runtime/Engine/Private/KismetGuidLibrary.cpp
--- a/Engine/Source/Runtime/Engine/Private/KismetGuidLibrary.cpp
+++ b/Engine/Source/Runtime/Engine/Private/KismetGuidLibrary.cpp
@@ -1,7 +1,13 @@
 // Copyright Epic Games, Inc. All Rights Reserved.
 
 #include "Kismet/KismetGuidLibrary.h"
+
+#include "Cooker/CookRand.h"
+#include "EngineLogs.h"
 #include "Misc/Guid.h"
+#include "Misc/PackagePath.h"
+#include "UObject/Stack.h"
+#include "UObject/UObjectGlobals.h"
 
 #include UE_INLINE_GENERATED_CPP_BY_NAME(KismetGuidLibrary)
 
@@ -46,6 +52,36 @@ void UKismetGuidLibrary::Invalidate_Guid( FGuid& InGuid )
 
 FGuid UKismetGuidLibrary::NewGuid()
 {
+#if WITH_EDITOR
+	TRefCountPtr<UE::Cook::FCookRand> CookRand = UE::Cook::FCookRand::GetThreadScope();
+	if (CookRand)
+	{
+		FPackagePath PackagePath;
+		FString ObjectPath(TEXT("<Unknown>"));
+		FString FunctionPath(TEXT("<Unknown>"));
+		FFrame* TopFrame = FFrame::GetThreadLocalTopStackFrame();
+		if (TopFrame)
+		{
+			if (TopFrame->Object)
+			{
+				FPackagePath::TryFromPackageName(TopFrame->Object->GetPackage()->GetName(), PackagePath);
+				ObjectPath = TopFrame->Object->GetPathName();
+			}
+			if (TopFrame->Node)
+			{
+				FunctionPath = TopFrame->Node->GetPathName();
+			}
+		}
+		FGuid Result = CookRand->RandomGuid();
+		UE_ASSET_LOG(LogBlueprintUserMessages, Display, PackagePath,
+			TEXT("NewGuid called by ConstructionScript during cooking. Returning a CookRand deterministic random guid.")
+			TEXT("\n\tObject: %s")
+			TEXT("\n\tFunction: %s")
+			TEXT("\n\tResult: %s"),
+			*ObjectPath, *FunctionPath, *LexToString(Result));
+		return Result;
+	}
+#endif
 	return FGuid::NewGuid();
 }
 
