Fix removal of deferred struct ops in VNI packages with merged modules

- Previously, we would assume the existence of a package for every module exposed to VNI that needs its deferred struct ops cleaned up. This does not appear to be a thing anymore.
- This change is merged-modules specific, we only go through this path when unmounting the dynamic library for such modules. Without the change we'll get crashes on re-mounting the plugin.
- Rename RemoveObjectsInModules to RemoveObjectsInPackages for clarity on what it does.

#rnx

#Synced-CL 52242993

diff --git a/Engine/Source/Runtime/CoreUObject/Private/UObject/UObjectBase.cpp b/Engine/Source/Runtime/CoreUObject/Private/UObject/UObjectBase.cpp
--- a/Engine/Source/Runtime/CoreUObject/Private/UObject/UObjectBase.cpp
+++ b/Engine/Source/Runtime/CoreUObject/Private/UObject/UObjectBase.cpp
@@ -1176,49 +1176,63 @@ void ProcessNewlyLoadedUObjects(FName InModuleName, bool bCanProcessNewlyLoadedO
 }
 #endif // !UE_WITH_CONSTINIT_UOBJECT
 
-static TArray<UPackage*> GetScriptPackagesFromModuleNames(TConstArrayView<FName> InModuleNames)
+static TArray<UPackage*> GetScriptPackagesFromModuleNames(TConstArrayView<FName> InModuleNames, TArray<FName>* OutPackageNames = nullptr)
 {
 	TArray<UPackage*> Packages;
 
 	Packages.Reserve(3 * InModuleNames.Num());
+	if (OutPackageNames)
+	{
+		OutPackageNames->Reserve(3 * InModuleNames.Num());
+	}
 
 	for (const FName ModuleName : InModuleNames)
 	{
 		// Identify module script packages. The expected package name is '/Script/ModuleName'.
-		if (UPackage* Package = FindObjectFast<UPackage>(/* Outer */ nullptr, WriteToString<256>(TEXT("/Script/"), ModuleName).ToString()))
+		const FName ModulePackageName = WriteToString<256>(TEXT("/Script/"), ModuleName).ToString();
+		if (UPackage* Package = FindObjectFast<UPackage>(/* Outer */ nullptr, ModulePackageName))
 		{
 			Packages.Add(Package);
 		}
 
 		// Identify plugin script packages. The expected package name is '/Script/PluginName'.
-		FName PluginName = IPluginManager::Get().PackageNameFromModuleName(ModuleName);
-		if (UPackage* Package = FindObjectFast<UPackage>(/* Outer */ nullptr, WriteToString<256>(TEXT("/Script/"), PluginName).ToString()))
+		const FName PluginName = IPluginManager::Get().PackageNameFromModuleName(ModuleName);
+		const FName PluginPackageName = WriteToString<256>(TEXT("/Script/"), PluginName).ToString();
+		if (UPackage* Package = FindObjectFast<UPackage>(/* Outer */ nullptr, PluginPackageName))
 		{
 			Packages.AddUnique(Package);
 		}
 
 		// Identify Verse packages. The expected package name is '/PluginName/_Verse/VNI/ModuleName'.
-		if (UPackage* Package = FindObjectFast<UPackage>(/* Outer */ nullptr, WriteToString<256>(TEXT("/"), PluginName, TEXT("/_Verse/VNI/"), ModuleName).ToString()))
+		const FName VersePackageName = WriteToString<256>(TEXT("/"), PluginName, TEXT("/_Verse/VNI/"), ModuleName).ToString();
+		if (UPackage* Package = FindObjectFast<UPackage>(/* Outer */ nullptr, VersePackageName))
 		{
 			Packages.Add(Package);
 		}
+
+		if (OutPackageNames)
+		{
+			OutPackageNames->Add(ModulePackageName);
+			OutPackageNames->AddUnique(PluginPackageName);
+			OutPackageNames->Add(VersePackageName);
+		}
 	}
 
 	return Packages;
 }
 
-static void RemoveObjectsInModules(TConstArrayView<UPackage*> Packages)
+static void RemoveObjectsInPackages(TConstArrayView<UPackage*> Packages, TArray<FName> PackageNames)
 {
 	LLM_SCOPE(ELLMTag::UObject);
 	LLM_SCOPE_BYTAG(UObject_UObjectBase);
-	DECLARE_SCOPE_CYCLE_COUNTER(TEXT("RemoveObjectsInModules"), STAT_RemoveObjectsInModules, STATGROUP_ObjectVerbose);
+	DECLARE_SCOPE_CYCLE_COUNTER(TEXT("RemoveObjectsInPackages"), STAT_RemoveObjectsInModules, STATGROUP_ObjectVerbose);
 
 	TSet<const UObject*> ObjectsToRemove;
 	const FPermanentObjectPoolExtents PermanentPool;
 
 	for (UPackage* Package : Packages)
 	{
-		UE_LOG(LogUObjectBase, Log, TEXT("RemoveObjectsInModules: removing package %s"), *Package->GetName());
+		UE_LOG(LogUObjectBase, Log, TEXT("RemoveObjectsInPackages: removing package %s"), *Package->GetName());
 
 		ForEachObjectWithPackage(Package, [&PermanentPool, &ObjectsToRemove](UObject* InObject) mutable -> bool
 		{
@@ -1239,10 +1253,14 @@ static void RemoveObjectsInModules(TConstArrayView<UPackage*> Packages)
 			return true;
 		}, true);
 
-		UScriptStruct::RemoveAllDeferredCppStructOps(*Package->GetName());
 		ObjectsToRemove.Add(Package);
 	}
 
+	for (const FName PackageName : PackageNames)
+	{
+		UScriptStruct::RemoveAllDeferredCppStructOps(PackageName);
+	}
+
 	UE::GC::Private::RemoveFromRootForUnload(ObjectsToRemove);
 }
 
@@ -1299,9 +1317,10 @@ static void RemoveCompiledInUObjects(TConstArrayView<FName> InModuleNames)
 
 	FlushAsyncLoading();
 
-	const TArray<UPackage*> Packages = GetScriptPackagesFromModuleNames(InModuleNames);
+	TArray<FName> PackageNames;
+	const TArray<UPackage*> Packages = GetScriptPackagesFromModuleNames(InModuleNames, &PackageNames);
 	FCoreUObjectDelegates::CompiledInUObjectsRemovedDelegate.Broadcast(Packages);
-	RemoveObjectsInModules(Packages);
+	RemoveObjectsInPackages(Packages, PackageNames);
 }
 
 static int32 GVarMaxObjectsNotConsideredByGC;
