Removed FArchiveRemoteReferencePatcher which is no longer required.
Wrapped remote object serialization verification code in a signle UE_WITH_REMOTEOBJECT_SERIALIZATION_VERIFICATION macro

#Synced-CL 52058383

diff --git a/Engine/Source/Runtime/CoreUObject/Private/UObject/RemoteObjectSerialization.cpp b/Engine/Source/Runtime/CoreUObject/Private/UObject/RemoteObjectSerialization.cpp
--- a/Engine/Source/Runtime/CoreUObject/Private/UObject/RemoteObjectSerialization.cpp
+++ b/Engine/Source/Runtime/CoreUObject/Private/UObject/RemoteObjectSerialization.cpp
@@ -49,6 +49,10 @@ static FAutoConsoleVariableRef CVarUseImmutableArchetypes(
 	GUseImmutableArchetypes,
 	TEXT("Whether to use immutable archetypes when serializing remote object data"));
 
+#ifndef UE_WITH_REMOTEOBJECT_SERIALIZATION_VERIFICATION
+#define UE_WITH_REMOTEOBJECT_SERIALIZATION_VERIFICATION DO_CHECK && !(UE_BUILD_TEST || UE_BUILD_SHIPPING)
+#endif
+
 #if UE_WITH_REMOTEOBJECTARCHIVE_DEBUGGING
 
 int32 GDebugRemoteObjectSerialization = 0;
@@ -1094,56 +1098,6 @@ public:
 	}
 };
 
-/**
-* Archive that replaces unresolved FObjectPtrs (TObjectPtrs<>) with actual pointers to deserialized objects
-*/
-class FArchiveRemoteReferencePatcher : public FArchiveUObject
-{
-	TMap<FRemoteObjectId, UObject*> IdToObjectMap;
-
-	void InitObjectMap(const TArray<UObject*>& DeserializedObjects)
-	{
-		for (UObject* Obj : DeserializedObjects)
-		{
-			IdToObjectMap.Add(FRemoteObjectId(Obj), Obj);
-		}
-	}
-
-	void PatchObjectReference(FObjectPtr& Value)
-	{
-#if UE_WITH_REMOTE_OBJECT_HANDLE
-		using namespace UE::RemoteObject::Handle;
-		const FRemoteObjectStub* Stub = Value.GetHandleRef().ToStub();
-		if (UObject** ResolvedObject = IdToObjectMap.Find(Stub->Id))
-		{
-			Value = FObjectPtr(*ResolvedObject);
-		}
-#endif
-	}
-
-public:
-
-	FArchiveRemoteReferencePatcher(const TArray<UObject*>& DeserializedObjects)
-	{
-		ArIsObjectReferenceCollector = true;
-#if WITH_EDITORONLY_DATA
-		SetIsSaving(true); // workaround for UE-225585 (UStruct::SerializeBinEx() asserts in structured archive code in the editor)
-#endif
-		SetIsPersistent(false);
-		InitObjectMap(DeserializedObjects);
-		SetPortFlags(PPF_AvoidRemoteObjectMigration);
-	}
-
-	virtual FArchive& operator<<(FObjectPtr& Value) override
-	{
-		if (Value.GetResidence() == EResidence::Remote)
-		{
-			PatchObjectReference(Value);
-		}
-		return *this;
-	}
-};
-
 FRemoteObjectConstructionOverrides::FRemoteObjectConstructionOverrides(const FRemoteObjectData& ObjectData, const TArray<FRemoteObjectHeader>& InObjectHeaders)
 {
 	using namespace UE::RemoteObject::Handle;
@@ -1328,7 +1282,7 @@ FORCENOINLINE void SortRemoteObjectHeadersByOuter(FArchiveRemoteObjectWriter& Ar
 	checkf(FRemoteObjectId(Ar.GetRootObject()) == Ar.GetRemoteObjectId(ObjectHeaders[0].RemoteId), TEXT("Root object id (%s) does not match the first object header id (%s)."),
 		*FRemoteObjectId(Ar.GetRootObject()).ToString(), *Ar.GetRemoteObjectId(ObjectHeaders[0].RemoteId).ToString());
 
-#if DO_CHECK && UE_WITH_REMOTE_OBJECT_HANDLE
+#if UE_WITH_REMOTEOBJECT_SERIALIZATION_VERIFICATION && UE_WITH_REMOTE_OBJECT_HANDLE
 	// Make sure the header array is sorted
 	for (int32 HeaderIndex = 0; HeaderIndex < ObjectHeaders.Num(); ++HeaderIndex)
 	{
@@ -1341,7 +1295,7 @@ FORCENOINLINE void SortRemoteObjectHeadersByOuter(FArchiveRemoteObjectWriter& Ar
 		checkf(OuterHeaderIndex < HeaderIndex, TEXT("Serialized object's[%d] %s Outer[%d] %s is serialized after the object"),
 			HeaderIndex, *Ar.GetRemoteObjectId(Header.RemoteId).ToString(ERemoteIdToStringVerbosity::PathName), OuterHeaderIndex, *OuterId.ToString(ERemoteIdToStringVerbosity::PathName));
 	}
-#endif // DO_CHECK
+#endif // UE_WITH_REMOTEOBJECT_SERIALIZATION_VERIFICATION
 }
 
 void SerializeObjectDataInternal(FArchiveRemoteObjectWriter& Ar, UObject* RequestedObject, FRemoteObjectId RequestedObjectId, TSet<UObject*>& OutObjects, ERemoteObjectSerializationFlags SerializationFlags)
@@ -1447,13 +1401,13 @@ FRemoteObjectData SerializeObjectData(UObject* InObject, TSet<UObject*>& OutObje
 		Ar.SetMigratingRemoteObjects(true);
 		if (GRemoteObjectsMigrateAllSubobjects)
 		{
-			GetObjectsWithOuter(Ar.GetRootObject(), Ar.GetObjectsToSerialize(),  /*bIncludeNestedObjects =*/ true);
+			GetObjectsWithOuter(Ar.GetRootObject(), Ar.GetObjectsToSerialize());
 		}
 		SerializeObjectDataInternal(Ar, InObject, RequestedObjectId, OutObjects, InFlags);
 	}
 	checkf(ObjectData.SerializedObjectHeaders.Num(), TEXT("Serialized object data for object %s is empty"), *RequestedObjectId.ToString());
 
-#if !(UE_BUILD_TEST || UE_BUILD_SHIPPING)
+#if UE_WITH_REMOTEOBJECT_SERIALIZATION_VERIFICATION
 	// We are doing an aggressive warning log here because these missing reference issues are hard to track down:
 	TArray<UObject*> AllSubobjects;
 	GetObjectsWithOuter(InObject, AllSubobjects);
@@ -1592,7 +1546,7 @@ int32 DeserializeObjectDataInternal(FArchiveRemoteObjectReader& Ar,
 			}
 
 			checkf(OutObjectRemoteIds.Num() == ObjectData.SerializedObjectHeaders.Num(), TEXT("The number of constructed object ids (%d) does not match the number of serialized object headers (%d)"), OutObjectRemoteIds.Num(), ObjectData.SerializedObjectHeaders.Num());
-#if DO_CHECK
+#if UE_WITH_REMOTEOBJECT_SERIALIZATION_VERIFICATION
 			for (int32 ObjectIndex = 0; ObjectIndex < OutObjectRemoteIds.Num(); ++ObjectIndex)
 			{
 				FRemoteObjectId ConstructedObjectId = OutObjectRemoteIds[ObjectIndex];
@@ -1621,7 +1575,7 @@ int32 DeserializeObjectDataInternal(FArchiveRemoteObjectReader& Ar,
 			TArray<UObject*> CreatedObjects;
 			CreatedObjects.Reserve(OutObjects.Num());
 			CreatedObjects.Add(RootObject);
-			GetObjectsWithOuter(RootObject, CreatedObjects, /*bIncludeNestedObjects =*/ true);
+			GetObjectsWithOuter(RootObject, CreatedObjects);
 
 			// ensure remote object stubs are created (ownership will be assigned after PostMigrate)
 			for (UObject* Object : CreatedObjects)
@@ -1639,7 +1593,7 @@ int32 DeserializeObjectDataInternal(FArchiveRemoteObjectReader& Ar,
 				}
 			}
 
-#if DO_CHECK
+#if UE_WITH_REMOTEOBJECT_SERIALIZATION_VERIFICATION
 			for (UObject* DeserializedObject : OutObjects)
 			{
 				checkf(!DeserializedObject || CreatedObjects.Contains(DeserializedObject), TEXT("Deserialized an object (%s) that's outside of the root object hierarchy"), *GetPathNameSafe(DeserializedObject));
@@ -1672,21 +1626,6 @@ int32 DeserializeObjectDataInternal(FArchiveRemoteObjectReader& Ar,
 				}
 			}
 		}
-
-		// Patch any unresolved remote references (skip if we're resetting an object to its CDO state)
-		if (!bResetting)
-		{
-			TRACE_CPUPROFILER_EVENT_SCOPE(FArchiveRemoteReferencePatcher);
-
-			FArchiveRemoteReferencePatcher PatchAr(OutObjects);
-			for (UObject* Object : OutObjects)
-			{
-				if (Object)
-				{
-					Object->Serialize(PatchAr);
-				}
-			}
-		}
 	}
 	checkf(RequestedObjectIndex >= 0, TEXT("Received remote object data but the requested object (%s) was not deserialized"), *RequestedObjectId.ToString());
 	return RequestedObjectIndex;
