Null check super class, in case we outer to UObject

#rnx
#Synced-CL 51986075

diff --git a/Engine/Source/Runtime/CoreUObject/Private/UObject/UObjectArchetype.cpp b/Engine/Source/Runtime/CoreUObject/Private/UObject/UObjectArchetype.cpp
--- a/Engine/Source/Runtime/CoreUObject/Private/UObject/UObjectArchetype.cpp
+++ b/Engine/Source/Runtime/CoreUObject/Private/UObject/UObjectArchetype.cpp
@@ -150,9 +150,15 @@ FORCENOINLINE UObject* FindArchetypeFromRequiredInfoImpl(const UClass* Class, co
 				Result = ArchetypeToSearch->GetClass()->FindArchetype(Class, Name);
 			}
 
-			if(!Result && (ObjectFlags & RF_ArchetypeObject) && Outer->IsA<UClass>())
+			if(!Result && (ObjectFlags & RF_ArchetypeObject))
 			{
-				Result = CastChecked<UClass>(Outer)->GetSuperClass()->FindArchetype(Class, Name);
+				if(const UClass* OuterClass = Cast<UClass>(Outer))
+				{
+					if(const UClass* SuperClass = OuterClass->GetSuperClass())
+					{
+						Result = SuperClass->FindArchetype(Class, Name);
+					}
+				}
 			}
 		}
 
