RemoteObjects: Fixing up the Region Tracing mismatch caused by AutoRTFM aborting during a transactional PostMigrate.

#rnx
#Synced-CL 52368967

diff --git a/Engine/Source/Runtime/CoreUObject/Private/UObject/RemoteObjectTransfer.cpp b/Engine/Source/Runtime/CoreUObject/Private/UObject/RemoteObjectTransfer.cpp
--- a/Engine/Source/Runtime/CoreUObject/Private/UObject/RemoteObjectTransfer.cpp
+++ b/Engine/Source/Runtime/CoreUObject/Private/UObject/RemoteObjectTransfer.cpp
@@ -1099,6 +1099,7 @@ void DeserializeObjectsFromMigratedData(FTransactionalMigrationData& InOutMigrat
 
 	checkf(InOutMigratedData.State == EMigratedDataState::Received, TEXT("Starting to deserialize remote object data but it's in an unexpected state %s (expected: %s)"), EnumToString(InOutMigratedData.State), EnumToString(EMigratedDataState::Received));
 
+	ensure(!AutoRTFM::IsTransactional());
 	static FName GetObjectRegionName("ObjectTransferQueue_RecvObject");
 	UE::RemoteExecutor::OnRegionBeginDelegate.ExecuteIfBound(GetObjectRegionName);
 
@@ -1197,6 +1198,10 @@ UObject* PostMigrateObjectsFromMigratedData(FTransactionalMigrationData& InOutMi
 
 	static FName PostMigrateRegionName("ObjectTransferQueue_PostMigrate");
 	UE::RemoteExecutor::OnRegionBeginDelegate.ExecuteIfBound(PostMigrateRegionName);
+	AutoRTFM::PushOnAbortHandler(&InOutMigratedData, []()
+	{
+		UE::RemoteExecutor::OnRegionEndDelegate.ExecuteIfBound(TEXT("Abort"));
+	});
 
 	const FRemoteServerId LocalServerId = FRemoteServerId::GetLocalServerId();
 	const FRemoteObjectData& ObjectData = InOutMigratedData.Data;
@@ -1265,6 +1270,7 @@ UObject* PostMigrateObjectsFromMigratedData(FTransactionalMigrationData& InOutMi
 		}
 	}
 
+	AutoRTFM::PopOnAbortHandler(&InOutMigratedData);
 	UE::RemoteExecutor::OnRegionEndDelegate.ExecuteIfBound(FString::Format(TEXT("{0} {1} num objects: {2}"), { *ReceivedObjects[0]->GetName(), *RootObjectRemoteId.ToString(), ReceivedObjects.Num(),}));
 
 	UObject* Result = ReceivedObjects[InOutMigratedData.RequestedObjectIndex];
diff --git a/Engine/Source/Runtime/Engine/Private/NetDriver.cpp b/Engine/Source/Runtime/Engine/Private/NetDriver.cpp
--- a/Engine/Source/Runtime/Engine/Private/NetDriver.cpp
+++ b/Engine/Source/Runtime/Engine/Private/NetDriver.cpp
@@ -6361,14 +6361,6 @@ int32 UNetDriver::ServerReplicateActors(float DeltaSeconds)
 	// Build the consider list (actors that are ready to replicate)
 	if (bUseGranularTransactions)
 	{
-		// Start the region here and end it when we're done ExecutePendingWork
-		static const FName FuncName { __func__ };
-		UE::RemoteExecutor::OnRegionBeginDelegate.ExecuteIfBound(FuncName);
-		AutoRTFM::PushOnAbortHandler(&ServerTickTime, []()
-		{
-			UE::RemoteExecutor::OnRegionEndDelegate.ExecuteIfBound(TEXT("Aborted"));
-		});
-
 		static const FName TransactionalWorkName("ServerReplicateActors_BuildConsiderList");
 		UE::RemoteExecutor::ExecuteTransactional(TransactionalWorkName, [this, ServerTickTime]()
 			{
@@ -6473,9 +6465,6 @@ int32 UNetDriver::ServerReplicateActors(float DeltaSeconds)
 	if (bUseGranularTransactions)
 	{
 		UE::RemoteExecutor::ExecutePendingWork();
-
-		UE::RemoteExecutor::OnRegionEndDelegate.ExecuteIfBound(TEXT(""));
-		AutoRTFM::PopOnAbortHandler(&ServerTickTime);
 	}
 
 	if (OnPostConsiderListUpdateOverride.IsBound())
