[Insights] ObjectTrace: Added a new TotalResourceSizeExtraSpec event (i.e. inclusive stats moved out of the ResourceSizeExtraSpec event).

#rnx
#Synced-CL 52347234

diff --git a/Engine/Source/Developer/TraceServices/Private/Analyzers/ObjectTraceAnalysis.cpp b/Engine/Source/Developer/TraceServices/Private/Analyzers/ObjectTraceAnalysis.cpp
--- a/Engine/Source/Developer/TraceServices/Private/Analyzers/ObjectTraceAnalysis.cpp
+++ b/Engine/Source/Developer/TraceServices/Private/Analyzers/ObjectTraceAnalysis.cpp
@@ -37,6 +37,7 @@ void FObjectAnalyzer::OnAnalysisBegin(const FOnAnalysisContext& Context)
 	ROUTE_EVENT(ObjectSpec);
 	ROUTE_EVENT(VerseExtraSpec);
 	ROUTE_EVENT(ResourceSizeExtraSpec);
+	ROUTE_EVENT(TotalResourceSizeExtraSpec);
 	ROUTE_EVENT(FieldExtraSpec);
 	ROUTE_EVENT(StructExtraSpec);
 	ROUTE_EVENT(ClassExtraSpec);
@@ -332,11 +333,8 @@ bool FObjectAnalyzer::OnEvent(uint16 RouteId, EStyle Style, const FOnEventContex
 		uint32 Id = EventData.GetValue<uint32>("Id");
 		uint64 SystemMemoryBytes = EventData.GetValue<uint64>("SystemMemoryBytes");
 		uint64 VideoMemoryBytes = EventData.GetValue<uint64>("VideoMemoryBytes");
-		uint64 TotalSystemMemoryBytes = EventData.GetValue<uint64>("TotalSystemMemoryBytes");
-		uint64 TotalVideoMemoryBytes = EventData.GetValue<uint64>("TotalVideoMemoryBytes");
 
-		UE_OBJECT_ANALYZER_LOG_API_L1("Id=%u, SystemMemoryBytes=%llu (Total=%llu), VideoMemoryBytes=%llu (Total=%llu)",
-			Id, SystemMemoryBytes, TotalSystemMemoryBytes, VideoMemoryBytes, TotalVideoMemoryBytes);
+		UE_OBJECT_ANALYZER_LOG_API_L1("Id=%u, SystemMemoryBytes=%llu, VideoMemoryBytes=%llu", Id, SystemMemoryBytes, VideoMemoryBytes);
 
 		FProviderEditScopeLock _(EditableProvider);
 		FObjectInfo* Object = EditableProvider.GetEditableObject(Id);
@@ -344,8 +342,34 @@ bool FObjectAnalyzer::OnEvent(uint16 RouteId, EStyle Style, const FOnEventContex
 		{
 			Object->SystemMemoryBytes = SystemMemoryBytes;
 			Object->VideoMemoryBytes = VideoMemoryBytes;
-			Object->TotalSystemMemoryBytes = TotalSystemMemoryBytes;
-			Object->TotalVideoMemoryBytes = TotalVideoMemoryBytes;
+		}
+		else
+		{
+			LogErrorForEventWithInvalidObjectId(EventName, Id);
+		}
+		break;
+	}
+
+	case RouteId_TotalResourceSizeExtraSpec:
+	{
+		UE_OBJECT_ANALYZER_EVENT("TotalResourceSizeExtraSpec");
+		if (!CheckValidSnapshotForEvent(EventName))
+		{
+			break;
+		}
+
+		uint32 Id = EventData.GetValue<uint32>("Id");
+		uint64 SystemMemoryBytes = EventData.GetValue<uint64>("SystemMemoryBytes");
+		uint64 VideoMemoryBytes = EventData.GetValue<uint64>("VideoMemoryBytes");
+
+		UE_OBJECT_ANALYZER_LOG_API_L1("Id=%u, SystemMemoryBytes=%llu, VideoMemoryBytes=%llu", Id, SystemMemoryBytes, VideoMemoryBytes);
+
+		FProviderEditScopeLock _(EditableProvider);
+		FObjectInfo* Object = EditableProvider.GetEditableObject(Id);
+		if (Object)
+		{
+			Object->TotalSystemMemoryBytes = SystemMemoryBytes;
+			Object->TotalVideoMemoryBytes = VideoMemoryBytes;
 		}
 		else
 		{
diff --git a/Engine/Source/Developer/TraceServices/Private/Analyzers/ObjectTraceAnalysis.h b/Engine/Source/Developer/TraceServices/Private/Analyzers/ObjectTraceAnalysis.h
--- a/Engine/Source/Developer/TraceServices/Private/Analyzers/ObjectTraceAnalysis.h
+++ b/Engine/Source/Developer/TraceServices/Private/Analyzers/ObjectTraceAnalysis.h
@@ -23,6 +23,7 @@ private:
 		RouteId_ObjectSpec,
 		RouteId_VerseExtraSpec,
 		RouteId_ResourceSizeExtraSpec,
+		RouteId_TotalResourceSizeExtraSpec,
 		RouteId_FieldExtraSpec,
 		RouteId_StructExtraSpec,
 		RouteId_ClassExtraSpec,
diff --git a/Engine/Source/Runtime/CoreUObject/Private/UObject/ObjectTrace.cpp b/Engine/Source/Runtime/CoreUObject/Private/UObject/ObjectTrace.cpp
--- a/Engine/Source/Runtime/CoreUObject/Private/UObject/ObjectTrace.cpp
+++ b/Engine/Source/Runtime/CoreUObject/Private/UObject/ObjectTrace.cpp
@@ -62,8 +62,12 @@ UE_TRACE_MINIMAL_EVENT_BEGIN(CoreUObject, ResourceSizeExtraSpec, NoSync)
 	UE_TRACE_MINIMAL_EVENT_FIELD(uint32, Id)
 	UE_TRACE_MINIMAL_EVENT_FIELD(uint64, SystemMemoryBytes) // exclusive estimated system memory
 	UE_TRACE_MINIMAL_EVENT_FIELD(uint64, VideoMemoryBytes) // exclusive estimated video memory
-	UE_TRACE_MINIMAL_EVENT_FIELD(uint64, TotalSystemMemoryBytes) // inclusive estimated system memory
-	UE_TRACE_MINIMAL_EVENT_FIELD(uint64, TotalVideoMemoryBytes) // inclusive estimated video memory
+UE_TRACE_MINIMAL_EVENT_END()
+
+UE_TRACE_MINIMAL_EVENT_BEGIN(CoreUObject, TotalResourceSizeExtraSpec, NoSync)
+	UE_TRACE_MINIMAL_EVENT_FIELD(uint32, Id)
+	UE_TRACE_MINIMAL_EVENT_FIELD(uint64, SystemMemoryBytes) // inclusive estimated system memory
+	UE_TRACE_MINIMAL_EVENT_FIELD(uint64, VideoMemoryBytes) // inclusive estimated video memory
 UE_TRACE_MINIMAL_EVENT_END()
 
 UE_TRACE_MINIMAL_EVENT_BEGIN(CoreUObject, FieldExtraSpec, NoSync)
@@ -247,7 +251,6 @@ bool FUObjectTrace::OutputObjectSpec(UObjectBase* Object)
 	}
 
 	constexpr bool bTraceResourceSize = true;
-	constexpr bool bTraceTotalResourceSize = false; // very slow (~20s / 1M objects) !!!
 	if (bTraceResourceSize)
 	{
 		// Exclusive Estimated (System/Video) Memory Size
@@ -256,25 +259,30 @@ bool FUObjectTrace::OutputObjectSpec(UObjectBase* Object)
 		SIZE_T SystemMemoryBytes = ResourceSize.GetDedicatedSystemMemoryBytes();
 		SIZE_T VideoMemoryBytes = ResourceSize.GetDedicatedVideoMemoryBytes();
 
-		if (SystemMemoryBytes || VideoMemoryBytes)
+		if (SystemMemoryBytes != 0 || VideoMemoryBytes != 0)
 		{
-			// Inclusive Estimated (System/Video) Memory Size
-			SIZE_T TotalSystemMemoryBytes = 0;
-			SIZE_T TotalVideoMemoryBytes = 0;
-			if (bTraceTotalResourceSize)
-			{
-				FResourceSizeEx TotalResourceSize(EResourceSizeMode::EstimatedTotal);
-				Obj->GetResourceSizeEx(TotalResourceSize);
-				TotalSystemMemoryBytes = TotalResourceSize.GetDedicatedSystemMemoryBytes();
-				TotalVideoMemoryBytes = TotalResourceSize.GetDedicatedVideoMemoryBytes();
-			}
-
 			UE_TRACE_MINIMAL_LOG(CoreUObject, ResourceSizeExtraSpec, CoreUObjectChannel)
 				<< ResourceSizeExtraSpec.Id(Id)
 				<< ResourceSizeExtraSpec.SystemMemoryBytes(SystemMemoryBytes)
-				<< ResourceSizeExtraSpec.VideoMemoryBytes(VideoMemoryBytes)
-				<< ResourceSizeExtraSpec.TotalSystemMemoryBytes(TotalSystemMemoryBytes)
-				<< ResourceSizeExtraSpec.TotalVideoMemoryBytes(TotalVideoMemoryBytes);
+				<< ResourceSizeExtraSpec.VideoMemoryBytes(VideoMemoryBytes);
+		}
+	}
+
+	constexpr bool bTraceTotalResourceSize = false; // very slow (~20s / 1M objects) !!!
+	if (bTraceTotalResourceSize)
+	{
+		// Inclusive Estimated (System/Video) Memory Size
+		FResourceSizeEx TotalResourceSize(EResourceSizeMode::EstimatedTotal);
+		Obj->GetResourceSizeEx(TotalResourceSize);
+		SIZE_T SystemMemoryBytes = TotalResourceSize.GetDedicatedSystemMemoryBytes();
+		SIZE_T VideoMemoryBytes = TotalResourceSize.GetDedicatedVideoMemoryBytes();
+
+		if (SystemMemoryBytes != 0 || VideoMemoryBytes != 0)
+		{
+			UE_TRACE_MINIMAL_LOG(CoreUObject, TotalResourceSizeExtraSpec, CoreUObjectChannel)
+				<< TotalResourceSizeExtraSpec.Id(Id)
+				<< TotalResourceSizeExtraSpec.SystemMemoryBytes(SystemMemoryBytes)
+				<< TotalResourceSizeExtraSpec.VideoMemoryBytes(VideoMemoryBytes);
 		}
 	}
 
