When instantiating instanced subobjects, ensure that the objects that have the RF_NeedPostLoad flag are reinstanced appropriately to avoid asserting against objects that aren't fully loaded yet.

#Synced-CL 52377827

diff --git a/Engine/Source/Runtime/CoreUObject/Private/UObject/CoreNative.cpp b/Engine/Source/Runtime/CoreUObject/Private/UObject/CoreNative.cpp
--- a/Engine/Source/Runtime/CoreUObject/Private/UObject/CoreNative.cpp
+++ b/Engine/Source/Runtime/CoreUObject/Private/UObject/CoreNative.cpp
@@ -296,7 +296,7 @@ UObject* FObjectInstancingGraph::GetInstancedSubobject( UObject* SourceSubobject
 								UObject* ExistingObject = StaticFindObjectFast(nullptr, SubobjectOuter, SubobjectName);
 								if (ExistingObject)
 								{
-									if (!ExistingObject->HasAnyFlags(RF_NeedLoad)
+									if (!ExistingObject->HasAnyFlags(RF_NeedLoad | RF_NeedPostLoad)
 										&& SourceSubobject->IsA(ExistingObject->GetClass()))
 									{
 										if (SourceSubobject->GetClass() == ExistingObject->GetClass())
