Fixing the logic behind storing unreachable remote references to database:

- Added EGetObjectsFlags::EvenIfUnreachable to allow including Unreachable objects in the results of GetObjectsWithOuter/ForEachObjectWithOuter functions (but only when storing remote objects to database) so that finding all subobjects when storing the unreachable root object works the same way as when performing migrations between servers
- fixed the logic for finding the unreachable root when storing objects to database. We previously relied on the RemoteReference to be set on the unreachable outer but the outer would get marked as a remote reference
when storing its subobjects to database anyway so if the code then attempted to store the root later it would end up storing its subobjects to database again (and it depended on the order of the objects in the GUObjectArray)
- removed GRemoteObjectsMigrateAllSubobjects because it's the only path being used when migrating subobjects

#Synced-CL 52380998

diff --git a/Engine/Source/Runtime/CoreUObject/Private/UObject/RemoteObjectSerialization.cpp b/Engine/Source/Runtime/CoreUObject/Private/UObject/RemoteObjectSerialization.cpp
--- a/Engine/Source/Runtime/CoreUObject/Private/UObject/RemoteObjectSerialization.cpp
+++ b/Engine/Source/Runtime/CoreUObject/Private/UObject/RemoteObjectSerialization.cpp
@@ -31,12 +31,6 @@ static FAutoConsoleVariableRef CVarRemoteObjectsMigrateFullHierarchy(
 	GRemoteObjectsMigrateFullHierarchy,
 	TEXT("Whether remote objects that are default subobjects should be always migrated with their parent objects"));
 
-int32 GRemoteObjectsMigrateAllSubobjects = 1;
-static FAutoConsoleVariableRef CVarRemoteObjectsMigrateAllSubobjects(
-	TEXT("ro.MigrateAllSubobjects"),
-	GRemoteObjectsMigrateAllSubobjects,
-	TEXT("Whether to collect all objects with the canonical root as outer for serialization (and not only the subobjects that are actually referenced by their root)"));
-
 int32 GResetBorrowedObjects = 1;
 static FAutoConsoleVariableRef CVarResetBorrowedObjects(
 	TEXT("ro.ResetBorrowedObjects"),
@@ -675,9 +669,7 @@ protected:
 
 			if (Info.Type == ERemoteReferenceType::IdOnly && Info.bIsSubobject && !EnumHasAnyFlags(RemoteSerializationFlags, ERemoteObjectSerializationFlags::Resetting))
 			{
-				// Add subobjects of the root object to the list of objects to serialize
-				ensureMsgf(Object->IsIn(RootObject), TEXT("We're about to serialize a subobject %s which is not a subobject of a root object %s"), *Object->GetPathName(), *RootObject->GetPathName());
-				ObjectsToSerialize.Add(Object);
+				checkf(ObjectsToSerialize.Contains(Object), TEXT("Serializing an object reference (%s) that is not in the ObjectsToSerialize list for root %s"), *GetPathNameSafe(Object), *GetPathNameSafe(RootObject));
 			}
 		}
 	}
@@ -1399,10 +1391,16 @@ FRemoteObjectData SerializeObjectData(UObject* InObject, TSet<UObject*>& OutObje
 	{
 		FArchiveRemoteObjectWriter Ar(Object, ObjectData, MigrationContext, nullptr, &OutReferencedObjects);
 		Ar.SetMigratingRemoteObjects(true);
-		if (GRemoteObjectsMigrateAllSubobjects)
+
+		EGetObjectsFlags GetObjectsFlags = EGetObjectsFlags::IncludeNestedObjects;
+#if UE_WITH_REMOTE_OBJECT_HANDLE
+		if (MigrationContext && MigrationContext->RemoteServerId.IsDatabase())
 		{
-			GetObjectsWithOuter(Ar.GetRootObject(), Ar.GetObjectsToSerialize());
+			GetObjectsFlags |= EGetObjectsFlags::EvenIfUnreachable;
 		}
+#endif
+		GetObjectsWithOuter(Ar.GetRootObject(), Ar.GetObjectsToSerialize(), GetObjectsFlags);
+
 		SerializeObjectDataInternal(Ar, InObject, RequestedObjectId, OutObjects, InFlags);
 	}
 	checkf(ObjectData.SerializedObjectHeaders.Num(), TEXT("Serialized object data for object %s is empty"), *RequestedObjectId.ToString());
diff --git a/Engine/Source/Runtime/CoreUObject/Private/UObject/RemoteObjectTransfer.cpp b/Engine/Source/Runtime/CoreUObject/Private/UObject/RemoteObjectTransfer.cpp
--- a/Engine/Source/Runtime/CoreUObject/Private/UObject/RemoteObjectTransfer.cpp
+++ b/Engine/Source/Runtime/CoreUObject/Private/UObject/RemoteObjectTransfer.cpp
@@ -61,6 +61,8 @@ bool FRemoteObjectData::Compare(const FRemoteObjectData& InOtherData) const
 	return true;
 }
 
+extern bool GAllowGetObjectsReturnUnreachableObjects;
+
 namespace UE::RemoteObject::Transfer
 {
 const FRemoteServerId DatabaseId(ERemoteServerIdConstants::Database);
@@ -1672,16 +1674,21 @@ void InitRemoteObjectTransfer(FRemoteExecutor* Executor)
 namespace UE::RemoteObject::Transfer::Private
 {
 
-UObject* GetOutermostUnreachableRemoteReference(UObject* Object)
+UObject* GetOutermostUnreachableRootObject(UObject* Object)
 {
-	UObject* OutermostUnreachableRemoteReference = Object;
+	UObject* OutermostUnreachableObject = Object;
 #if UE_WITH_REMOTE_OBJECT_HANDLE
 	for (UObject* Outer = Object->GetOuter(); Outer; Outer = Outer->GetOuter())
 	{
 		FUObjectItem* OuterItem = GUObjectArray.ObjectToObjectItem(Outer);
-		if (OuterItem->HasAllFlags(EInternalObjectFlags::RemoteReference | EInternalObjectFlags::Unreachable))
+		if (OuterItem->HasAllFlags(EInternalObjectFlags::Unreachable))
 		{
-			OutermostUnreachableRemoteReference = Outer;
+			OutermostUnreachableObject = Outer;
+			if (OutermostUnreachableObject->IsMigrationRoot())
+			{
+				// We don't want to go deeper than the migration root
+				break;
+			}
 		}
 		else
 		{
@@ -1689,7 +1696,7 @@ UObject* GetOutermostUnreachableRemoteReference(UObject* Object)
 		}
 	}
 #endif
-	return OutermostUnreachableRemoteReference;
+	return OutermostUnreachableObject;
 }
 
 void StoreObjectToDatabase(UObject* Object, FRemoteObjectId ObjectId)
@@ -1717,20 +1724,24 @@ void StoreObjectToDatabase(UObject* Object, FRemoteObjectId ObjectId)
 		return;
 	}
 
+	UE_LOGF(LogRemoteObject, VeryVerbose, "Storing %ls %ls to database", *ObjectId.ToString(), *GetPathNameSafe(Object));
 	do
 	{
 		// Object could be a subobject of its outer that is going to be destroyed later (and have this function invoked on)
 		// so to make storing to database (disk) consistent with how objects are migrated (root object with its subobjects in one transfer), 
-		// find this object's outermost object that is also unreachable and marked as remote reference.
-		// If we re-entered this loop because Object has not been stored to disk with its OutermostUnreachableRemoteReference
-		// we still want to try and get an Outer that might also have not been stored. If such outer does not exist GetOutermostUnreachableRemoteReference(Object) will return Object itself.
-		UObject* OutermostUnreachableRemoteReference = GetOutermostUnreachableRemoteReference(Object);
+		// find this object's outermost object that is also unreachable (but we never go deeper than the migration root).
+		// The Outer does not need to be marked as a remote reference because the act of storing its subobjects to database will mark it as such anyway and in this case we can end up
+		// storing the Outer with its subobjects to database again. We also don't want to only rely on the remote reference flag when searching for the root to store to database because if the unreachable
+		// outer was not marked as a remote reference and we only stored its subobjects to database then we wouldn't be able to restore them because the outer would never get stored to database.
+		// If we re-entered this loop because Object has not been stored to disk with its OutermostUnreachableObject
+		// we still want to try and get an Outer that might also have not been stored. If such outer does not exist GetOutermostUnreachableRootObject(Object) will return Object itself.
+		UObject* OutermostUnreachableObject = GetOutermostUnreachableRootObject(Object);
 
 		TSet<UObject*> SentObjects;
 		TSet<UObject*> ReferencedObjects;
 		{
 			// Let's setup a fake migration context for storing this
-			FRemoteObjectId RootObjectId = FRemoteObjectId { OutermostUnreachableRemoteReference };
+			FRemoteObjectId RootObjectId = FRemoteObjectId { OutermostUnreachableObject };
 
 
 			// Setup a migration context for "push migrating to the new owner: the Database Server"
@@ -1744,7 +1755,7 @@ void StoreObjectToDatabase(UObject* Object, FRemoteObjectId ObjectId)
 
 			// Do the actual serialization but skip canonical root object search since we've already done this with the outermost unreachable reference and
 			// we don't want to find the root again if we're serializing a subobject that's unreachable from its root (Outer)
-			FRemoteObjectData ObjectData = SerializeObjectData(OutermostUnreachableRemoteReference, SentObjects, ReferencedObjects, &StoreToDatabaseContext, ERemoteObjectSerializationFlags::SkipCanonicalRootSearch);
+			FRemoteObjectData ObjectData = SerializeObjectData(OutermostUnreachableObject, SentObjects, ReferencedObjects, &StoreToDatabaseContext, ERemoteObjectSerializationFlags::SkipCanonicalRootSearch);
 #if DO_CHECK
 			// Make sure none of the sent objects is marked as remote at this point, otherwise we either stored a remote (an outdated) object memory to database or we stored the object to database twice
 			for (UObject* SentObject : SentObjects)
@@ -1752,10 +1763,12 @@ void StoreObjectToDatabase(UObject* Object, FRemoteObjectId ObjectId)
 				FRemoteObjectId SentObjectId(SentObject);
 				if (FRemoteObjectStub* Stub = FindRemoteObjectStub(SentObjectId))
 				{
-					checkf(GetResidence(SentObject) != EResidence::Remote, TEXT("Double store to database for %s? Object is remote on server %s"), //-V547
-						*SentObjectId.ToString(ERemoteIdToStringVerbosity::PathName),
+					checkf(GetResidence(SentObject) != EResidence::Remote, TEXT("Double store to database for %s (%s)? Object is remote on server %s"), //-V547
+						*SentObjectId.ToString(ERemoteIdToStringVerbosity::Id),
+						*GetPathNameSafe(SentObject),
 						*Stub->ResidentServerId.ToString());
 				}
+				UE_LOGF(LogRemoteObject, VeryVerbose, "\tStored %ls %ls to database", *SentObjectId.ToString(), *GetPathNameSafe(SentObject));
 			}
 #endif
 
@@ -1779,18 +1792,20 @@ void StoreObjectToDatabase(UObject* Object, FRemoteObjectId ObjectId)
 			SetResidence(SentObject, EResidence::Remote, DatabaseId);
 		}
 
-		// It's possible that OutermostUnreachableRemoteReference no longer referenced this Object and Object hasn't been stored to disk
+		// It's possible that OutermostUnreachableObject no longer referenced this Object and Object hasn't been stored to disk
 		// in which case we need to repeat this loop until Object is marked as remote (and is no longer marked as RemoteReference).
 		
 		// Assert if we're about to enter an infinite loop.
-		checkf(!IsRemoteReference(Object) || OutermostUnreachableRemoteReference != Object,
+		checkf(!IsRemoteReference(Object) || OutermostUnreachableObject != Object,
 			TEXT("PeerId:%s failed to store %s to database"), *FRemoteServerId::GetLocalServerId().ToString(), *Object->GetPathName());
 
 	} while (IsRemoteReference(Object));
 }
 
 void StoreUnreachableRemoteObjectsToDatabase(const TArrayView<FUObjectItem*>& UnreachableObjects)
-{
+{	
+	TGuardValue<bool> GuardAllowUnreachableObjects(GAllowGetObjectsReturnUnreachableObjects, true);
+
 	for (const FUObjectItem* ObjectItem : UnreachableObjects)
 	{
 		if (ObjectItem->HasAnyFlags(EInternalObjectFlags::RemoteReference) && !ObjectItem->HasAnyFlags(EInternalObjectFlags::Remote))
diff --git a/Engine/Source/Runtime/CoreUObject/Private/UObject/UObjectHash.cpp b/Engine/Source/Runtime/CoreUObject/Private/UObject/UObjectHash.cpp
--- a/Engine/Source/Runtime/CoreUObject/Private/UObject/UObjectHash.cpp
+++ b/Engine/Source/Runtime/CoreUObject/Private/UObject/UObjectHash.cpp
@@ -1564,6 +1564,23 @@ static FAutoConsoleCommand ShrinkUObjectHashTablesCmd(
 	FConsoleCommandWithArgsDelegate::CreateStatic(&ShrinkUObjectHashTablesDel)
 );
 
+bool GAllowGetObjectsReturnUnreachableObjects = false;
+
+static inline void AddMandatoryInternalExclusionFlags(EInternalObjectFlags& InOutExclusionFlags)
+{
+	InOutExclusionFlags |= GetMandatoryInternalExclusionFlags();
+}
+static inline void ConditionallyRemoveUnreachableFlag(EInternalObjectFlags& InOutExclusionFlags, EGetObjectsFlags Flags)
+{
+#if UE_WITH_REMOTE_OBJECT_HANDLE
+	if (EnumHasAnyFlags(Flags, EGetObjectsFlags::EvenIfUnreachable))
+	{
+		checkf(GAllowGetObjectsReturnUnreachableObjects, TEXT("EGetObjectsFlags::EvenIfUnreachable specified but it's not allowed in the current context"));
+		InOutExclusionFlags &= ~(EInternalObjectFlags::Unreachable);
+	}
+#endif
+}
+
 void GetObjectsWithOuter(const class UObjectBase* Outer, TArray<UObject *>& Results, EGetObjectsFlags Flags, EObjectFlags ExclusionFlags, EInternalObjectFlags ExclusionInternalFlags)
 {
 	checkf(Outer != nullptr, TEXT("Getting objects with a null outer is no longer supported. If you want to get all packages you might consider using GetObjectsOfClass instead."));
@@ -1580,7 +1597,8 @@ void GetObjectsWithOuter(const class UObjectBase* Outer, TArray<UObject *>& Resu
 #endif
 
 	// We don't want to return any objects that are currently being background loaded unless we're using the object iterator during async loading.
-	ExclusionInternalFlags |= GetMandatoryInternalExclusionFlags();
+	AddMandatoryInternalExclusionFlags(ExclusionInternalFlags);
+	ConditionallyRemoveUnreachableFlag(ExclusionInternalFlags, Flags);
 	const bool bIncludeNestedObjects = EnumHasAnyFlags(Flags, EGetObjectsFlags::IncludeNestedObjects);
 
 	int32 StartNum = Results.Num();
@@ -1647,7 +1665,8 @@ void ForEachObjectWithOuterBreakable(const class UObjectBase* Outer, TFunctionRe
 #endif
 
 	// We don't want to return any objects that are currently being background loaded unless we're using the object iterator during async loading.
-	ExclusionInternalFlags |= GetMandatoryInternalExclusionFlags();
+	AddMandatoryInternalExclusionFlags(ExclusionInternalFlags);
+	ConditionallyRemoveUnreachableFlag(ExclusionInternalFlags, Flags);
 	const bool bIncludeNestedObjects = EnumHasAnyFlags(Flags, EGetObjectsFlags::IncludeNestedObjects);
 
 	FUObjectHashTables& ThreadHash = FUObjectHashTables::Get();
@@ -1745,7 +1764,8 @@ void ForEachObjectWithPackage(const class UPackage* Package, TFunctionRef<bool(U
 	check(Package != nullptr);
 
 	// We don't want to return any objects that are currently being background loaded unless we're using the object iterator during async loading.
-	ExclusionInternalFlags |= GetMandatoryInternalExclusionFlags();
+	AddMandatoryInternalExclusionFlags(ExclusionInternalFlags);
+	ConditionallyRemoveUnreachableFlag(ExclusionInternalFlags, Flags);
 	const bool bIncludeNestedObjects = EnumHasAnyFlags(Flags, EGetObjectsFlags::IncludeNestedObjects);
 
 	FUObjectHashTables& ThreadHash = FUObjectHashTables::Get();
diff --git a/Engine/Source/Runtime/CoreUObject/Public/UObject/FindObjectFlags.h b/Engine/Source/Runtime/CoreUObject/Public/UObject/FindObjectFlags.h
--- a/Engine/Source/Runtime/CoreUObject/Public/UObject/FindObjectFlags.h
+++ b/Engine/Source/Runtime/CoreUObject/Public/UObject/FindObjectFlags.h
@@ -21,6 +21,9 @@ enum class EGetObjectsFlags
 
 	/** If specified, then objects whose outers directly or indirectly have Outer as an outer are included, these are the nested objects */
 	IncludeNestedObjects = 1 << 0,
+
+	/** Whether to include Unreachable objects (unsafe, only allowed inside of the garbage collector) */
+	EvenIfUnreachable = 1 << 1
 };
 ENUM_CLASS_FLAGS(EGetObjectsFlags)
 
