Make merged modules leaks a fatal error in test builds

These leaks will rarely get reported without this step, and they will often cause crashes down the road even if they currently prevent unmounting.

#rnx
#Synced-CL 52374856

diff --git a/Engine/Source/Runtime/CoreUObject/Private/UObject/UObjectBase.cpp b/Engine/Source/Runtime/CoreUObject/Private/UObject/UObjectBase.cpp
--- a/Engine/Source/Runtime/CoreUObject/Private/UObject/UObjectBase.cpp
+++ b/Engine/Source/Runtime/CoreUObject/Private/UObject/UObjectBase.cpp
@@ -1299,6 +1299,7 @@ static bool CheckLiveObjectsInModules(TConstArrayView<FName> InModuleNames)
 	{
 #if !UE_BUILD_SHIPPING
 		FReferenceChainSearch::FindAndPrintStaleReferencesToObjects(ObjectsToRemoveStillLive, EPrintStaleReferencesOptions::Log);
+		UE_LOG(LogUObjectBase, Fatal, TEXT("CheckLiveObjectsInModules: all marked objects need to be removed after GC"));
 #endif // !UE_BUILD_SHIPPING
 		return true;
 	}
