Ensure FPersistentObjectRef respects IsCapturingAsRootObjectForTransaction

IsCapturingAsRootObjectForTransaction means that the object knows better than us that it's safe to be a root, and we should respect that rather than continue

#rnx
#Synced-CL 52335712

diff --git a/Engine/Source/Runtime/CoreUObject/Private/Misc/TransactionCommon.cpp b/Engine/Source/Runtime/CoreUObject/Private/Misc/TransactionCommon.cpp
--- a/Engine/Source/Runtime/CoreUObject/Private/Misc/TransactionCommon.cpp
+++ b/Engine/Source/Runtime/CoreUObject/Private/Misc/TransactionCommon.cpp
@@ -17,11 +17,6 @@ FPersistentObjectRef::FPersistentObjectRef(UObject* InObject)
 	{
 		auto UseOuter = [](const UObject* Obj)
 		{
-			if (Obj == nullptr || Obj->IsCapturingAsRootObjectForTransaction())
-			{
-				return false;
-			}
-
 			const bool bIsCDO = Obj->HasAllFlags(RF_ClassDefaultObject);
 			const UObject* CDO = bIsCDO ? Obj : nullptr;
 			const bool bIsClassCDO = (CDO != nullptr) ? (CDO->GetClass()->GetDefaultObject(false) == CDO) : false;
@@ -50,6 +45,15 @@ FPersistentObjectRef::FPersistentObjectRef(UObject* InObject)
 		// First pass to find the correct root object
 		for (UObject* CurObject = InObject; CurObject; CurObject = CurObject->GetOuter())
 		{
+			if (CurObject->IsCapturingAsRootObjectForTransaction())
+			{
+				// Explicitly a root object; stop walking
+				if (!RootObject)
+				{
+					RootObject = CurObject;
+				}
+				break;
+			}
 			if (UseOuter(CurObject))
 			{
 				// This reference is unstable, so reset any current root object and keep going
