Appeased static analysis

#rnx

#Synced-CL 52335712

diff --git a/Engine/Source/Runtime/CoreUObject/Private/Misc/TransactionCommon.cpp b/Engine/Source/Runtime/CoreUObject/Private/Misc/TransactionCommon.cpp
--- a/Engine/Source/Runtime/CoreUObject/Private/Misc/TransactionCommon.cpp
+++ b/Engine/Source/Runtime/CoreUObject/Private/Misc/TransactionCommon.cpp
@@ -17,6 +17,8 @@ FPersistentObjectRef::FPersistentObjectRef(UObject* InObject)
 	{
 		auto UseOuter = [](const UObject* Obj)
 		{
+			check(Obj);
+
 			const bool bIsCDO = Obj->HasAllFlags(RF_ClassDefaultObject);
 			const UObject* CDO = bIsCDO ? Obj : nullptr;
 			const bool bIsClassCDO = (CDO != nullptr) ? (CDO->GetClass()->GetDefaultObject(false) == CDO) : false;
