Change post-plugin GC back to a full purge to fix issues with leaks using incremental reachability
#Synced-CL 52444556

diff --git a/Engine/Source/Runtime/CoreUObject/Private/Misc/CoreUObjectPluginManager.cpp b/Engine/Source/Runtime/CoreUObject/Private/Misc/CoreUObjectPluginManager.cpp
--- a/Engine/Source/Runtime/CoreUObject/Private/Misc/CoreUObjectPluginManager.cpp
+++ b/Engine/Source/Runtime/CoreUObject/Private/Misc/CoreUObjectPluginManager.cpp
@@ -342,7 +342,8 @@ namespace UE::CoreUObject::Private
 		if (!DeferredGCDelegate.IsValid())
 		{
 			DeferredGCDelegate = FTSTicker::GetCoreTicker().AddTicker(FTickerDelegate::CreateLambda([this](float){
-				CollectGarbage(GARBAGE_COLLECTION_KEEPFLAGS, false);
+				const bool bPerformFullPurge = true;
+				CollectGarbage(GARBAGE_COLLECTION_KEEPFLAGS, bPerformFullPurge);
 				HandlePossibleAssetLeaks(DeferredPluginsToGC);
 				DeferredPluginsToGC.Empty();
 				DeferredGCDelegate.Reset();
