[CrashFix] Loading a project with placeholder components causes a crash because the placeholder class isn't tagged VCLASS_FinalSuper

Cause: CL-52310122 made placeholders capable of keeping the super class of component or entity and reduced restrictions that hid them from component iterators but it didn't update the creation code of placeholders in LinkerLoad to apply the required VCLASS_FinalSuper | VCLASS_Castable flags.

#Synced-CL 52427606

diff --git a/Engine/Source/Runtime/CoreUObject/Private/UObject/PropertyBagRepository.cpp b/Engine/Source/Runtime/CoreUObject/Private/UObject/PropertyBagRepository.cpp
--- a/Engine/Source/Runtime/CoreUObject/Private/UObject/PropertyBagRepository.cpp
+++ b/Engine/Source/Runtime/CoreUObject/Private/UObject/PropertyBagRepository.cpp
@@ -27,6 +27,7 @@
 #include "UObject/OverridableManager.h"
 #include "UObject/PropertyPathNameTree.h"
 #include "Templates/UnrealTemplate.h"
+#include "VerseVM/VVMVerseClass.h"
 
 #if WITH_EDITOR
 #include "HAL/IConsoleManager.h"
@@ -1087,7 +1088,12 @@ UStruct* FPropertyBagRepository::CreatePropertyBagPlaceholderType(UObject* Outer
 		// Required by garbage collection for class types.
 		PlaceholderTypeAsClass->AssembleReferenceTokenStream();
 	}
-
+	// Extra configuration needed for verse class types.
+	if (UVerseClass* PlaceholderTypeAsVerseClass = Cast<UVerseClass>(PlaceholderType))
+	{
+		PlaceholderTypeAsVerseClass->SolClassFlags |= VCLASS_FinalSuper | VCLASS_Castable | VCLASS_Concrete;
+	}
+	
 	// Use the property bag repository for now to register property bag placeholder types for query purposes.
 	// Note: Object lifetimes of this type and its instances depend on existing references that are serialized.
 	FPropertyBagPlaceholderTypeRegistry::Get().Add(PlaceholderType);
