Cleaned up confusing code in FPersistentObjectRef

#rnx

#Synced-CL 52465190

diff --git a/Engine/Source/Runtime/CoreUObject/Private/Misc/TransactionCommon.cpp b/Engine/Source/Runtime/CoreUObject/Private/Misc/TransactionCommon.cpp
--- a/Engine/Source/Runtime/CoreUObject/Private/Misc/TransactionCommon.cpp
+++ b/Engine/Source/Runtime/CoreUObject/Private/Misc/TransactionCommon.cpp
@@ -19,17 +19,9 @@ FPersistentObjectRef::FPersistentObjectRef(UObject* InObject)
 		{
 			check(Obj);
 
-			const bool bIsCDO = Obj->HasAllFlags(RF_ClassDefaultObject);
-			const UObject* CDO = bIsCDO ? Obj : nullptr;
-			const bool bIsClassCDO = (CDO != nullptr) ? (CDO->GetClass()->GetDefaultObject(false) == CDO) : false;
-			if (!bIsClassCDO && CDO)
-			{
-				// Likely a trashed CDO, try to recover
-				// Only known cause of this is ambiguous use of DSOs
-				CDO = CDO->GetClass()->GetDefaultObject(false);
-			}
+			const bool bIsClassCDO = Obj->HasAllFlags(RF_ClassDefaultObject) && Obj->GetClass()->GetDefaultObject(false) == Obj; // Guard against a trashed CDO
 			const bool bIsDSO = Obj->HasAnyFlags(RF_DefaultSubObject);
-			if ((bIsCDO && bIsClassCDO) || bIsDSO)
+			if (bIsClassCDO || bIsDSO)
 			{
 				return true;
 			}
@@ -41,6 +33,7 @@ FPersistentObjectRef::FPersistentObjectRef(UObject* InObject)
 					return true;
 				}
 			}
+
 			return false;
 		};
 
