[CoreUObject]

Add lock around initialization of FMetaData/UMetaData static redirect maps as these can be initialized during serialization which could occur from multiple threads.

#rnx

#Synced-CL 52461870

diff --git a/Engine/Source/Runtime/CoreUObject/Private/UObject/MetaData.cpp b/Engine/Source/Runtime/CoreUObject/Private/UObject/MetaData.cpp
--- a/Engine/Source/Runtime/CoreUObject/Private/UObject/MetaData.cpp
+++ b/Engine/Source/Runtime/CoreUObject/Private/UObject/MetaData.cpp
@@ -2,6 +2,7 @@
 
 #include "UObject/MetaData.h"
 #include "Misc/ConfigCacheIni.h"
+#include "Misc/ScopeLock.h"
 #include "HAL/IConsoleManager.h"
 #include "Algo/Transform.h"
 #include "UObject/Linker.h"
@@ -121,6 +122,7 @@ FMetaDataUtilities::FMoveMetadataHelperContext::~FMoveMetadataHelperContext()
 // UDEPRECATED_MetaData implementation.
 
 TMap<FName, FName> UDEPRECATED_MetaData::KeyRedirectMap;
+UE::FMutex UDEPRECATED_MetaData::KeyRedirectMapLock;
 
 UDEPRECATED_MetaData::UDEPRECATED_MetaData(class FObjectInitializer const& ObjectInitializer)
 	: Super(ObjectInitializer)
@@ -265,14 +267,19 @@ void UDEPRECATED_MetaData::Serialize(FArchive& Ar)
 
 void UDEPRECATED_MetaData::InitializeRedirectMap()
 {
-	static bool bKeyRedirectMapInitialized = false;
+	static std::atomic<bool> bKeyRedirectMapInitialized = false;
 
-	if (!bKeyRedirectMapInitialized)
+	if (!bKeyRedirectMapInitialized.load(std::memory_order_acquire))
 	{
 		if (GConfig)
 		{
-			const FName MetadataRedirectsName("MetadataRedirects");
+			UE::TScopeLock _(KeyRedirectMapLock);
+			if (bKeyRedirectMapInitialized.load(std::memory_order_relaxed))
+			{
+				return;
+			}
 
+			const FName MetadataRedirectsName("MetadataRedirects");
 			const FConfigSection* PackageRedirects = GConfig->GetSection(TEXT("CoreUObject.Metadata"), false, GEngineIni);
 			if (PackageRedirects)
 			{
@@ -296,7 +303,7 @@ void UDEPRECATED_MetaData::InitializeRedirectMap()
 					}			
 				}
 			}
-			bKeyRedirectMapInitialized = true;
+			bKeyRedirectMapInitialized.store(true, std::memory_order_release);
 		}
 	}
 }
@@ -306,6 +313,7 @@ void UDEPRECATED_MetaData::InitializeRedirectMap()
 // FMetaData implementation.
 
 TMap<FName, FName> FMetaData::KeyRedirectMap;
+UE::FMutex FMetaData::KeyRedirectMapLock;
 
 void FMetaData::RemapObjectKeys(FName OldPackageName, FName NewPackageName)
 {
@@ -551,14 +559,18 @@ FName FMetaData::GetRemappedKeyName(FName OldKey)
 
 void FMetaData::InitializeRedirectMap()
 {
-	static bool bKeyRedirectMapInitialized = false;
+	static std::atomic<bool> bKeyRedirectMapInitialized = false;
 
-	if (!bKeyRedirectMapInitialized)
+	if (!bKeyRedirectMapInitialized.load(std::memory_order_acquire))
 	{
 		if (GConfig)
 		{
+			UE::TScopeLock _(KeyRedirectMapLock);
+			if (bKeyRedirectMapInitialized.load(std::memory_order_relaxed))
+			{
+				return;
+			}
 			const FName MetadataRedirectsName("MetadataRedirects");
-
 			const FConfigSection* PackageRedirects = GConfig->GetSection(TEXT("CoreUObject.Metadata"), false, GEngineIni);
 			if (PackageRedirects)
 			{
@@ -582,7 +594,7 @@ void FMetaData::InitializeRedirectMap()
 					}			
 				}
 			}
-			bKeyRedirectMapInitialized = true;
+			bKeyRedirectMapInitialized.store(true, std::memory_order_release);
 		}
 	}
 }
diff --git a/Engine/Source/Runtime/CoreUObject/Public/UObject/MetaData.h b/Engine/Source/Runtime/CoreUObject/Public/UObject/MetaData.h
--- a/Engine/Source/Runtime/CoreUObject/Public/UObject/MetaData.h
+++ b/Engine/Source/Runtime/CoreUObject/Public/UObject/MetaData.h
@@ -2,6 +2,7 @@
 
 #pragma once
 
+#include "Async/Mutex.h"
 #include "Containers/Map.h"
 #include "Containers/UnrealString.h"
 #include "HAL/Platform.h"
@@ -59,6 +60,7 @@ private:
 private:
 	// Redirect map from deprecated keys to current key names
 	static COREUOBJECT_API TMap<FName, FName> KeyRedirectMap;
+	static COREUOBJECT_API UE::FMutex KeyRedirectMapLock;
 };
 
 #if WITH_METADATA
@@ -219,6 +221,7 @@ private:
 
 	// Redirect map from deprecated keys to current key names
 	static COREUOBJECT_API TMap<FName, FName> KeyRedirectMap;
+	static COREUOBJECT_API UE::FMutex KeyRedirectMapLock;
 };
 
 struct FMetaDataUtilities
