Take gc.MaxObjectsInGame CVar to account

This allows to override the .ini value, also unlocking it to be configrulable

#rnx
#Synced-CL 52463319

diff --git a/Engine/Source/Runtime/CoreUObject/Private/UObject/UObjectBase.cpp b/Engine/Source/Runtime/CoreUObject/Private/UObject/UObjectBase.cpp
--- a/Engine/Source/Runtime/CoreUObject/Private/UObject/UObjectBase.cpp
+++ b/Engine/Source/Runtime/CoreUObject/Private/UObject/UObjectBase.cpp
@@ -1340,11 +1340,13 @@ static FAutoConsoleVariableRef CMaxObjectsInEditor(
 	ECVF_Default
 	);
 
-static int32 GMaxObjectsInGame;
+static int32 GMaxObjectsInGame = 0;
 static FAutoConsoleVariableRef CMaxObjectsInGame(
 	TEXT("gc.MaxObjectsInGame"),
 	GMaxObjectsInGame,
-	TEXT("Placeholder console variable, currently not used in runtime."),
+	TEXT("Allows to override the maximum number of UObjects that can exist in cooked game\n")
+	TEXT("- 0: No override, .ini value is used (default).\n")
+	TEXT("- any other value overrides the runtime max number"),
 	ECVF_Default
 	);
 
@@ -1379,6 +1381,10 @@ void UObjectBaseInit()
 
 		// Maximum number of UObjects in cooked game
 		GConfig->GetInt(TEXT("/Script/Engine.GarbageCollectionSettings"), TEXT("gc.MaxObjectsInGame"), MaxUObjects, GEngineIni);
+		if (GMaxObjectsInGame > 0)
+		{
+			MaxUObjects = GMaxObjectsInGame;
+		}
 
 		// If true, the UObjectArray will pre-allocate all entries for UObject pointers
 		GConfig->GetBool(TEXT("/Script/Engine.GarbageCollectionSettings"), TEXT("gc.PreAllocateUObjectArray"), bPreAllocateUObjectArray, GEngineIni);
