Switched dynamic delegate tests over to using int16 for an upcoming change.

diff --git a/Engine/Source/Runtime/CoreUObject/Public/UObject/DynamicDelegateTestObject.cpp b/Engine/Source/Runtime/CoreUObject/Public/UObject/DynamicDelegateTestObject.cpp
--- a/Engine/Source/Runtime/CoreUObject/Public/UObject/DynamicDelegateTestObject.cpp
+++ b/Engine/Source/Runtime/CoreUObject/Public/UObject/DynamicDelegateTestObject.cpp
@@ -12,8 +12,8 @@
 // Set by the name of the function to say which one was called.
 FAnsiString GCalledFunctionName;
 
-int32 GIntArg = 5;
-int32 GIntOut = 97;
+int16 GInt16Arg = 5;
+int16 GInt16Out = 97;
 
 FString GStrArg = "Hello";
 FString GStrOut = "Goodbye";
@@ -43,62 +43,62 @@ void UDynamicDelegateTestObject::Func0()
 	GCalledFunctionName = "Func0";
 }
 
-void UDynamicDelegateTestObject::Func1(int Int)
+void UDynamicDelegateTestObject::Func1(int16 Int)
 {
 	CHECK(GCalledFunctionName.IsEmpty());
 	GCalledFunctionName = "Func1";
 
-	CHECK(Int == GIntArg);
+	CHECK(Int == GInt16Arg);
 }
 
-void UDynamicDelegateTestObject::Func2(int Int, FString Str)
+void UDynamicDelegateTestObject::Func2(int16 Int, FString Str)
 {
 	CHECK(GCalledFunctionName.IsEmpty());
 	GCalledFunctionName = "Func2";
 
-	CHECK(Int == GIntArg);
+	CHECK(Int == GInt16Arg);
 	CHECK(Str == GStrArg);
 }
 
-void UDynamicDelegateTestObject::Func3(int Int, FString Str, const UDynamicDelegateTestObject* Ptr)
+void UDynamicDelegateTestObject::Func3(int16 Int, FString Str, const UDynamicDelegateTestObject* Ptr)
 {
 	CHECK(GCalledFunctionName.IsEmpty());
 	GCalledFunctionName = "Func3";
 
-	CHECK(Int == GIntArg);
+	CHECK(Int == GInt16Arg);
 	CHECK(Str == GStrArg);
 	CHECK(Ptr == GPtrArg);
 }
 
-void UDynamicDelegateTestObject::Func4(int Int, FString Str, const UDynamicDelegateTestObject* Ptr, FName Name)
+void UDynamicDelegateTestObject::Func4(int16 Int, FString Str, const UDynamicDelegateTestObject* Ptr, FName Name)
 {
 	CHECK(GCalledFunctionName.IsEmpty());
 	GCalledFunctionName = "Func4";
 
-	CHECK(Int  == GIntArg);
+	CHECK(Int  == GInt16Arg);
 	CHECK(Str  == GStrArg);
 	CHECK(Ptr  == GPtrArg);
 	CHECK(Name == GNameArg);
 }
 
-void UDynamicDelegateTestObject::Func5(int Int, FString Str, const UDynamicDelegateTestObject* Ptr, FName Name, bool bBool)
+void UDynamicDelegateTestObject::Func5(int16 Int, FString Str, const UDynamicDelegateTestObject* Ptr, FName Name, bool bBool)
 {
 	CHECK(GCalledFunctionName.IsEmpty());
 	GCalledFunctionName = "Func5";
 
-	CHECK(Int   == GIntArg);
+	CHECK(Int   == GInt16Arg);
 	CHECK(Str   == GStrArg);
 	CHECK(Ptr   == GPtrArg);
 	CHECK(Name  == GNameArg);
 	CHECK(bBool == bGBoolArg);
 }
 
-void UDynamicDelegateTestObject::Func6(int Int, FString Str, const UDynamicDelegateTestObject* Ptr, FName Name, bool bBool, TWeakObjectPtr<UDynamicDelegateTestObject> WeakObjPtr)
+void UDynamicDelegateTestObject::Func6(int16 Int, FString Str, const UDynamicDelegateTestObject* Ptr, FName Name, bool bBool, TWeakObjectPtr<UDynamicDelegateTestObject> WeakObjPtr)
 {
 	CHECK(GCalledFunctionName.IsEmpty());
 	GCalledFunctionName = "Func6";
 
-	CHECK(Int        == GIntArg);
+	CHECK(Int        == GInt16Arg);
 	CHECK(Str        == GStrArg);
 	CHECK(Ptr        == GPtrArg);
 	CHECK(Name       == GNameArg);
@@ -112,62 +112,62 @@ void UDynamicDelegateTestObject::ConstFunc0() const
 	GCalledFunctionName = "ConstFunc0";
 }
 
-void UDynamicDelegateTestObject::ConstFunc1(int Int) const
+void UDynamicDelegateTestObject::ConstFunc1(int16 Int) const
 {
 	CHECK(GCalledFunctionName.IsEmpty());
 	GCalledFunctionName = "ConstFunc1";
 
-	CHECK(Int == GIntArg);
+	CHECK(Int == GInt16Arg);
 }
 
-void UDynamicDelegateTestObject::ConstFunc2(int Int, FString Str) const
+void UDynamicDelegateTestObject::ConstFunc2(int16 Int, FString Str) const
 {
 	CHECK(GCalledFunctionName.IsEmpty());
 	GCalledFunctionName = "ConstFunc2";
 
-	CHECK(Int == GIntArg);
+	CHECK(Int == GInt16Arg);
 	CHECK(Str == GStrArg);
 }
 
-void UDynamicDelegateTestObject::ConstFunc3(int Int, FString Str, const UDynamicDelegateTestObject* Ptr) const
+void UDynamicDelegateTestObject::ConstFunc3(int16 Int, FString Str, const UDynamicDelegateTestObject* Ptr) const
 {
 	CHECK(GCalledFunctionName.IsEmpty());
 	GCalledFunctionName = "ConstFunc3";
 
-	CHECK(Int == GIntArg);
+	CHECK(Int == GInt16Arg);
 	CHECK(Str == GStrArg);
 	CHECK(Ptr == GPtrArg);
 }
 
-void UDynamicDelegateTestObject::ConstFunc4(int Int, FString Str, const UDynamicDelegateTestObject* Ptr, FName Name) const
+void UDynamicDelegateTestObject::ConstFunc4(int16 Int, FString Str, const UDynamicDelegateTestObject* Ptr, FName Name) const
 {
 	CHECK(GCalledFunctionName.IsEmpty());
 	GCalledFunctionName = "ConstFunc4";
 
-	CHECK(Int  == GIntArg);
+	CHECK(Int  == GInt16Arg);
 	CHECK(Str  == GStrArg);
 	CHECK(Ptr  == GPtrArg);
 	CHECK(Name == GNameArg);
 }
 
-void UDynamicDelegateTestObject::ConstFunc5(int Int, FString Str, const UDynamicDelegateTestObject* Ptr, FName Name, bool bBool) const
+void UDynamicDelegateTestObject::ConstFunc5(int16 Int, FString Str, const UDynamicDelegateTestObject* Ptr, FName Name, bool bBool) const
 {
 	CHECK(GCalledFunctionName.IsEmpty());
 	GCalledFunctionName = "ConstFunc5";
 
-	CHECK(Int   == GIntArg);
+	CHECK(Int   == GInt16Arg);
 	CHECK(Str   == GStrArg);
 	CHECK(Ptr   == GPtrArg);
 	CHECK(Name  == GNameArg);
 	CHECK(bBool == bGBoolArg);
 }
 
-void UDynamicDelegateTestObject::ConstFunc6(int Int, FString Str, const UDynamicDelegateTestObject* Ptr, FName Name, bool bBool, TWeakObjectPtr<UDynamicDelegateTestObject> WeakObjPtr) const
+void UDynamicDelegateTestObject::ConstFunc6(int16 Int, FString Str, const UDynamicDelegateTestObject* Ptr, FName Name, bool bBool, TWeakObjectPtr<UDynamicDelegateTestObject> WeakObjPtr) const
 {
 	CHECK(GCalledFunctionName.IsEmpty());
 	GCalledFunctionName = "ConstFunc6";
 
-	CHECK(Int        == GIntArg);
+	CHECK(Int        == GInt16Arg);
 	CHECK(Str        == GStrArg);
 	CHECK(Ptr        == GPtrArg);
 	CHECK(Name       == GNameArg);
@@ -182,42 +182,42 @@ int UDynamicDelegateTestObject::FuncRetInt0()
 	return GIntRet;
 }
 
-int UDynamicDelegateTestObject::FuncRetInt1(int Int)
+int UDynamicDelegateTestObject::FuncRetInt1(int16 Int)
 {
 	Func1(Int);
 	GCalledFunctionName = "FuncRetInt1";
 	return GIntRet;
 }
 
-int UDynamicDelegateTestObject::FuncRetInt2(int Int, FString Str)
+int UDynamicDelegateTestObject::FuncRetInt2(int16 Int, FString Str)
 {
 	Func2(Int, Str);
 	GCalledFunctionName = "FuncRetInt2";
 	return GIntRet;
 }
 
-int UDynamicDelegateTestObject::FuncRetInt3(int Int, FString Str, const UDynamicDelegateTestObject* Ptr)
+int UDynamicDelegateTestObject::FuncRetInt3(int16 Int, FString Str, const UDynamicDelegateTestObject* Ptr)
 {
 	Func3(Int, Str, Ptr);
 	GCalledFunctionName = "FuncRetInt3";
 	return GIntRet;
 }
 
-int UDynamicDelegateTestObject::FuncRetInt4(int Int, FString Str, const UDynamicDelegateTestObject* Ptr, FName Name)
+int UDynamicDelegateTestObject::FuncRetInt4(int16 Int, FString Str, const UDynamicDelegateTestObject* Ptr, FName Name)
 {
 	Func4(Int, Str, Ptr, Name);
 	GCalledFunctionName = "FuncRetInt4";
 	return GIntRet;
 }
 
-int UDynamicDelegateTestObject::FuncRetInt5(int Int, FString Str, const UDynamicDelegateTestObject* Ptr, FName Name, bool bBool)
+int UDynamicDelegateTestObject::FuncRetInt5(int16 Int, FString Str, const UDynamicDelegateTestObject* Ptr, FName Name, bool bBool)
 {
 	Func5(Int, Str, Ptr, Name, bBool);
 	GCalledFunctionName = "FuncRetInt5";
 	return GIntRet;
 }
 
-int UDynamicDelegateTestObject::FuncRetInt6(int Int, FString Str, const UDynamicDelegateTestObject* Ptr, FName Name, bool bBool, TWeakObjectPtr<UDynamicDelegateTestObject> WeakObjPtr)
+int UDynamicDelegateTestObject::FuncRetInt6(int16 Int, FString Str, const UDynamicDelegateTestObject* Ptr, FName Name, bool bBool, TWeakObjectPtr<UDynamicDelegateTestObject> WeakObjPtr)
 {
 	Func6(Int, Str, Ptr, Name, bBool, WeakObjPtr);
 	GCalledFunctionName = "FuncRetInt6";
@@ -231,42 +231,42 @@ FString UDynamicDelegateTestObject::FuncRetStr0()
 	return GStrRet;
 }
 
-FString UDynamicDelegateTestObject::FuncRetStr1(int Int)
+FString UDynamicDelegateTestObject::FuncRetStr1(int16 Int)
 {
 	Func1(Int);
 	GCalledFunctionName = "FuncRetStr1";
 	return GStrRet;
 }
 
-FString UDynamicDelegateTestObject::FuncRetStr2(int Int, FString Str)
+FString UDynamicDelegateTestObject::FuncRetStr2(int16 Int, FString Str)
 {
 	Func2(Int, Str);
 	GCalledFunctionName = "FuncRetStr2";
 	return GStrRet;
 }
 
-FString UDynamicDelegateTestObject::FuncRetStr3(int Int, FString Str, const UDynamicDelegateTestObject* Ptr)
+FString UDynamicDelegateTestObject::FuncRetStr3(int16 Int, FString Str, const UDynamicDelegateTestObject* Ptr)
 {
 	Func3(Int, Str, Ptr);
 	GCalledFunctionName = "FuncRetStr3";
 	return GStrRet;
 }
 
-FString UDynamicDelegateTestObject::FuncRetStr4(int Int, FString Str, const UDynamicDelegateTestObject* Ptr, FName Name)
+FString UDynamicDelegateTestObject::FuncRetStr4(int16 Int, FString Str, const UDynamicDelegateTestObject* Ptr, FName Name)
 {
 	Func4(Int, Str, Ptr, Name);
 	GCalledFunctionName = "FuncRetStr4";
 	return GStrRet;
 }
 
-FString UDynamicDelegateTestObject::FuncRetStr5(int Int, FString Str, const UDynamicDelegateTestObject* Ptr, FName Name, bool bBool)
+FString UDynamicDelegateTestObject::FuncRetStr5(int16 Int, FString Str, const UDynamicDelegateTestObject* Ptr, FName Name, bool bBool)
 {
 	Func5(Int, Str, Ptr, Name, bBool);
 	GCalledFunctionName = "FuncRetStr5";
 	return GStrRet;
 }
 
-FString UDynamicDelegateTestObject::FuncRetStr6(int Int, FString Str, const UDynamicDelegateTestObject* Ptr, FName Name, bool bBool, TWeakObjectPtr<UDynamicDelegateTestObject> WeakObjPtr)
+FString UDynamicDelegateTestObject::FuncRetStr6(int16 Int, FString Str, const UDynamicDelegateTestObject* Ptr, FName Name, bool bBool, TWeakObjectPtr<UDynamicDelegateTestObject> WeakObjPtr)
 {
 	Func6(Int, Str, Ptr, Name, bBool, WeakObjPtr);
 	GCalledFunctionName = "FuncRetStr6";
@@ -280,42 +280,42 @@ int UDynamicDelegateTestObject::ConstFuncRetInt0() const
 	return GIntRet;
 }
 
-int UDynamicDelegateTestObject::ConstFuncRetInt1(int Int) const
+int UDynamicDelegateTestObject::ConstFuncRetInt1(int16 Int) const
 {
 	ConstFunc1(Int);
 	GCalledFunctionName = "ConstFuncRetInt1";
 	return GIntRet;
 }
 
-int UDynamicDelegateTestObject::ConstFuncRetInt2(int Int, FString Str) const
+int UDynamicDelegateTestObject::ConstFuncRetInt2(int16 Int, FString Str) const
 {
 	ConstFunc2(Int, Str);
 	GCalledFunctionName = "ConstFuncRetInt2";
 	return GIntRet;
 }
 
-int UDynamicDelegateTestObject::ConstFuncRetInt3(int Int, FString Str, const UDynamicDelegateTestObject* Ptr) const
+int UDynamicDelegateTestObject::ConstFuncRetInt3(int16 Int, FString Str, const UDynamicDelegateTestObject* Ptr) const
 {
 	ConstFunc3(Int, Str, Ptr);
 	GCalledFunctionName = "ConstFuncRetInt3";
 	return GIntRet;
 }
 
-int UDynamicDelegateTestObject::ConstFuncRetInt4(int Int, FString Str, const UDynamicDelegateTestObject* Ptr, FName Name) const
+int UDynamicDelegateTestObject::ConstFuncRetInt4(int16 Int, FString Str, const UDynamicDelegateTestObject* Ptr, FName Name) const
 {
 	ConstFunc4(Int, Str, Ptr, Name);
 	GCalledFunctionName = "ConstFuncRetInt4";
 	return GIntRet;
 }
 
-int UDynamicDelegateTestObject::ConstFuncRetInt5(int Int, FString Str, const UDynamicDelegateTestObject* Ptr, FName Name, bool bBool) const
+int UDynamicDelegateTestObject::ConstFuncRetInt5(int16 Int, FString Str, const UDynamicDelegateTestObject* Ptr, FName Name, bool bBool) const
 {
 	ConstFunc5(Int, Str, Ptr, Name, bBool);
 	GCalledFunctionName = "ConstFuncRetInt5";
 	return GIntRet;
 }
 
-int UDynamicDelegateTestObject::ConstFuncRetInt6(int Int, FString Str, const UDynamicDelegateTestObject* Ptr, FName Name, bool bBool, TWeakObjectPtr<UDynamicDelegateTestObject> WeakObjPtr) const
+int UDynamicDelegateTestObject::ConstFuncRetInt6(int16 Int, FString Str, const UDynamicDelegateTestObject* Ptr, FName Name, bool bBool, TWeakObjectPtr<UDynamicDelegateTestObject> WeakObjPtr) const
 {
 	ConstFunc6(Int, Str, Ptr, Name, bBool, WeakObjPtr);
 	GCalledFunctionName = "ConstFuncRetInt6";
@@ -329,57 +329,57 @@ FString UDynamicDelegateTestObject::ConstFuncRetStr0() const
 	return GStrRet;
 }
 
-FString UDynamicDelegateTestObject::ConstFuncRetStr1(int Int) const
+FString UDynamicDelegateTestObject::ConstFuncRetStr1(int16 Int) const
 {
 	ConstFunc1(Int);
 	GCalledFunctionName = "ConstFuncRetStr1";
 	return GStrRet;
 }
 
-FString UDynamicDelegateTestObject::ConstFuncRetStr2(int Int, FString Str) const
+FString UDynamicDelegateTestObject::ConstFuncRetStr2(int16 Int, FString Str) const
 {
 	ConstFunc2(Int, Str);
 	GCalledFunctionName = "ConstFuncRetStr2";
 	return GStrRet;
 }
 
-FString UDynamicDelegateTestObject::ConstFuncRetStr3(int Int, FString Str, const UDynamicDelegateTestObject* Ptr) const
+FString UDynamicDelegateTestObject::ConstFuncRetStr3(int16 Int, FString Str, const UDynamicDelegateTestObject* Ptr) const
 {
 	ConstFunc3(Int, Str, Ptr);
 	GCalledFunctionName = "ConstFuncRetStr3";
 	return GStrRet;
 }
 
-FString UDynamicDelegateTestObject::ConstFuncRetStr4(int Int, FString Str, const UDynamicDelegateTestObject* Ptr, FName Name) const
+FString UDynamicDelegateTestObject::ConstFuncRetStr4(int16 Int, FString Str, const UDynamicDelegateTestObject* Ptr, FName Name) const
 {
 	ConstFunc4(Int, Str, Ptr, Name);
 	GCalledFunctionName = "ConstFuncRetStr4";
 	return GStrRet;
 }
 
-FString UDynamicDelegateTestObject::ConstFuncRetStr5(int Int, FString Str, const UDynamicDelegateTestObject* Ptr, FName Name, bool bBool) const
+FString UDynamicDelegateTestObject::ConstFuncRetStr5(int16 Int, FString Str, const UDynamicDelegateTestObject* Ptr, FName Name, bool bBool) const
 {
 	ConstFunc5(Int, Str, Ptr, Name, bBool);
 	GCalledFunctionName = "ConstFuncRetStr5";
 	return GStrRet;
 }
 
-FString UDynamicDelegateTestObject::ConstFuncRetStr6(int Int, FString Str, const UDynamicDelegateTestObject* Ptr, FName Name, bool bBool, TWeakObjectPtr<UDynamicDelegateTestObject> WeakObjPtr) const
+FString UDynamicDelegateTestObject::ConstFuncRetStr6(int16 Int, FString Str, const UDynamicDelegateTestObject* Ptr, FName Name, bool bBool, TWeakObjectPtr<UDynamicDelegateTestObject> WeakObjPtr) const
 {
 	ConstFunc6(Int, Str, Ptr, Name, bBool, WeakObjPtr);
 	GCalledFunctionName = "ConstFuncRetStr6";
 	return GStrRet;
 }
 
-void UDynamicDelegateTestObject::FuncOut1(int& Int)
+void UDynamicDelegateTestObject::FuncOut1(int16& Int)
 {
 	CHECK(GCalledFunctionName.IsEmpty());
 	GCalledFunctionName = "FuncOut1";
-	CHECK(Int == GIntArg);
-	Int = GIntOut;
+	CHECK(Int == GInt16Arg);
+	Int = GInt16Out;
 }
 
-void UDynamicDelegateTestObject::FuncOut2(int& Int, FString& Str)
+void UDynamicDelegateTestObject::FuncOut2(int16& Int, FString& Str)
 {
 	FuncOut1(Int);
 	GCalledFunctionName = "FuncOut2";
@@ -387,7 +387,7 @@ void UDynamicDelegateTestObject::FuncOut2(int& Int, FString& Str)
 	Str = GStrOut;
 }
 
-void UDynamicDelegateTestObject::FuncOut3(int& Int, FString& Str, FName& Name)
+void UDynamicDelegateTestObject::FuncOut3(int16& Int, FString& Str, FName& Name)
 {
 	FuncOut2(Int, Str);
 	GCalledFunctionName = "FuncOut3";
@@ -395,7 +395,7 @@ void UDynamicDelegateTestObject::FuncOut3(int& Int, FString& Str, FName& Name)
 	Name = GNameOut;
 }
 
-void UDynamicDelegateTestObject::FuncOut4(int& Int, FString& Str, FName& Name, bool& bBool)
+void UDynamicDelegateTestObject::FuncOut4(int16& Int, FString& Str, FName& Name, bool& bBool)
 {
 	FuncOut3(Int, Str, Name);
 	GCalledFunctionName = "FuncOut4";
@@ -403,7 +403,7 @@ void UDynamicDelegateTestObject::FuncOut4(int& Int, FString& Str, FName& Name, b
 	bBool = bGBoolOut;
 }
 
-void UDynamicDelegateTestObject::FuncOut5(int& Int, FString& Str, FName& Name, bool& bBool, TWeakObjectPtr<UDynamicDelegateTestObject>& WeakObjPtr)
+void UDynamicDelegateTestObject::FuncOut5(int16& Int, FString& Str, FName& Name, bool& bBool, TWeakObjectPtr<UDynamicDelegateTestObject>& WeakObjPtr)
 {
 	FuncOut4(Int, Str, Name, bBool);
 	GCalledFunctionName = "FuncOut5";
@@ -411,15 +411,15 @@ void UDynamicDelegateTestObject::FuncOut5(int& Int, FString& Str, FName& Name, b
 	WeakObjPtr = GWeakObjPtrOut;
 }
 
-void UDynamicDelegateTestObject::ConstFuncOut1(int& Int) const
+void UDynamicDelegateTestObject::ConstFuncOut1(int16& Int) const
 {
 	CHECK(GCalledFunctionName.IsEmpty());
 	GCalledFunctionName = "ConstFuncOut1";
-	CHECK(Int == GIntArg);
-	Int = GIntOut;
+	CHECK(Int == GInt16Arg);
+	Int = GInt16Out;
 }
 
-void UDynamicDelegateTestObject::ConstFuncOut2(int& Int, FString& Str) const
+void UDynamicDelegateTestObject::ConstFuncOut2(int16& Int, FString& Str) const
 {
 	ConstFuncOut1(Int);
 	GCalledFunctionName = "ConstFuncOut2";
@@ -427,7 +427,7 @@ void UDynamicDelegateTestObject::ConstFuncOut2(int& Int, FString& Str) const
 	Str = GStrOut;
 }
 
-void UDynamicDelegateTestObject::ConstFuncOut3(int& Int, FString& Str, FName& Name) const
+void UDynamicDelegateTestObject::ConstFuncOut3(int16& Int, FString& Str, FName& Name) const
 {
 	ConstFuncOut2(Int, Str);
 	GCalledFunctionName = "ConstFuncOut3";
@@ -435,7 +435,7 @@ void UDynamicDelegateTestObject::ConstFuncOut3(int& Int, FString& Str, FName& Na
 	Name = GNameOut;
 }
 
-void UDynamicDelegateTestObject::ConstFuncOut4(int& Int, FString& Str, FName& Name, bool& bBool) const
+void UDynamicDelegateTestObject::ConstFuncOut4(int16& Int, FString& Str, FName& Name, bool& bBool) const
 {
 	ConstFuncOut3(Int, Str, Name);
 	GCalledFunctionName = "ConstFuncOut4";
@@ -443,7 +443,7 @@ void UDynamicDelegateTestObject::ConstFuncOut4(int& Int, FString& Str, FName& Na
 	bBool = bGBoolOut;
 }
 
-void UDynamicDelegateTestObject::ConstFuncOut5(int& Int, FString& Str, FName& Name, bool& bBool, TWeakObjectPtr<UDynamicDelegateTestObject>& WeakObjPtr) const
+void UDynamicDelegateTestObject::ConstFuncOut5(int16& Int, FString& Str, FName& Name, bool& bBool, TWeakObjectPtr<UDynamicDelegateTestObject>& WeakObjPtr) const
 {
 	ConstFuncOut4(Int, Str, Name, bBool);
 	GCalledFunctionName = "ConstFuncOut5";
@@ -521,53 +521,55 @@ void UDynamicDelegateTestObject::RunTests()
 	UDynamicDelegateTestObject* TestOut = NewObject<UDynamicDelegateTestObject>();
 
 	// Update ptr args, because they can't be given a reasonable default as a global
-	TGuardValue<const UDynamicDelegateTestObject*>          PtrArgGuard(GPtrArg, this);
+	TGuardValue<const UDynamicDelegateTestObject*>          PtrArgGuard    (GPtrArg,        this);
 	TGuardValue<TWeakObjectPtr<UDynamicDelegateTestObject>> WeakPtrArgGuard(GWeakObjPtrArg, this);
-	TGuardValue<const UDynamicDelegateTestObject*>          OutPtrArgGuard(GPtrOut, TestOut);
+	TGuardValue<const UDynamicDelegateTestObject*>          OutPtrArgGuard (GPtrOut,        TestOut);
 	TGuardValue<TWeakObjectPtr<UDynamicDelegateTestObject>> WeakPtrOutGuard(GWeakObjPtrOut, TestOut);
 
 	// 0 param function
 	{
 		// 0 arg payload
-		BIND_CALL_CHECK_DELEGATE(FTestDynamicMulticastDelegate0::FDelegate, 0, &UDynamicDelegateTestObject::Func0, ());
-		BIND_CALL_CHECK_DELEGATE(FTestDynamicDelegate0, 0, &UDynamicDelegateTestObject::Func0, ());
-		BIND_CALL_CHECK_DELEGATE(FTestDynamicMulticastDelegate0, 1, &UDynamicDelegateTestObject::Func0, ());
+		BIND_CALL_CHECK_DELEGATE(FTestDynamicMulticastDelegate0::FDelegate, 0, &UDynamicDelegateTestObject::Func0,      ());
+		BIND_CALL_CHECK_DELEGATE(FTestDynamicDelegate0,                     0, &UDynamicDelegateTestObject::Func0,      ());
+		BIND_CALL_CHECK_DELEGATE(FTestDynamicMulticastDelegate0,            1, &UDynamicDelegateTestObject::Func0,      ());
 		BIND_CALL_CHECK_DELEGATE(FTestDynamicMulticastDelegate0::FDelegate, 0, &UDynamicDelegateTestObject::ConstFunc0, ());
-		BIND_CALL_CHECK_DELEGATE(FTestDynamicDelegate0, 0, &UDynamicDelegateTestObject::ConstFunc0, ());
-		BIND_CALL_CHECK_DELEGATE(FTestDynamicMulticastDelegate0, 1, &UDynamicDelegateTestObject::ConstFunc0, ());
-		BIND_CALL_CHECK_DELEGATE_RETVAL(GIntRet, FTestDynamicDelegateRetInt0, 0, &UDynamicDelegateTestObject::FuncRetInt0, ());
+		BIND_CALL_CHECK_DELEGATE(FTestDynamicDelegate0,                     0, &UDynamicDelegateTestObject::ConstFunc0, ());
+		BIND_CALL_CHECK_DELEGATE(FTestDynamicMulticastDelegate0,            1, &UDynamicDelegateTestObject::ConstFunc0, ());
+
+		BIND_CALL_CHECK_DELEGATE_RETVAL(GIntRet, FTestDynamicDelegateRetInt0, 0, &UDynamicDelegateTestObject::FuncRetInt0,      ());
 		BIND_CALL_CHECK_DELEGATE_RETVAL(GIntRet, FTestDynamicDelegateRetInt0, 0, &UDynamicDelegateTestObject::ConstFuncRetInt0, ());
-		BIND_CALL_CHECK_DELEGATE_RETVAL(GStrRet, FTestDynamicDelegateRetStr0, 0, &UDynamicDelegateTestObject::FuncRetStr0, ());
+		BIND_CALL_CHECK_DELEGATE_RETVAL(GStrRet, FTestDynamicDelegateRetStr0, 0, &UDynamicDelegateTestObject::FuncRetStr0,      ());
 		BIND_CALL_CHECK_DELEGATE_RETVAL(GStrRet, FTestDynamicDelegateRetStr0, 0, &UDynamicDelegateTestObject::ConstFuncRetStr0, ());
 	}
 
 	// 1 param function
 	{
 		// 0 arg payload
-		BIND_CALL_CHECK_DELEGATE(FTestDynamicMulticastDelegate1::FDelegate, 0, &UDynamicDelegateTestObject::Func1, (GIntArg));
-		BIND_CALL_CHECK_DELEGATE(FTestDynamicDelegate1, 0, &UDynamicDelegateTestObject::Func1, (GIntArg));
-		BIND_CALL_CHECK_DELEGATE(FTestDynamicMulticastDelegate1, 1, &UDynamicDelegateTestObject::Func1, (GIntArg));
-		BIND_CALL_CHECK_DELEGATE(FTestDynamicMulticastDelegate1::FDelegate, 0, &UDynamicDelegateTestObject::ConstFunc1, (GIntArg));
-		BIND_CALL_CHECK_DELEGATE(FTestDynamicDelegate1, 0, &UDynamicDelegateTestObject::ConstFunc1, (GIntArg));
-		BIND_CALL_CHECK_DELEGATE(FTestDynamicMulticastDelegate1, 1, &UDynamicDelegateTestObject::ConstFunc1, (GIntArg));
-		BIND_CALL_CHECK_DELEGATE_RETVAL(GIntRet, FTestDynamicDelegateRetInt1, 0, &UDynamicDelegateTestObject::FuncRetInt1, (GIntArg));
-		BIND_CALL_CHECK_DELEGATE_RETVAL(GIntRet, FTestDynamicDelegateRetInt1, 0, &UDynamicDelegateTestObject::ConstFuncRetInt1, (GIntArg));
-		BIND_CALL_CHECK_DELEGATE_RETVAL(GStrRet, FTestDynamicDelegateRetStr1, 0, &UDynamicDelegateTestObject::FuncRetStr1, (GIntArg));
-		BIND_CALL_CHECK_DELEGATE_RETVAL(GStrRet, FTestDynamicDelegateRetStr1, 0, &UDynamicDelegateTestObject::ConstFuncRetStr1, (GIntArg));
+		BIND_CALL_CHECK_DELEGATE(FTestDynamicMulticastDelegate1::FDelegate, 0, &UDynamicDelegateTestObject::Func1,      (GInt16Arg));
+		BIND_CALL_CHECK_DELEGATE(FTestDynamicDelegate1,                     0, &UDynamicDelegateTestObject::Func1,      (GInt16Arg));
+		BIND_CALL_CHECK_DELEGATE(FTestDynamicMulticastDelegate1,            1, &UDynamicDelegateTestObject::Func1,      (GInt16Arg));
+		BIND_CALL_CHECK_DELEGATE(FTestDynamicMulticastDelegate1::FDelegate, 0, &UDynamicDelegateTestObject::ConstFunc1, (GInt16Arg));
+		BIND_CALL_CHECK_DELEGATE(FTestDynamicDelegate1,                     0, &UDynamicDelegateTestObject::ConstFunc1, (GInt16Arg));
+		BIND_CALL_CHECK_DELEGATE(FTestDynamicMulticastDelegate1,            1, &UDynamicDelegateTestObject::ConstFunc1, (GInt16Arg));
+
+		BIND_CALL_CHECK_DELEGATE_RETVAL(GIntRet, FTestDynamicDelegateRetInt1, 0, &UDynamicDelegateTestObject::FuncRetInt1,      (GInt16Arg));
+		BIND_CALL_CHECK_DELEGATE_RETVAL(GIntRet, FTestDynamicDelegateRetInt1, 0, &UDynamicDelegateTestObject::ConstFuncRetInt1, (GInt16Arg));
+		BIND_CALL_CHECK_DELEGATE_RETVAL(GStrRet, FTestDynamicDelegateRetStr1, 0, &UDynamicDelegateTestObject::FuncRetStr1,      (GInt16Arg));
+		BIND_CALL_CHECK_DELEGATE_RETVAL(GStrRet, FTestDynamicDelegateRetStr1, 0, &UDynamicDelegateTestObject::ConstFuncRetStr1, (GInt16Arg));
 
 #if defined(UE_USE_DYNAMIC_DELEGATE_PAYLOADS) && UE_USE_DYNAMIC_DELEGATE_PAYLOADS
 		// 1 arg payload
-		BIND_CALL_CHECK_DELEGATE(FTestDynamicMulticastDelegate0::FDelegate, 0, &UDynamicDelegateTestObject::Func1, (), GIntArg);
-		BIND_CALL_CHECK_DELEGATE(FTestDynamicDelegate0, 0, &UDynamicDelegateTestObject::Func1, (), GIntArg);
-		BIND_CALL_CHECK_DELEGATE(FTestDynamicMulticastDelegate0, 1, &UDynamicDelegateTestObject::Func1, (), GIntArg);
-		BIND_CALL_CHECK_DELEGATE(FTestDynamicMulticastDelegate0::FDelegate, 0, &UDynamicDelegateTestObject::ConstFunc1, (), GIntArg);
-		BIND_CALL_CHECK_DELEGATE(FTestDynamicDelegate0, 0, &UDynamicDelegateTestObject::ConstFunc1, (), GIntArg);
-		BIND_CALL_CHECK_DELEGATE(FTestDynamicMulticastDelegate0, 1, &UDynamicDelegateTestObject::ConstFunc1, (), GIntArg);
+		BIND_CALL_CHECK_DELEGATE(FTestDynamicMulticastDelegate0::FDelegate, 0, &UDynamicDelegateTestObject::Func1,      (), GInt16Arg);
+		BIND_CALL_CHECK_DELEGATE(FTestDynamicDelegate0,                     0, &UDynamicDelegateTestObject::Func1,      (), GInt16Arg);
+		BIND_CALL_CHECK_DELEGATE(FTestDynamicMulticastDelegate0,            1, &UDynamicDelegateTestObject::Func1,      (), GInt16Arg);
+		BIND_CALL_CHECK_DELEGATE(FTestDynamicMulticastDelegate0::FDelegate, 0, &UDynamicDelegateTestObject::ConstFunc1, (), GInt16Arg);
+		BIND_CALL_CHECK_DELEGATE(FTestDynamicDelegate0,                     0, &UDynamicDelegateTestObject::ConstFunc1, (), GInt16Arg);
+		BIND_CALL_CHECK_DELEGATE(FTestDynamicMulticastDelegate0,            1, &UDynamicDelegateTestObject::ConstFunc1, (), GInt16Arg);
 #if UE_USE_DYNAMIC_DELEGATE_PAYLOADS_RETVAL
-		BIND_CALL_CHECK_DELEGATE_RETVAL(GIntRet, FTestDynamicDelegateRetInt0, 0, &UDynamicDelegateTestObject::FuncRetInt1, (), GIntArg);
-		BIND_CALL_CHECK_DELEGATE_RETVAL(GIntRet, FTestDynamicDelegateRetInt0, 0, &UDynamicDelegateTestObject::ConstFuncRetInt1, (), GIntArg);
-		BIND_CALL_CHECK_DELEGATE_RETVAL(GStrRet, FTestDynamicDelegateRetStr0, 0, &UDynamicDelegateTestObject::FuncRetStr1, (), GIntArg);
-		BIND_CALL_CHECK_DELEGATE_RETVAL(GStrRet, FTestDynamicDelegateRetStr0, 0, &UDynamicDelegateTestObject::ConstFuncRetStr1, (), GIntArg);
+		BIND_CALL_CHECK_DELEGATE_RETVAL(GIntRet, FTestDynamicDelegateRetInt0, 0, &UDynamicDelegateTestObject::FuncRetInt1,      (), GInt16Arg);
+		BIND_CALL_CHECK_DELEGATE_RETVAL(GIntRet, FTestDynamicDelegateRetInt0, 0, &UDynamicDelegateTestObject::ConstFuncRetInt1, (), GInt16Arg);
+		BIND_CALL_CHECK_DELEGATE_RETVAL(GStrRet, FTestDynamicDelegateRetStr0, 0, &UDynamicDelegateTestObject::FuncRetStr1,      (), GInt16Arg);
+		BIND_CALL_CHECK_DELEGATE_RETVAL(GStrRet, FTestDynamicDelegateRetStr0, 0, &UDynamicDelegateTestObject::ConstFuncRetStr1, (), GInt16Arg);
 #endif
 #endif
 	}
@@ -575,44 +577,45 @@ void UDynamicDelegateTestObject::RunTests()
 	// 2 param function
 	{
 		// 0 arg payload
-		BIND_CALL_CHECK_DELEGATE(FTestDynamicMulticastDelegate2::FDelegate, 0, &UDynamicDelegateTestObject::Func2, (GIntArg, GStrArg));
-		BIND_CALL_CHECK_DELEGATE(FTestDynamicDelegate2, 0, &UDynamicDelegateTestObject::Func2, (GIntArg, GStrArg));
-		BIND_CALL_CHECK_DELEGATE(FTestDynamicMulticastDelegate2, 1, &UDynamicDelegateTestObject::Func2, (GIntArg, GStrArg));
-		BIND_CALL_CHECK_DELEGATE(FTestDynamicMulticastDelegate2::FDelegate, 0, &UDynamicDelegateTestObject::ConstFunc2, (GIntArg, GStrArg));
-		BIND_CALL_CHECK_DELEGATE(FTestDynamicDelegate2, 0, &UDynamicDelegateTestObject::ConstFunc2, (GIntArg, GStrArg));
-		BIND_CALL_CHECK_DELEGATE(FTestDynamicMulticastDelegate2, 1, &UDynamicDelegateTestObject::ConstFunc2, (GIntArg, GStrArg));
-		BIND_CALL_CHECK_DELEGATE_RETVAL(GIntRet, FTestDynamicDelegateRetInt2, 0, &UDynamicDelegateTestObject::FuncRetInt2, (GIntArg, GStrArg));
-		BIND_CALL_CHECK_DELEGATE_RETVAL(GIntRet, FTestDynamicDelegateRetInt2, 0, &UDynamicDelegateTestObject::ConstFuncRetInt2, (GIntArg, GStrArg));
-		BIND_CALL_CHECK_DELEGATE_RETVAL(GStrRet, FTestDynamicDelegateRetStr2, 0, &UDynamicDelegateTestObject::FuncRetStr2, (GIntArg, GStrArg));
-		BIND_CALL_CHECK_DELEGATE_RETVAL(GStrRet, FTestDynamicDelegateRetStr2, 0, &UDynamicDelegateTestObject::ConstFuncRetStr2, (GIntArg, GStrArg));
+		BIND_CALL_CHECK_DELEGATE(FTestDynamicMulticastDelegate2::FDelegate, 0, &UDynamicDelegateTestObject::Func2,      (GInt16Arg, GStrArg));
+		BIND_CALL_CHECK_DELEGATE(FTestDynamicDelegate2,                     0, &UDynamicDelegateTestObject::Func2,      (GInt16Arg, GStrArg));
+		BIND_CALL_CHECK_DELEGATE(FTestDynamicMulticastDelegate2,            1, &UDynamicDelegateTestObject::Func2,      (GInt16Arg, GStrArg));
+		BIND_CALL_CHECK_DELEGATE(FTestDynamicMulticastDelegate2::FDelegate, 0, &UDynamicDelegateTestObject::ConstFunc2, (GInt16Arg, GStrArg));
+		BIND_CALL_CHECK_DELEGATE(FTestDynamicDelegate2,                     0, &UDynamicDelegateTestObject::ConstFunc2, (GInt16Arg, GStrArg));
+		BIND_CALL_CHECK_DELEGATE(FTestDynamicMulticastDelegate2,            1, &UDynamicDelegateTestObject::ConstFunc2, (GInt16Arg, GStrArg));
+
+		BIND_CALL_CHECK_DELEGATE_RETVAL(GIntRet, FTestDynamicDelegateRetInt2, 0, &UDynamicDelegateTestObject::FuncRetInt2,      (GInt16Arg, GStrArg));
+		BIND_CALL_CHECK_DELEGATE_RETVAL(GIntRet, FTestDynamicDelegateRetInt2, 0, &UDynamicDelegateTestObject::ConstFuncRetInt2, (GInt16Arg, GStrArg));
+		BIND_CALL_CHECK_DELEGATE_RETVAL(GStrRet, FTestDynamicDelegateRetStr2, 0, &UDynamicDelegateTestObject::FuncRetStr2,      (GInt16Arg, GStrArg));
+		BIND_CALL_CHECK_DELEGATE_RETVAL(GStrRet, FTestDynamicDelegateRetStr2, 0, &UDynamicDelegateTestObject::ConstFuncRetStr2, (GInt16Arg, GStrArg));
 
 #if defined(UE_USE_DYNAMIC_DELEGATE_PAYLOADS) && UE_USE_DYNAMIC_DELEGATE_PAYLOADS
 		// 1 arg payload
-		BIND_CALL_CHECK_DELEGATE(FTestDynamicMulticastDelegate1::FDelegate, 0, &UDynamicDelegateTestObject::Func2, (GIntArg), GStrArg);
-		BIND_CALL_CHECK_DELEGATE(FTestDynamicDelegate1, 0, &UDynamicDelegateTestObject::Func2, (GIntArg), GStrArg);
-		BIND_CALL_CHECK_DELEGATE(FTestDynamicMulticastDelegate1, 1, &UDynamicDelegateTestObject::Func2, (GIntArg), GStrArg);
-		BIND_CALL_CHECK_DELEGATE(FTestDynamicMulticastDelegate1::FDelegate, 0, &UDynamicDelegateTestObject::ConstFunc2, (GIntArg), GStrArg);
-		BIND_CALL_CHECK_DELEGATE(FTestDynamicDelegate1, 0, &UDynamicDelegateTestObject::ConstFunc2, (GIntArg), GStrArg);
-		BIND_CALL_CHECK_DELEGATE(FTestDynamicMulticastDelegate1, 1, &UDynamicDelegateTestObject::ConstFunc2, (GIntArg), GStrArg);
+		BIND_CALL_CHECK_DELEGATE(FTestDynamicMulticastDelegate1::FDelegate, 0, &UDynamicDelegateTestObject::Func2,      (GInt16Arg), GStrArg);
+		BIND_CALL_CHECK_DELEGATE(FTestDynamicDelegate1,                     0, &UDynamicDelegateTestObject::Func2,      (GInt16Arg), GStrArg);
+		BIND_CALL_CHECK_DELEGATE(FTestDynamicMulticastDelegate1,            1, &UDynamicDelegateTestObject::Func2,      (GInt16Arg), GStrArg);
+		BIND_CALL_CHECK_DELEGATE(FTestDynamicMulticastDelegate1::FDelegate, 0, &UDynamicDelegateTestObject::ConstFunc2, (GInt16Arg), GStrArg);
+		BIND_CALL_CHECK_DELEGATE(FTestDynamicDelegate1,                     0, &UDynamicDelegateTestObject::ConstFunc2, (GInt16Arg), GStrArg);
+		BIND_CALL_CHECK_DELEGATE(FTestDynamicMulticastDelegate1,            1, &UDynamicDelegateTestObject::ConstFunc2, (GInt16Arg), GStrArg);
 #if UE_USE_DYNAMIC_DELEGATE_PAYLOADS_RETVAL
-		BIND_CALL_CHECK_DELEGATE_RETVAL(GIntRet, FTestDynamicDelegateRetInt1, 0, &UDynamicDelegateTestObject::FuncRetInt2, (GIntArg), GStrArg);
-		BIND_CALL_CHECK_DELEGATE_RETVAL(GIntRet, FTestDynamicDelegateRetInt1, 0, &UDynamicDelegateTestObject::ConstFuncRetInt2, (GIntArg), GStrArg);
-		BIND_CALL_CHECK_DELEGATE_RETVAL(GStrRet, FTestDynamicDelegateRetStr1, 0, &UDynamicDelegateTestObject::FuncRetStr2, (GIntArg), GStrArg);
-		BIND_CALL_CHECK_DELEGATE_RETVAL(GStrRet, FTestDynamicDelegateRetStr1, 0, &UDynamicDelegateTestObject::ConstFuncRetStr2, (GIntArg), GStrArg);
+		BIND_CALL_CHECK_DELEGATE_RETVAL(GIntRet, FTestDynamicDelegateRetInt1, 0, &UDynamicDelegateTestObject::FuncRetInt2,      (GInt16Arg), GStrArg);
+		BIND_CALL_CHECK_DELEGATE_RETVAL(GIntRet, FTestDynamicDelegateRetInt1, 0, &UDynamicDelegateTestObject::ConstFuncRetInt2, (GInt16Arg), GStrArg);
+		BIND_CALL_CHECK_DELEGATE_RETVAL(GStrRet, FTestDynamicDelegateRetStr1, 0, &UDynamicDelegateTestObject::FuncRetStr2,      (GInt16Arg), GStrArg);
+		BIND_CALL_CHECK_DELEGATE_RETVAL(GStrRet, FTestDynamicDelegateRetStr1, 0, &UDynamicDelegateTestObject::ConstFuncRetStr2, (GInt16Arg), GStrArg);
 #endif
 
 		// 2 arg payload
-		BIND_CALL_CHECK_DELEGATE(FTestDynamicMulticastDelegate0::FDelegate, 0, &UDynamicDelegateTestObject::Func2, (), GIntArg, GStrArg);
-		BIND_CALL_CHECK_DELEGATE(FTestDynamicDelegate0, 0, &UDynamicDelegateTestObject::Func2, (), GIntArg, GStrArg);
-		BIND_CALL_CHECK_DELEGATE(FTestDynamicMulticastDelegate0, 1, &UDynamicDelegateTestObject::Func2, (), GIntArg, GStrArg);
-		BIND_CALL_CHECK_DELEGATE(FTestDynamicMulticastDelegate0::FDelegate, 0, &UDynamicDelegateTestObject::ConstFunc2, (), GIntArg, GStrArg);
-		BIND_CALL_CHECK_DELEGATE(FTestDynamicDelegate0, 0, &UDynamicDelegateTestObject::ConstFunc2, (), GIntArg, GStrArg);
-		BIND_CALL_CHECK_DELEGATE(FTestDynamicMulticastDelegate0, 1, &UDynamicDelegateTestObject::ConstFunc2, (), GIntArg, GStrArg);
+		BIND_CALL_CHECK_DELEGATE(FTestDynamicMulticastDelegate0::FDelegate, 0, &UDynamicDelegateTestObject::Func2,      (), GInt16Arg, GStrArg);
+		BIND_CALL_CHECK_DELEGATE(FTestDynamicDelegate0,                     0, &UDynamicDelegateTestObject::Func2,      (), GInt16Arg, GStrArg);
+		BIND_CALL_CHECK_DELEGATE(FTestDynamicMulticastDelegate0,            1, &UDynamicDelegateTestObject::Func2,      (), GInt16Arg, GStrArg);
+		BIND_CALL_CHECK_DELEGATE(FTestDynamicMulticastDelegate0::FDelegate, 0, &UDynamicDelegateTestObject::ConstFunc2, (), GInt16Arg, GStrArg);
+		BIND_CALL_CHECK_DELEGATE(FTestDynamicDelegate0,                     0, &UDynamicDelegateTestObject::ConstFunc2, (), GInt16Arg, GStrArg);
+		BIND_CALL_CHECK_DELEGATE(FTestDynamicMulticastDelegate0,            1, &UDynamicDelegateTestObject::ConstFunc2, (), GInt16Arg, GStrArg);
 #if UE_USE_DYNAMIC_DELEGATE_PAYLOADS_RETVAL
-		BIND_CALL_CHECK_DELEGATE_RETVAL(GIntRet, FTestDynamicDelegateRetInt0, 0, &UDynamicDelegateTestObject::FuncRetInt2, (), GIntArg, GStrArg);
-		BIND_CALL_CHECK_DELEGATE_RETVAL(GIntRet, FTestDynamicDelegateRetInt0, 0, &UDynamicDelegateTestObject::ConstFuncRetInt2, (), GIntArg, GStrArg);
-		BIND_CALL_CHECK_DELEGATE_RETVAL(GStrRet, FTestDynamicDelegateRetStr0, 0, &UDynamicDelegateTestObject::FuncRetStr2, (), GIntArg, GStrArg);
-		BIND_CALL_CHECK_DELEGATE_RETVAL(GStrRet, FTestDynamicDelegateRetStr0, 0, &UDynamicDelegateTestObject::ConstFuncRetStr2, (), GIntArg, GStrArg);
+		BIND_CALL_CHECK_DELEGATE_RETVAL(GIntRet, FTestDynamicDelegateRetInt0, 0, &UDynamicDelegateTestObject::FuncRetInt2,      (), GInt16Arg, GStrArg);
+		BIND_CALL_CHECK_DELEGATE_RETVAL(GIntRet, FTestDynamicDelegateRetInt0, 0, &UDynamicDelegateTestObject::ConstFuncRetInt2, (), GInt16Arg, GStrArg);
+		BIND_CALL_CHECK_DELEGATE_RETVAL(GStrRet, FTestDynamicDelegateRetStr0, 0, &UDynamicDelegateTestObject::FuncRetStr2,      (), GInt16Arg, GStrArg);
+		BIND_CALL_CHECK_DELEGATE_RETVAL(GStrRet, FTestDynamicDelegateRetStr0, 0, &UDynamicDelegateTestObject::ConstFuncRetStr2, (), GInt16Arg, GStrArg);
 #endif
 #endif
 	}
@@ -620,58 +623,59 @@ void UDynamicDelegateTestObject::RunTests()
 	// 3 param function
 	{
 		// 0 arg payload
-		BIND_CALL_CHECK_DELEGATE(FTestDynamicMulticastDelegate3::FDelegate, 0, &UDynamicDelegateTestObject::Func3, (GIntArg, GStrArg, GPtrArg));
-		BIND_CALL_CHECK_DELEGATE(FTestDynamicDelegate3, 0, &UDynamicDelegateTestObject::Func3, (GIntArg, GStrArg, GPtrArg));
-		BIND_CALL_CHECK_DELEGATE(FTestDynamicMulticastDelegate3, 1, &UDynamicDelegateTestObject::Func3, (GIntArg, GStrArg, GPtrArg));
-		BIND_CALL_CHECK_DELEGATE(FTestDynamicMulticastDelegate3::FDelegate, 0, &UDynamicDelegateTestObject::ConstFunc3, (GIntArg, GStrArg, GPtrArg));
-		BIND_CALL_CHECK_DELEGATE(FTestDynamicDelegate3, 0, &UDynamicDelegateTestObject::ConstFunc3, (GIntArg, GStrArg, GPtrArg));
-		BIND_CALL_CHECK_DELEGATE(FTestDynamicMulticastDelegate3, 1, &UDynamicDelegateTestObject::ConstFunc3, (GIntArg, GStrArg, GPtrArg));
-		BIND_CALL_CHECK_DELEGATE_RETVAL(GIntRet, FTestDynamicDelegateRetInt3, 0, &UDynamicDelegateTestObject::FuncRetInt3, (GIntArg, GStrArg, GPtrArg));
-		BIND_CALL_CHECK_DELEGATE_RETVAL(GIntRet, FTestDynamicDelegateRetInt3, 0, &UDynamicDelegateTestObject::ConstFuncRetInt3, (GIntArg, GStrArg, GPtrArg));
-		BIND_CALL_CHECK_DELEGATE_RETVAL(GStrRet, FTestDynamicDelegateRetStr3, 0, &UDynamicDelegateTestObject::FuncRetStr3, (GIntArg, GStrArg, GPtrArg));
-		BIND_CALL_CHECK_DELEGATE_RETVAL(GStrRet, FTestDynamicDelegateRetStr3, 0, &UDynamicDelegateTestObject::ConstFuncRetStr3, (GIntArg, GStrArg, GPtrArg));
+		BIND_CALL_CHECK_DELEGATE(FTestDynamicMulticastDelegate3::FDelegate, 0, &UDynamicDelegateTestObject::Func3,      (GInt16Arg, GStrArg, GPtrArg));
+		BIND_CALL_CHECK_DELEGATE(FTestDynamicDelegate3,                     0, &UDynamicDelegateTestObject::Func3,      (GInt16Arg, GStrArg, GPtrArg));
+		BIND_CALL_CHECK_DELEGATE(FTestDynamicMulticastDelegate3,            1, &UDynamicDelegateTestObject::Func3,      (GInt16Arg, GStrArg, GPtrArg));
+		BIND_CALL_CHECK_DELEGATE(FTestDynamicMulticastDelegate3::FDelegate, 0, &UDynamicDelegateTestObject::ConstFunc3, (GInt16Arg, GStrArg, GPtrArg));
+		BIND_CALL_CHECK_DELEGATE(FTestDynamicDelegate3,                     0, &UDynamicDelegateTestObject::ConstFunc3, (GInt16Arg, GStrArg, GPtrArg));
+		BIND_CALL_CHECK_DELEGATE(FTestDynamicMulticastDelegate3,            1, &UDynamicDelegateTestObject::ConstFunc3, (GInt16Arg, GStrArg, GPtrArg));
+
+		BIND_CALL_CHECK_DELEGATE_RETVAL(GIntRet, FTestDynamicDelegateRetInt3, 0, &UDynamicDelegateTestObject::FuncRetInt3,      (GInt16Arg, GStrArg, GPtrArg));
+		BIND_CALL_CHECK_DELEGATE_RETVAL(GIntRet, FTestDynamicDelegateRetInt3, 0, &UDynamicDelegateTestObject::ConstFuncRetInt3, (GInt16Arg, GStrArg, GPtrArg));
+		BIND_CALL_CHECK_DELEGATE_RETVAL(GStrRet, FTestDynamicDelegateRetStr3, 0, &UDynamicDelegateTestObject::FuncRetStr3,      (GInt16Arg, GStrArg, GPtrArg));
+		BIND_CALL_CHECK_DELEGATE_RETVAL(GStrRet, FTestDynamicDelegateRetStr3, 0, &UDynamicDelegateTestObject::ConstFuncRetStr3, (GInt16Arg, GStrArg, GPtrArg));
 
 #if defined(UE_USE_DYNAMIC_DELEGATE_PAYLOADS) && UE_USE_DYNAMIC_DELEGATE_PAYLOADS
 		// 1 arg payload
-		BIND_CALL_CHECK_DELEGATE(FTestDynamicMulticastDelegate2::FDelegate, 0, &UDynamicDelegateTestObject::Func3, (GIntArg, GStrArg), GPtrArg);
-		BIND_CALL_CHECK_DELEGATE(FTestDynamicDelegate2, 0, &UDynamicDelegateTestObject::Func3, (GIntArg, GStrArg), GPtrArg);
-		BIND_CALL_CHECK_DELEGATE(FTestDynamicMulticastDelegate2, 1, &UDynamicDelegateTestObject::Func3, (GIntArg, GStrArg), GPtrArg);
-		BIND_CALL_CHECK_DELEGATE(FTestDynamicMulticastDelegate2::FDelegate, 0, &UDynamicDelegateTestObject::ConstFunc3, (GIntArg, GStrArg), GPtrArg);
-		BIND_CALL_CHECK_DELEGATE(FTestDynamicDelegate2, 0, &UDynamicDelegateTestObject::ConstFunc3, (GIntArg, GStrArg), GPtrArg);
-		BIND_CALL_CHECK_DELEGATE(FTestDynamicMulticastDelegate2, 1, &UDynamicDelegateTestObject::ConstFunc3, (GIntArg, GStrArg), GPtrArg);
+		BIND_CALL_CHECK_DELEGATE(FTestDynamicMulticastDelegate2::FDelegate, 0, &UDynamicDelegateTestObject::Func3,      (GInt16Arg, GStrArg), GPtrArg);
+		BIND_CALL_CHECK_DELEGATE(FTestDynamicDelegate2,                     0, &UDynamicDelegateTestObject::Func3,      (GInt16Arg, GStrArg), GPtrArg);
+		BIND_CALL_CHECK_DELEGATE(FTestDynamicMulticastDelegate2,            1, &UDynamicDelegateTestObject::Func3,      (GInt16Arg, GStrArg), GPtrArg);
+		BIND_CALL_CHECK_DELEGATE(FTestDynamicMulticastDelegate2::FDelegate, 0, &UDynamicDelegateTestObject::ConstFunc3, (GInt16Arg, GStrArg), GPtrArg);
+		BIND_CALL_CHECK_DELEGATE(FTestDynamicDelegate2,                     0, &UDynamicDelegateTestObject::ConstFunc3, (GInt16Arg, GStrArg), GPtrArg);
+		BIND_CALL_CHECK_DELEGATE(FTestDynamicMulticastDelegate2,            1, &UDynamicDelegateTestObject::ConstFunc3, (GInt16Arg, GStrArg), GPtrArg);
 #if UE_USE_DYNAMIC_DELEGATE_PAYLOADS_RETVAL
-		BIND_CALL_CHECK_DELEGATE_RETVAL(GIntRet, FTestDynamicDelegateRetInt2, 0, &UDynamicDelegateTestObject::FuncRetInt3, (GIntArg, GStrArg), GPtrArg);
-		BIND_CALL_CHECK_DELEGATE_RETVAL(GIntRet, FTestDynamicDelegateRetInt2, 0, &UDynamicDelegateTestObject::ConstFuncRetInt3, (GIntArg, GStrArg), GPtrArg);
-		BIND_CALL_CHECK_DELEGATE_RETVAL(GStrRet, FTestDynamicDelegateRetStr2, 0, &UDynamicDelegateTestObject::FuncRetStr3, (GIntArg, GStrArg), GPtrArg);
-		BIND_CALL_CHECK_DELEGATE_RETVAL(GStrRet, FTestDynamicDelegateRetStr2, 0, &UDynamicDelegateTestObject::ConstFuncRetStr3, (GIntArg, GStrArg), GPtrArg);
+		BIND_CALL_CHECK_DELEGATE_RETVAL(GIntRet, FTestDynamicDelegateRetInt2, 0, &UDynamicDelegateTestObject::FuncRetInt3,      (GInt16Arg, GStrArg), GPtrArg);
+		BIND_CALL_CHECK_DELEGATE_RETVAL(GIntRet, FTestDynamicDelegateRetInt2, 0, &UDynamicDelegateTestObject::ConstFuncRetInt3, (GInt16Arg, GStrArg), GPtrArg);
+		BIND_CALL_CHECK_DELEGATE_RETVAL(GStrRet, FTestDynamicDelegateRetStr2, 0, &UDynamicDelegateTestObject::FuncRetStr3,      (GInt16Arg, GStrArg), GPtrArg);
+		BIND_CALL_CHECK_DELEGATE_RETVAL(GStrRet, FTestDynamicDelegateRetStr2, 0, &UDynamicDelegateTestObject::ConstFuncRetStr3, (GInt16Arg, GStrArg), GPtrArg);
 #endif
 
 		// 2 arg payload
-		BIND_CALL_CHECK_DELEGATE(FTestDynamicMulticastDelegate1::FDelegate, 0, &UDynamicDelegateTestObject::Func3, (GIntArg), GStrArg, GPtrArg);
-		BIND_CALL_CHECK_DELEGATE(FTestDynamicDelegate1, 0, &UDynamicDelegateTestObject::Func3, (GIntArg), GStrArg, GPtrArg);
-		BIND_CALL_CHECK_DELEGATE(FTestDynamicMulticastDelegate1, 1, &UDynamicDelegateTestObject::Func3, (GIntArg), GStrArg, GPtrArg);
-		BIND_CALL_CHECK_DELEGATE(FTestDynamicMulticastDelegate1::FDelegate, 0, &UDynamicDelegateTestObject::ConstFunc3, (GIntArg), GStrArg, GPtrArg);
-		BIND_CALL_CHECK_DELEGATE(FTestDynamicDelegate1, 0, &UDynamicDelegateTestObject::ConstFunc3, (GIntArg), GStrArg, GPtrArg);
-		BIND_CALL_CHECK_DELEGATE(FTestDynamicMulticastDelegate1, 1, &UDynamicDelegateTestObject::ConstFunc3, (GIntArg), GStrArg, GPtrArg);
+		BIND_CALL_CHECK_DELEGATE(FTestDynamicMulticastDelegate1::FDelegate, 0, &UDynamicDelegateTestObject::Func3,      (GInt16Arg), GStrArg, GPtrArg);
+		BIND_CALL_CHECK_DELEGATE(FTestDynamicDelegate1,                     0, &UDynamicDelegateTestObject::Func3,      (GInt16Arg), GStrArg, GPtrArg);
+		BIND_CALL_CHECK_DELEGATE(FTestDynamicMulticastDelegate1,            1, &UDynamicDelegateTestObject::Func3,      (GInt16Arg), GStrArg, GPtrArg);
+		BIND_CALL_CHECK_DELEGATE(FTestDynamicMulticastDelegate1::FDelegate, 0, &UDynamicDelegateTestObject::ConstFunc3, (GInt16Arg), GStrArg, GPtrArg);
+		BIND_CALL_CHECK_DELEGATE(FTestDynamicDelegate1,                     0, &UDynamicDelegateTestObject::ConstFunc3, (GInt16Arg), GStrArg, GPtrArg);
+		BIND_CALL_CHECK_DELEGATE(FTestDynamicMulticastDelegate1,            1, &UDynamicDelegateTestObject::ConstFunc3, (GInt16Arg), GStrArg, GPtrArg);
 #if UE_USE_DYNAMIC_DELEGATE_PAYLOADS_RETVAL
-		BIND_CALL_CHECK_DELEGATE_RETVAL(GIntRet, FTestDynamicDelegateRetInt1, 0, &UDynamicDelegateTestObject::FuncRetInt3, (GIntArg), GStrArg, GPtrArg);
-		BIND_CALL_CHECK_DELEGATE_RETVAL(GIntRet, FTestDynamicDelegateRetInt1, 0, &UDynamicDelegateTestObject::ConstFuncRetInt3, (GIntArg), GStrArg, GPtrArg);
-		BIND_CALL_CHECK_DELEGATE_RETVAL(GStrRet, FTestDynamicDelegateRetStr1, 0, &UDynamicDelegateTestObject::FuncRetStr3, (GIntArg), GStrArg, GPtrArg);
-		BIND_CALL_CHECK_DELEGATE_RETVAL(GStrRet, FTestDynamicDelegateRetStr1, 0, &UDynamicDelegateTestObject::ConstFuncRetStr3, (GIntArg), GStrArg, GPtrArg);
+		BIND_CALL_CHECK_DELEGATE_RETVAL(GIntRet, FTestDynamicDelegateRetInt1, 0, &UDynamicDelegateTestObject::FuncRetInt3,      (GInt16Arg), GStrArg, GPtrArg);
+		BIND_CALL_CHECK_DELEGATE_RETVAL(GIntRet, FTestDynamicDelegateRetInt1, 0, &UDynamicDelegateTestObject::ConstFuncRetInt3, (GInt16Arg), GStrArg, GPtrArg);
+		BIND_CALL_CHECK_DELEGATE_RETVAL(GStrRet, FTestDynamicDelegateRetStr1, 0, &UDynamicDelegateTestObject::FuncRetStr3,      (GInt16Arg), GStrArg, GPtrArg);
+		BIND_CALL_CHECK_DELEGATE_RETVAL(GStrRet, FTestDynamicDelegateRetStr1, 0, &UDynamicDelegateTestObject::ConstFuncRetStr3, (GInt16Arg), GStrArg, GPtrArg);
 #endif
 
 		// 3 arg payload
-		BIND_CALL_CHECK_DELEGATE(FTestDynamicMulticastDelegate0::FDelegate, 0, &UDynamicDelegateTestObject::Func3, (), GIntArg, GStrArg, GPtrArg);
-		BIND_CALL_CHECK_DELEGATE(FTestDynamicDelegate0, 0, &UDynamicDelegateTestObject::Func3, (), GIntArg, GStrArg, GPtrArg);
-		BIND_CALL_CHECK_DELEGATE(FTestDynamicMulticastDelegate0, 1, &UDynamicDelegateTestObject::Func3, (), GIntArg, GStrArg, GPtrArg);
-		BIND_CALL_CHECK_DELEGATE(FTestDynamicMulticastDelegate0::FDelegate, 0, &UDynamicDelegateTestObject::ConstFunc3, (), GIntArg, GStrArg, GPtrArg);
-		BIND_CALL_CHECK_DELEGATE(FTestDynamicDelegate0, 0, &UDynamicDelegateTestObject::ConstFunc3, (), GIntArg, GStrArg, GPtrArg);
-		BIND_CALL_CHECK_DELEGATE(FTestDynamicMulticastDelegate0, 1, &UDynamicDelegateTestObject::ConstFunc3, (), GIntArg, GStrArg, GPtrArg);
+		BIND_CALL_CHECK_DELEGATE(FTestDynamicMulticastDelegate0::FDelegate, 0, &UDynamicDelegateTestObject::Func3,      (), GInt16Arg, GStrArg, GPtrArg);
+		BIND_CALL_CHECK_DELEGATE(FTestDynamicDelegate0,                     0, &UDynamicDelegateTestObject::Func3,      (), GInt16Arg, GStrArg, GPtrArg);
+		BIND_CALL_CHECK_DELEGATE(FTestDynamicMulticastDelegate0,            1, &UDynamicDelegateTestObject::Func3,      (), GInt16Arg, GStrArg, GPtrArg);
+		BIND_CALL_CHECK_DELEGATE(FTestDynamicMulticastDelegate0::FDelegate, 0, &UDynamicDelegateTestObject::ConstFunc3, (), GInt16Arg, GStrArg, GPtrArg);
+		BIND_CALL_CHECK_DELEGATE(FTestDynamicDelegate0,                     0, &UDynamicDelegateTestObject::ConstFunc3, (), GInt16Arg, GStrArg, GPtrArg);
+		BIND_CALL_CHECK_DELEGATE(FTestDynamicMulticastDelegate0,            1, &UDynamicDelegateTestObject::ConstFunc3, (), GInt16Arg, GStrArg, GPtrArg);
 #if UE_USE_DYNAMIC_DELEGATE_PAYLOADS_RETVAL
-		BIND_CALL_CHECK_DELEGATE_RETVAL(GIntRet, FTestDynamicDelegateRetInt0, 0, &UDynamicDelegateTestObject::FuncRetInt3, (), GIntArg, GStrArg, GPtrArg);
-		BIND_CALL_CHECK_DELEGATE_RETVAL(GIntRet, FTestDynamicDelegateRetInt0, 0, &UDynamicDelegateTestObject::ConstFuncRetInt3, (), GIntArg, GStrArg, GPtrArg);
-		BIND_CALL_CHECK_DELEGATE_RETVAL(GStrRet, FTestDynamicDelegateRetStr0, 0, &UDynamicDelegateTestObject::FuncRetStr3, (), GIntArg, GStrArg, GPtrArg);
-		BIND_CALL_CHECK_DELEGATE_RETVAL(GStrRet, FTestDynamicDelegateRetStr0, 0, &UDynamicDelegateTestObject::ConstFuncRetStr3, (), GIntArg, GStrArg, GPtrArg);
+		BIND_CALL_CHECK_DELEGATE_RETVAL(GIntRet, FTestDynamicDelegateRetInt0, 0, &UDynamicDelegateTestObject::FuncRetInt3,      (), GInt16Arg, GStrArg, GPtrArg);
+		BIND_CALL_CHECK_DELEGATE_RETVAL(GIntRet, FTestDynamicDelegateRetInt0, 0, &UDynamicDelegateTestObject::ConstFuncRetInt3, (), GInt16Arg, GStrArg, GPtrArg);
+		BIND_CALL_CHECK_DELEGATE_RETVAL(GStrRet, FTestDynamicDelegateRetStr0, 0, &UDynamicDelegateTestObject::FuncRetStr3,      (), GInt16Arg, GStrArg, GPtrArg);
+		BIND_CALL_CHECK_DELEGATE_RETVAL(GStrRet, FTestDynamicDelegateRetStr0, 0, &UDynamicDelegateTestObject::ConstFuncRetStr3, (), GInt16Arg, GStrArg, GPtrArg);
 #endif
 #endif
 	}
@@ -679,72 +683,73 @@ void UDynamicDelegateTestObject::RunTests()
 	// 4 param function
 	{
 		// 0 arg payload
-		BIND_CALL_CHECK_DELEGATE(FTestDynamicMulticastDelegate4::FDelegate, 0, &UDynamicDelegateTestObject::Func4, (GIntArg, GStrArg, GPtrArg, GNameArg));
-		BIND_CALL_CHECK_DELEGATE(FTestDynamicDelegate4, 0, &UDynamicDelegateTestObject::Func4, (GIntArg, GStrArg, GPtrArg, GNameArg));
-		BIND_CALL_CHECK_DELEGATE(FTestDynamicMulticastDelegate4, 1, &UDynamicDelegateTestObject::Func4, (GIntArg, GStrArg, GPtrArg, GNameArg));
-		BIND_CALL_CHECK_DELEGATE(FTestDynamicMulticastDelegate4::FDelegate, 0, &UDynamicDelegateTestObject::ConstFunc4, (GIntArg, GStrArg, GPtrArg, GNameArg));
-		BIND_CALL_CHECK_DELEGATE(FTestDynamicDelegate4, 0, &UDynamicDelegateTestObject::ConstFunc4, (GIntArg, GStrArg, GPtrArg, GNameArg));
-		BIND_CALL_CHECK_DELEGATE(FTestDynamicMulticastDelegate4, 1, &UDynamicDelegateTestObject::ConstFunc4, (GIntArg, GStrArg, GPtrArg, GNameArg));
-		BIND_CALL_CHECK_DELEGATE_RETVAL(GIntRet, FTestDynamicDelegateRetInt4, 0, &UDynamicDelegateTestObject::FuncRetInt4, (GIntArg, GStrArg, GPtrArg, GNameArg));
-		BIND_CALL_CHECK_DELEGATE_RETVAL(GIntRet, FTestDynamicDelegateRetInt4, 0, &UDynamicDelegateTestObject::ConstFuncRetInt4, (GIntArg, GStrArg, GPtrArg, GNameArg));
-		BIND_CALL_CHECK_DELEGATE_RETVAL(GStrRet, FTestDynamicDelegateRetStr4, 0, &UDynamicDelegateTestObject::FuncRetStr4, (GIntArg, GStrArg, GPtrArg, GNameArg));
-		BIND_CALL_CHECK_DELEGATE_RETVAL(GStrRet, FTestDynamicDelegateRetStr4, 0, &UDynamicDelegateTestObject::ConstFuncRetStr4, (GIntArg, GStrArg, GPtrArg, GNameArg));
+		BIND_CALL_CHECK_DELEGATE(FTestDynamicMulticastDelegate4::FDelegate, 0, &UDynamicDelegateTestObject::Func4,      (GInt16Arg, GStrArg, GPtrArg, GNameArg));
+		BIND_CALL_CHECK_DELEGATE(FTestDynamicDelegate4,                     0, &UDynamicDelegateTestObject::Func4,      (GInt16Arg, GStrArg, GPtrArg, GNameArg));
+		BIND_CALL_CHECK_DELEGATE(FTestDynamicMulticastDelegate4,            1, &UDynamicDelegateTestObject::Func4,      (GInt16Arg, GStrArg, GPtrArg, GNameArg));
+		BIND_CALL_CHECK_DELEGATE(FTestDynamicMulticastDelegate4::FDelegate, 0, &UDynamicDelegateTestObject::ConstFunc4, (GInt16Arg, GStrArg, GPtrArg, GNameArg));
+		BIND_CALL_CHECK_DELEGATE(FTestDynamicDelegate4,                     0, &UDynamicDelegateTestObject::ConstFunc4, (GInt16Arg, GStrArg, GPtrArg, GNameArg));
+		BIND_CALL_CHECK_DELEGATE(FTestDynamicMulticastDelegate4,            1, &UDynamicDelegateTestObject::ConstFunc4, (GInt16Arg, GStrArg, GPtrArg, GNameArg));
+
+		BIND_CALL_CHECK_DELEGATE_RETVAL(GIntRet, FTestDynamicDelegateRetInt4, 0, &UDynamicDelegateTestObject::FuncRetInt4,      (GInt16Arg, GStrArg, GPtrArg, GNameArg));
+		BIND_CALL_CHECK_DELEGATE_RETVAL(GIntRet, FTestDynamicDelegateRetInt4, 0, &UDynamicDelegateTestObject::ConstFuncRetInt4, (GInt16Arg, GStrArg, GPtrArg, GNameArg));
+		BIND_CALL_CHECK_DELEGATE_RETVAL(GStrRet, FTestDynamicDelegateRetStr4, 0, &UDynamicDelegateTestObject::FuncRetStr4,      (GInt16Arg, GStrArg, GPtrArg, GNameArg));
+		BIND_CALL_CHECK_DELEGATE_RETVAL(GStrRet, FTestDynamicDelegateRetStr4, 0, &UDynamicDelegateTestObject::ConstFuncRetStr4, (GInt16Arg, GStrArg, GPtrArg, GNameArg));
 
 #if defined(UE_USE_DYNAMIC_DELEGATE_PAYLOADS) && UE_USE_DYNAMIC_DELEGATE_PAYLOADS
 		// 1 arg payload
-		BIND_CALL_CHECK_DELEGATE(FTestDynamicMulticastDelegate3::FDelegate, 0, &UDynamicDelegateTestObject::Func4, (GIntArg, GStrArg, GPtrArg), GNameArg);
-		BIND_CALL_CHECK_DELEGATE(FTestDynamicDelegate3, 0, &UDynamicDelegateTestObject::Func4, (GIntArg, GStrArg, GPtrArg), GNameArg);
-		BIND_CALL_CHECK_DELEGATE(FTestDynamicMulticastDelegate3, 1, &UDynamicDelegateTestObject::Func4, (GIntArg, GStrArg, GPtrArg), GNameArg);
-		BIND_CALL_CHECK_DELEGATE(FTestDynamicMulticastDelegate3::FDelegate, 0, &UDynamicDelegateTestObject::ConstFunc4, (GIntArg, GStrArg, GPtrArg), GNameArg);
-		BIND_CALL_CHECK_DELEGATE(FTestDynamicDelegate3, 0, &UDynamicDelegateTestObject::ConstFunc4, (GIntArg, GStrArg, GPtrArg), GNameArg);
-		BIND_CALL_CHECK_DELEGATE(FTestDynamicMulticastDelegate3, 1, &UDynamicDelegateTestObject::ConstFunc4, (GIntArg, GStrArg, GPtrArg), GNameArg);
+		BIND_CALL_CHECK_DELEGATE(FTestDynamicMulticastDelegate3::FDelegate, 0, &UDynamicDelegateTestObject::Func4,      (GInt16Arg, GStrArg, GPtrArg), GNameArg);
+		BIND_CALL_CHECK_DELEGATE(FTestDynamicDelegate3,                     0, &UDynamicDelegateTestObject::Func4,      (GInt16Arg, GStrArg, GPtrArg), GNameArg);
+		BIND_CALL_CHECK_DELEGATE(FTestDynamicMulticastDelegate3,            1, &UDynamicDelegateTestObject::Func4,      (GInt16Arg, GStrArg, GPtrArg), GNameArg);
+		BIND_CALL_CHECK_DELEGATE(FTestDynamicMulticastDelegate3::FDelegate, 0, &UDynamicDelegateTestObject::ConstFunc4, (GInt16Arg, GStrArg, GPtrArg), GNameArg);
+		BIND_CALL_CHECK_DELEGATE(FTestDynamicDelegate3,                     0, &UDynamicDelegateTestObject::ConstFunc4, (GInt16Arg, GStrArg, GPtrArg), GNameArg);
+		BIND_CALL_CHECK_DELEGATE(FTestDynamicMulticastDelegate3,            1, &UDynamicDelegateTestObject::ConstFunc4, (GInt16Arg, GStrArg, GPtrArg), GNameArg);
 #if UE_USE_DYNAMIC_DELEGATE_PAYLOADS_RETVAL
-		BIND_CALL_CHECK_DELEGATE_RETVAL(GIntRet, FTestDynamicDelegateRetInt3, 0, &UDynamicDelegateTestObject::FuncRetInt4, (GIntArg, GStrArg, GPtrArg), GNameArg);
-		BIND_CALL_CHECK_DELEGATE_RETVAL(GIntRet, FTestDynamicDelegateRetInt3, 0, &UDynamicDelegateTestObject::ConstFuncRetInt4, (GIntArg, GStrArg, GPtrArg), GNameArg);
-		BIND_CALL_CHECK_DELEGATE_RETVAL(GStrRet, FTestDynamicDelegateRetStr3, 0, &UDynamicDelegateTestObject::FuncRetStr4, (GIntArg, GStrArg, GPtrArg), GNameArg);
-		BIND_CALL_CHECK_DELEGATE_RETVAL(GStrRet, FTestDynamicDelegateRetStr3, 0, &UDynamicDelegateTestObject::ConstFuncRetStr4, (GIntArg, GStrArg, GPtrArg), GNameArg);
+		BIND_CALL_CHECK_DELEGATE_RETVAL(GIntRet, FTestDynamicDelegateRetInt3, 0, &UDynamicDelegateTestObject::FuncRetInt4,      (GInt16Arg, GStrArg, GPtrArg), GNameArg);
+		BIND_CALL_CHECK_DELEGATE_RETVAL(GIntRet, FTestDynamicDelegateRetInt3, 0, &UDynamicDelegateTestObject::ConstFuncRetInt4, (GInt16Arg, GStrArg, GPtrArg), GNameArg);
+		BIND_CALL_CHECK_DELEGATE_RETVAL(GStrRet, FTestDynamicDelegateRetStr3, 0, &UDynamicDelegateTestObject::FuncRetStr4,      (GInt16Arg, GStrArg, GPtrArg), GNameArg);
+		BIND_CALL_CHECK_DELEGATE_RETVAL(GStrRet, FTestDynamicDelegateRetStr3, 0, &UDynamicDelegateTestObject::ConstFuncRetStr4, (GInt16Arg, GStrArg, GPtrArg), GNameArg);
 #endif
 
 		// 2 arg payload
-		BIND_CALL_CHECK_DELEGATE(FTestDynamicMulticastDelegate2::FDelegate, 0, &UDynamicDelegateTestObject::Func4, (GIntArg, GStrArg), GPtrArg, GNameArg);
-		BIND_CALL_CHECK_DELEGATE(FTestDynamicDelegate2, 0, &UDynamicDelegateTestObject::Func4, (GIntArg, GStrArg), GPtrArg, GNameArg);
-		BIND_CALL_CHECK_DELEGATE(FTestDynamicMulticastDelegate2, 1, &UDynamicDelegateTestObject::Func4, (GIntArg, GStrArg), GPtrArg, GNameArg);
-		BIND_CALL_CHECK_DELEGATE(FTestDynamicMulticastDelegate2::FDelegate, 0, &UDynamicDelegateTestObject::ConstFunc4, (GIntArg, GStrArg), GPtrArg, GNameArg);
-		BIND_CALL_CHECK_DELEGATE(FTestDynamicDelegate2, 0, &UDynamicDelegateTestObject::ConstFunc4, (GIntArg, GStrArg), GPtrArg, GNameArg);
-		BIND_CALL_CHECK_DELEGATE(FTestDynamicMulticastDelegate2, 1, &UDynamicDelegateTestObject::ConstFunc4, (GIntArg, GStrArg), GPtrArg, GNameArg);
+		BIND_CALL_CHECK_DELEGATE(FTestDynamicMulticastDelegate2::FDelegate, 0, &UDynamicDelegateTestObject::Func4,      (GInt16Arg, GStrArg), GPtrArg, GNameArg);
+		BIND_CALL_CHECK_DELEGATE(FTestDynamicDelegate2,                     0, &UDynamicDelegateTestObject::Func4,      (GInt16Arg, GStrArg), GPtrArg, GNameArg);
+		BIND_CALL_CHECK_DELEGATE(FTestDynamicMulticastDelegate2,            1, &UDynamicDelegateTestObject::Func4,      (GInt16Arg, GStrArg), GPtrArg, GNameArg);
+		BIND_CALL_CHECK_DELEGATE(FTestDynamicMulticastDelegate2::FDelegate, 0, &UDynamicDelegateTestObject::ConstFunc4, (GInt16Arg, GStrArg), GPtrArg, GNameArg);
+		BIND_CALL_CHECK_DELEGATE(FTestDynamicDelegate2,                     0, &UDynamicDelegateTestObject::ConstFunc4, (GInt16Arg, GStrArg), GPtrArg, GNameArg);
+		BIND_CALL_CHECK_DELEGATE(FTestDynamicMulticastDelegate2,            1, &UDynamicDelegateTestObject::ConstFunc4, (GInt16Arg, GStrArg), GPtrArg, GNameArg);
 #if UE_USE_DYNAMIC_DELEGATE_PAYLOADS_RETVAL
-		BIND_CALL_CHECK_DELEGATE_RETVAL(GIntRet, FTestDynamicDelegateRetInt2, 0, &UDynamicDelegateTestObject::FuncRetInt4, (GIntArg, GStrArg), GPtrArg, GNameArg);
-		BIND_CALL_CHECK_DELEGATE_RETVAL(GIntRet, FTestDynamicDelegateRetInt2, 0, &UDynamicDelegateTestObject::ConstFuncRetInt4, (GIntArg, GStrArg), GPtrArg, GNameArg);
-		BIND_CALL_CHECK_DELEGATE_RETVAL(GStrRet, FTestDynamicDelegateRetStr2, 0, &UDynamicDelegateTestObject::FuncRetStr4, (GIntArg, GStrArg), GPtrArg, GNameArg);
-		BIND_CALL_CHECK_DELEGATE_RETVAL(GStrRet, FTestDynamicDelegateRetStr2, 0, &UDynamicDelegateTestObject::ConstFuncRetStr4, (GIntArg, GStrArg), GPtrArg, GNameArg);
+		BIND_CALL_CHECK_DELEGATE_RETVAL(GIntRet, FTestDynamicDelegateRetInt2, 0, &UDynamicDelegateTestObject::FuncRetInt4,      (GInt16Arg, GStrArg), GPtrArg, GNameArg);
+		BIND_CALL_CHECK_DELEGATE_RETVAL(GIntRet, FTestDynamicDelegateRetInt2, 0, &UDynamicDelegateTestObject::ConstFuncRetInt4, (GInt16Arg, GStrArg), GPtrArg, GNameArg);
+		BIND_CALL_CHECK_DELEGATE_RETVAL(GStrRet, FTestDynamicDelegateRetStr2, 0, &UDynamicDelegateTestObject::FuncRetStr4,      (GInt16Arg, GStrArg), GPtrArg, GNameArg);
+		BIND_CALL_CHECK_DELEGATE_RETVAL(GStrRet, FTestDynamicDelegateRetStr2, 0, &UDynamicDelegateTestObject::ConstFuncRetStr4, (GInt16Arg, GStrArg), GPtrArg, GNameArg);
 #endif
 
 		// 3 arg payload
-		BIND_CALL_CHECK_DELEGATE(FTestDynamicMulticastDelegate1::FDelegate, 0, &UDynamicDelegateTestObject::Func4, (GIntArg), GStrArg, GPtrArg, GNameArg);
-		BIND_CALL_CHECK_DELEGATE(FTestDynamicDelegate1, 0, &UDynamicDelegateTestObject::Func4, (GIntArg), GStrArg, GPtrArg, GNameArg);
-		BIND_CALL_CHECK_DELEGATE(FTestDynamicMulticastDelegate1, 1, &UDynamicDelegateTestObject::Func4, (GIntArg), GStrArg, GPtrArg, GNameArg);
-		BIND_CALL_CHECK_DELEGATE(FTestDynamicMulticastDelegate1::FDelegate, 0, &UDynamicDelegateTestObject::ConstFunc4, (GIntArg), GStrArg, GPtrArg, GNameArg);
-		BIND_CALL_CHECK_DELEGATE(FTestDynamicDelegate1, 0, &UDynamicDelegateTestObject::ConstFunc4, (GIntArg), GStrArg, GPtrArg, GNameArg);
-		BIND_CALL_CHECK_DELEGATE(FTestDynamicMulticastDelegate1, 1, &UDynamicDelegateTestObject::ConstFunc4, (GIntArg), GStrArg, GPtrArg, GNameArg);
+		BIND_CALL_CHECK_DELEGATE(FTestDynamicMulticastDelegate1::FDelegate, 0, &UDynamicDelegateTestObject::Func4,      (GInt16Arg), GStrArg, GPtrArg, GNameArg);
+		BIND_CALL_CHECK_DELEGATE(FTestDynamicDelegate1,                     0, &UDynamicDelegateTestObject::Func4,      (GInt16Arg), GStrArg, GPtrArg, GNameArg);
+		BIND_CALL_CHECK_DELEGATE(FTestDynamicMulticastDelegate1,            1, &UDynamicDelegateTestObject::Func4,      (GInt16Arg), GStrArg, GPtrArg, GNameArg);
+		BIND_CALL_CHECK_DELEGATE(FTestDynamicMulticastDelegate1::FDelegate, 0, &UDynamicDelegateTestObject::ConstFunc4, (GInt16Arg), GStrArg, GPtrArg, GNameArg);
+		BIND_CALL_CHECK_DELEGATE(FTestDynamicDelegate1,                     0, &UDynamicDelegateTestObject::ConstFunc4, (GInt16Arg), GStrArg, GPtrArg, GNameArg);
+		BIND_CALL_CHECK_DELEGATE(FTestDynamicMulticastDelegate1,            1, &UDynamicDelegateTestObject::ConstFunc4, (GInt16Arg), GStrArg, GPtrArg, GNameArg);
 #if UE_USE_DYNAMIC_DELEGATE_PAYLOADS_RETVAL
-		BIND_CALL_CHECK_DELEGATE_RETVAL(GIntRet, FTestDynamicDelegateRetInt1, 0, &UDynamicDelegateTestObject::FuncRetInt4, (GIntArg), GStrArg, GPtrArg, GNameArg);
-		BIND_CALL_CHECK_DELEGATE_RETVAL(GIntRet, FTestDynamicDelegateRetInt1, 0, &UDynamicDelegateTestObject::ConstFuncRetInt4, (GIntArg), GStrArg, GPtrArg, GNameArg);
-		BIND_CALL_CHECK_DELEGATE_RETVAL(GStrRet, FTestDynamicDelegateRetStr1, 0, &UDynamicDelegateTestObject::FuncRetStr4, (GIntArg), GStrArg, GPtrArg, GNameArg);
-		BIND_CALL_CHECK_DELEGATE_RETVAL(GStrRet, FTestDynamicDelegateRetStr1, 0, &UDynamicDelegateTestObject::ConstFuncRetStr4, (GIntArg), GStrArg, GPtrArg, GNameArg);
+		BIND_CALL_CHECK_DELEGATE_RETVAL(GIntRet, FTestDynamicDelegateRetInt1, 0, &UDynamicDelegateTestObject::FuncRetInt4,      (GInt16Arg), GStrArg, GPtrArg, GNameArg);
+		BIND_CALL_CHECK_DELEGATE_RETVAL(GIntRet, FTestDynamicDelegateRetInt1, 0, &UDynamicDelegateTestObject::ConstFuncRetInt4, (GInt16Arg), GStrArg, GPtrArg, GNameArg);
+		BIND_CALL_CHECK_DELEGATE_RETVAL(GStrRet, FTestDynamicDelegateRetStr1, 0, &UDynamicDelegateTestObject::FuncRetStr4,      (GInt16Arg), GStrArg, GPtrArg, GNameArg);
+		BIND_CALL_CHECK_DELEGATE_RETVAL(GStrRet, FTestDynamicDelegateRetStr1, 0, &UDynamicDelegateTestObject::ConstFuncRetStr4, (GInt16Arg), GStrArg, GPtrArg, GNameArg);
 #endif
 
 		// 4 arg payload
-		BIND_CALL_CHECK_DELEGATE(FTestDynamicMulticastDelegate0::FDelegate, 0, &UDynamicDelegateTestObject::Func4, (), GIntArg, GStrArg, GPtrArg, GNameArg);
-		BIND_CALL_CHECK_DELEGATE(FTestDynamicDelegate0, 0, &UDynamicDelegateTestObject::Func4, (), GIntArg, GStrArg, GPtrArg, GNameArg);
-		BIND_CALL_CHECK_DELEGATE(FTestDynamicMulticastDelegate0, 1, &UDynamicDelegateTestObject::Func4, (), GIntArg, GStrArg, GPtrArg, GNameArg);
-		BIND_CALL_CHECK_DELEGATE(FTestDynamicMulticastDelegate0::FDelegate, 0, &UDynamicDelegateTestObject::ConstFunc4, (), GIntArg, GStrArg, GPtrArg, GNameArg);
-		BIND_CALL_CHECK_DELEGATE(FTestDynamicDelegate0, 0, &UDynamicDelegateTestObject::ConstFunc4, (), GIntArg, GStrArg, GPtrArg, GNameArg);
-		BIND_CALL_CHECK_DELEGATE(FTestDynamicMulticastDelegate0, 1, &UDynamicDelegateTestObject::ConstFunc4, (), GIntArg, GStrArg, GPtrArg, GNameArg);
+		BIND_CALL_CHECK_DELEGATE(FTestDynamicMulticastDelegate0::FDelegate, 0, &UDynamicDelegateTestObject::Func4,      (), GInt16Arg, GStrArg, GPtrArg, GNameArg);
+		BIND_CALL_CHECK_DELEGATE(FTestDynamicDelegate0,                     0, &UDynamicDelegateTestObject::Func4,      (), GInt16Arg, GStrArg, GPtrArg, GNameArg);
+		BIND_CALL_CHECK_DELEGATE(FTestDynamicMulticastDelegate0,            1, &UDynamicDelegateTestObject::Func4,      (), GInt16Arg, GStrArg, GPtrArg, GNameArg);
+		BIND_CALL_CHECK_DELEGATE(FTestDynamicMulticastDelegate0::FDelegate, 0, &UDynamicDelegateTestObject::ConstFunc4, (), GInt16Arg, GStrArg, GPtrArg, GNameArg);
+		BIND_CALL_CHECK_DELEGATE(FTestDynamicDelegate0,                     0, &UDynamicDelegateTestObject::ConstFunc4, (), GInt16Arg, GStrArg, GPtrArg, GNameArg);
+		BIND_CALL_CHECK_DELEGATE(FTestDynamicMulticastDelegate0,            1, &UDynamicDelegateTestObject::ConstFunc4, (), GInt16Arg, GStrArg, GPtrArg, GNameArg);
 #if UE_USE_DYNAMIC_DELEGATE_PAYLOADS_RETVAL
-		BIND_CALL_CHECK_DELEGATE_RETVAL(GIntRet, FTestDynamicDelegateRetInt0, 0, &UDynamicDelegateTestObject::FuncRetInt4, (), GIntArg, GStrArg, GPtrArg, GNameArg);
-		BIND_CALL_CHECK_DELEGATE_RETVAL(GIntRet, FTestDynamicDelegateRetInt0, 0, &UDynamicDelegateTestObject::ConstFuncRetInt4, (), GIntArg, GStrArg, GPtrArg, GNameArg);
-		BIND_CALL_CHECK_DELEGATE_RETVAL(GStrRet, FTestDynamicDelegateRetStr0, 0, &UDynamicDelegateTestObject::FuncRetStr4, (), GIntArg, GStrArg, GPtrArg, GNameArg);
-		BIND_CALL_CHECK_DELEGATE_RETVAL(GStrRet, FTestDynamicDelegateRetStr0, 0, &UDynamicDelegateTestObject::ConstFuncRetStr4, (), GIntArg, GStrArg, GPtrArg, GNameArg);
+		BIND_CALL_CHECK_DELEGATE_RETVAL(GIntRet, FTestDynamicDelegateRetInt0, 0, &UDynamicDelegateTestObject::FuncRetInt4,      (), GInt16Arg, GStrArg, GPtrArg, GNameArg);
+		BIND_CALL_CHECK_DELEGATE_RETVAL(GIntRet, FTestDynamicDelegateRetInt0, 0, &UDynamicDelegateTestObject::ConstFuncRetInt4, (), GInt16Arg, GStrArg, GPtrArg, GNameArg);
+		BIND_CALL_CHECK_DELEGATE_RETVAL(GStrRet, FTestDynamicDelegateRetStr0, 0, &UDynamicDelegateTestObject::FuncRetStr4,      (), GInt16Arg, GStrArg, GPtrArg, GNameArg);
+		BIND_CALL_CHECK_DELEGATE_RETVAL(GStrRet, FTestDynamicDelegateRetStr0, 0, &UDynamicDelegateTestObject::ConstFuncRetStr4, (), GInt16Arg, GStrArg, GPtrArg, GNameArg);
 #endif
 #endif
 	}
@@ -752,86 +757,87 @@ void UDynamicDelegateTestObject::RunTests()
 	// 5 param function
 	{
 		// 0 arg payload
-		BIND_CALL_CHECK_DELEGATE(FTestDynamicMulticastDelegate5::FDelegate, 0, &UDynamicDelegateTestObject::Func5, (GIntArg, GStrArg, GPtrArg, GNameArg, bGBoolArg));
-		BIND_CALL_CHECK_DELEGATE(FTestDynamicDelegate5, 0, &UDynamicDelegateTestObject::Func5, (GIntArg, GStrArg, GPtrArg, GNameArg, bGBoolArg));
-		BIND_CALL_CHECK_DELEGATE(FTestDynamicMulticastDelegate5, 1, &UDynamicDelegateTestObject::Func5, (GIntArg, GStrArg, GPtrArg, GNameArg, bGBoolArg));
-		BIND_CALL_CHECK_DELEGATE(FTestDynamicMulticastDelegate5::FDelegate, 0, &UDynamicDelegateTestObject::ConstFunc5, (GIntArg, GStrArg, GPtrArg, GNameArg, bGBoolArg));
-		BIND_CALL_CHECK_DELEGATE(FTestDynamicDelegate5, 0, &UDynamicDelegateTestObject::ConstFunc5, (GIntArg, GStrArg, GPtrArg, GNameArg, bGBoolArg));
-		BIND_CALL_CHECK_DELEGATE(FTestDynamicMulticastDelegate5, 1, &UDynamicDelegateTestObject::ConstFunc5, (GIntArg, GStrArg, GPtrArg, GNameArg, bGBoolArg));
-		BIND_CALL_CHECK_DELEGATE_RETVAL(GIntRet, FTestDynamicDelegateRetInt5, 0, &UDynamicDelegateTestObject::FuncRetInt5, (GIntArg, GStrArg, GPtrArg, GNameArg, bGBoolArg));
-		BIND_CALL_CHECK_DELEGATE_RETVAL(GIntRet, FTestDynamicDelegateRetInt5, 0, &UDynamicDelegateTestObject::ConstFuncRetInt5, (GIntArg, GStrArg, GPtrArg, GNameArg, bGBoolArg));
-		BIND_CALL_CHECK_DELEGATE_RETVAL(GStrRet, FTestDynamicDelegateRetStr5, 0, &UDynamicDelegateTestObject::FuncRetStr5, (GIntArg, GStrArg, GPtrArg, GNameArg, bGBoolArg));
-		BIND_CALL_CHECK_DELEGATE_RETVAL(GStrRet, FTestDynamicDelegateRetStr5, 0, &UDynamicDelegateTestObject::ConstFuncRetStr5, (GIntArg, GStrArg, GPtrArg, GNameArg, bGBoolArg));
+		BIND_CALL_CHECK_DELEGATE(FTestDynamicMulticastDelegate5::FDelegate, 0, &UDynamicDelegateTestObject::Func5,      (GInt16Arg, GStrArg, GPtrArg, GNameArg, bGBoolArg));
+		BIND_CALL_CHECK_DELEGATE(FTestDynamicDelegate5,                     0, &UDynamicDelegateTestObject::Func5,      (GInt16Arg, GStrArg, GPtrArg, GNameArg, bGBoolArg));
+		BIND_CALL_CHECK_DELEGATE(FTestDynamicMulticastDelegate5,            1, &UDynamicDelegateTestObject::Func5,      (GInt16Arg, GStrArg, GPtrArg, GNameArg, bGBoolArg));
+		BIND_CALL_CHECK_DELEGATE(FTestDynamicMulticastDelegate5::FDelegate, 0, &UDynamicDelegateTestObject::ConstFunc5, (GInt16Arg, GStrArg, GPtrArg, GNameArg, bGBoolArg));
+		BIND_CALL_CHECK_DELEGATE(FTestDynamicDelegate5,                     0, &UDynamicDelegateTestObject::ConstFunc5, (GInt16Arg, GStrArg, GPtrArg, GNameArg, bGBoolArg));
+		BIND_CALL_CHECK_DELEGATE(FTestDynamicMulticastDelegate5,            1, &UDynamicDelegateTestObject::ConstFunc5, (GInt16Arg, GStrArg, GPtrArg, GNameArg, bGBoolArg));
+
+		BIND_CALL_CHECK_DELEGATE_RETVAL(GIntRet, FTestDynamicDelegateRetInt5, 0, &UDynamicDelegateTestObject::FuncRetInt5,      (GInt16Arg, GStrArg, GPtrArg, GNameArg, bGBoolArg));
+		BIND_CALL_CHECK_DELEGATE_RETVAL(GIntRet, FTestDynamicDelegateRetInt5, 0, &UDynamicDelegateTestObject::ConstFuncRetInt5, (GInt16Arg, GStrArg, GPtrArg, GNameArg, bGBoolArg));
+		BIND_CALL_CHECK_DELEGATE_RETVAL(GStrRet, FTestDynamicDelegateRetStr5, 0, &UDynamicDelegateTestObject::FuncRetStr5,      (GInt16Arg, GStrArg, GPtrArg, GNameArg, bGBoolArg));
+		BIND_CALL_CHECK_DELEGATE_RETVAL(GStrRet, FTestDynamicDelegateRetStr5, 0, &UDynamicDelegateTestObject::ConstFuncRetStr5, (GInt16Arg, GStrArg, GPtrArg, GNameArg, bGBoolArg));
 
 #if defined(UE_USE_DYNAMIC_DELEGATE_PAYLOADS) && UE_USE_DYNAMIC_DELEGATE_PAYLOADS
 		// 1 arg payload
-		BIND_CALL_CHECK_DELEGATE(FTestDynamicMulticastDelegate4::FDelegate, 0, &UDynamicDelegateTestObject::Func5, (GIntArg, GStrArg, GPtrArg, GNameArg), bGBoolArg);
-		BIND_CALL_CHECK_DELEGATE(FTestDynamicDelegate4, 0, &UDynamicDelegateTestObject::Func5, (GIntArg, GStrArg, GPtrArg, GNameArg), bGBoolArg);
-		BIND_CALL_CHECK_DELEGATE(FTestDynamicMulticastDelegate4, 1, &UDynamicDelegateTestObject::Func5, (GIntArg, GStrArg, GPtrArg, GNameArg), bGBoolArg);
-		BIND_CALL_CHECK_DELEGATE(FTestDynamicMulticastDelegate4::FDelegate, 0, &UDynamicDelegateTestObject::ConstFunc5, (GIntArg, GStrArg, GPtrArg, GNameArg), bGBoolArg);
-		BIND_CALL_CHECK_DELEGATE(FTestDynamicDelegate4, 0, &UDynamicDelegateTestObject::ConstFunc5, (GIntArg, GStrArg, GPtrArg, GNameArg), bGBoolArg);
-		BIND_CALL_CHECK_DELEGATE(FTestDynamicMulticastDelegate4, 1, &UDynamicDelegateTestObject::ConstFunc5, (GIntArg, GStrArg, GPtrArg, GNameArg), bGBoolArg);
+		BIND_CALL_CHECK_DELEGATE(FTestDynamicMulticastDelegate4::FDelegate, 0, &UDynamicDelegateTestObject::Func5,      (GInt16Arg, GStrArg, GPtrArg, GNameArg), bGBoolArg);
+		BIND_CALL_CHECK_DELEGATE(FTestDynamicDelegate4,                     0, &UDynamicDelegateTestObject::Func5,      (GInt16Arg, GStrArg, GPtrArg, GNameArg), bGBoolArg);
+		BIND_CALL_CHECK_DELEGATE(FTestDynamicMulticastDelegate4,            1, &UDynamicDelegateTestObject::Func5,      (GInt16Arg, GStrArg, GPtrArg, GNameArg), bGBoolArg);
+		BIND_CALL_CHECK_DELEGATE(FTestDynamicMulticastDelegate4::FDelegate, 0, &UDynamicDelegateTestObject::ConstFunc5, (GInt16Arg, GStrArg, GPtrArg, GNameArg), bGBoolArg);
+		BIND_CALL_CHECK_DELEGATE(FTestDynamicDelegate4,                     0, &UDynamicDelegateTestObject::ConstFunc5, (GInt16Arg, GStrArg, GPtrArg, GNameArg), bGBoolArg);
+		BIND_CALL_CHECK_DELEGATE(FTestDynamicMulticastDelegate4,            1, &UDynamicDelegateTestObject::ConstFunc5, (GInt16Arg, GStrArg, GPtrArg, GNameArg), bGBoolArg);
 #if UE_USE_DYNAMIC_DELEGATE_PAYLOADS_RETVAL
-		BIND_CALL_CHECK_DELEGATE_RETVAL(GIntRet, FTestDynamicDelegateRetInt4, 0, &UDynamicDelegateTestObject::FuncRetInt5, (GIntArg, GStrArg, GPtrArg, GNameArg), bGBoolArg);
-		BIND_CALL_CHECK_DELEGATE_RETVAL(GIntRet, FTestDynamicDelegateRetInt4, 0, &UDynamicDelegateTestObject::ConstFuncRetInt5, (GIntArg, GStrArg, GPtrArg, GNameArg), bGBoolArg);
-		BIND_CALL_CHECK_DELEGATE_RETVAL(GStrRet, FTestDynamicDelegateRetStr4, 0, &UDynamicDelegateTestObject::FuncRetStr5, (GIntArg, GStrArg, GPtrArg, GNameArg), bGBoolArg);
-		BIND_CALL_CHECK_DELEGATE_RETVAL(GStrRet, FTestDynamicDelegateRetStr4, 0, &UDynamicDelegateTestObject::ConstFuncRetStr5, (GIntArg, GStrArg, GPtrArg, GNameArg), bGBoolArg);
+		BIND_CALL_CHECK_DELEGATE_RETVAL(GIntRet, FTestDynamicDelegateRetInt4, 0, &UDynamicDelegateTestObject::FuncRetInt5,      (GInt16Arg, GStrArg, GPtrArg, GNameArg), bGBoolArg);
+		BIND_CALL_CHECK_DELEGATE_RETVAL(GIntRet, FTestDynamicDelegateRetInt4, 0, &UDynamicDelegateTestObject::ConstFuncRetInt5, (GInt16Arg, GStrArg, GPtrArg, GNameArg), bGBoolArg);
+		BIND_CALL_CHECK_DELEGATE_RETVAL(GStrRet, FTestDynamicDelegateRetStr4, 0, &UDynamicDelegateTestObject::FuncRetStr5,      (GInt16Arg, GStrArg, GPtrArg, GNameArg), bGBoolArg);
+		BIND_CALL_CHECK_DELEGATE_RETVAL(GStrRet, FTestDynamicDelegateRetStr4, 0, &UDynamicDelegateTestObject::ConstFuncRetStr5, (GInt16Arg, GStrArg, GPtrArg, GNameArg), bGBoolArg);
 #endif
 
 		// 2 arg payload
-		BIND_CALL_CHECK_DELEGATE(FTestDynamicMulticastDelegate3::FDelegate, 0, &UDynamicDelegateTestObject::Func5, (GIntArg, GStrArg, GPtrArg), GNameArg, bGBoolArg);
-		BIND_CALL_CHECK_DELEGATE(FTestDynamicDelegate3, 0, &UDynamicDelegateTestObject::Func5, (GIntArg, GStrArg, GPtrArg), GNameArg, bGBoolArg);
-		BIND_CALL_CHECK_DELEGATE(FTestDynamicMulticastDelegate3, 1, &UDynamicDelegateTestObject::Func5, (GIntArg, GStrArg, GPtrArg), GNameArg, bGBoolArg);
-		BIND_CALL_CHECK_DELEGATE(FTestDynamicMulticastDelegate3::FDelegate, 0, &UDynamicDelegateTestObject::ConstFunc5, (GIntArg, GStrArg, GPtrArg), GNameArg, bGBoolArg);
-		BIND_CALL_CHECK_DELEGATE(FTestDynamicDelegate3, 0, &UDynamicDelegateTestObject::ConstFunc5, (GIntArg, GStrArg, GPtrArg), GNameArg, bGBoolArg);
-		BIND_CALL_CHECK_DELEGATE(FTestDynamicMulticastDelegate3, 1, &UDynamicDelegateTestObject::ConstFunc5, (GIntArg, GStrArg, GPtrArg), GNameArg, bGBoolArg);
+		BIND_CALL_CHECK_DELEGATE(FTestDynamicMulticastDelegate3::FDelegate, 0, &UDynamicDelegateTestObject::Func5,      (GInt16Arg, GStrArg, GPtrArg), GNameArg, bGBoolArg);
+		BIND_CALL_CHECK_DELEGATE(FTestDynamicDelegate3,                     0, &UDynamicDelegateTestObject::Func5,      (GInt16Arg, GStrArg, GPtrArg), GNameArg, bGBoolArg);
+		BIND_CALL_CHECK_DELEGATE(FTestDynamicMulticastDelegate3,            1, &UDynamicDelegateTestObject::Func5,      (GInt16Arg, GStrArg, GPtrArg), GNameArg, bGBoolArg);
+		BIND_CALL_CHECK_DELEGATE(FTestDynamicMulticastDelegate3::FDelegate, 0, &UDynamicDelegateTestObject::ConstFunc5, (GInt16Arg, GStrArg, GPtrArg), GNameArg, bGBoolArg);
+		BIND_CALL_CHECK_DELEGATE(FTestDynamicDelegate3,                     0, &UDynamicDelegateTestObject::ConstFunc5, (GInt16Arg, GStrArg, GPtrArg), GNameArg, bGBoolArg);
+		BIND_CALL_CHECK_DELEGATE(FTestDynamicMulticastDelegate3,            1, &UDynamicDelegateTestObject::ConstFunc5, (GInt16Arg, GStrArg, GPtrArg), GNameArg, bGBoolArg);
 #if UE_USE_DYNAMIC_DELEGATE_PAYLOADS_RETVAL
-		BIND_CALL_CHECK_DELEGATE_RETVAL(GIntRet, FTestDynamicDelegateRetInt3, 0, &UDynamicDelegateTestObject::FuncRetInt5, (GIntArg, GStrArg, GPtrArg), GNameArg, bGBoolArg);
-		BIND_CALL_CHECK_DELEGATE_RETVAL(GIntRet, FTestDynamicDelegateRetInt3, 0, &UDynamicDelegateTestObject::ConstFuncRetInt5, (GIntArg, GStrArg, GPtrArg), GNameArg, bGBoolArg);
-		BIND_CALL_CHECK_DELEGATE_RETVAL(GStrRet, FTestDynamicDelegateRetStr3, 0, &UDynamicDelegateTestObject::FuncRetStr5, (GIntArg, GStrArg, GPtrArg), GNameArg, bGBoolArg);
-		BIND_CALL_CHECK_DELEGATE_RETVAL(GStrRet, FTestDynamicDelegateRetStr3, 0, &UDynamicDelegateTestObject::ConstFuncRetStr5, (GIntArg, GStrArg, GPtrArg), GNameArg, bGBoolArg);
+		BIND_CALL_CHECK_DELEGATE_RETVAL(GIntRet, FTestDynamicDelegateRetInt3, 0, &UDynamicDelegateTestObject::FuncRetInt5,      (GInt16Arg, GStrArg, GPtrArg), GNameArg, bGBoolArg);
+		BIND_CALL_CHECK_DELEGATE_RETVAL(GIntRet, FTestDynamicDelegateRetInt3, 0, &UDynamicDelegateTestObject::ConstFuncRetInt5, (GInt16Arg, GStrArg, GPtrArg), GNameArg, bGBoolArg);
+		BIND_CALL_CHECK_DELEGATE_RETVAL(GStrRet, FTestDynamicDelegateRetStr3, 0, &UDynamicDelegateTestObject::FuncRetStr5,      (GInt16Arg, GStrArg, GPtrArg), GNameArg, bGBoolArg);
+		BIND_CALL_CHECK_DELEGATE_RETVAL(GStrRet, FTestDynamicDelegateRetStr3, 0, &UDynamicDelegateTestObject::ConstFuncRetStr5, (GInt16Arg, GStrArg, GPtrArg), GNameArg, bGBoolArg);
 #endif
 
 		// 3 arg payload
-		BIND_CALL_CHECK_DELEGATE(FTestDynamicMulticastDelegate2::FDelegate, 0, &UDynamicDelegateTestObject::Func5, (GIntArg, GStrArg), GPtrArg, GNameArg, bGBoolArg);
-		BIND_CALL_CHECK_DELEGATE(FTestDynamicDelegate2, 0, &UDynamicDelegateTestObject::Func5, (GIntArg, GStrArg), GPtrArg, GNameArg, bGBoolArg);
-		BIND_CALL_CHECK_DELEGATE(FTestDynamicMulticastDelegate2, 1, &UDynamicDelegateTestObject::Func5, (GIntArg, GStrArg), GPtrArg, GNameArg, bGBoolArg);
-		BIND_CALL_CHECK_DELEGATE(FTestDynamicMulticastDelegate2::FDelegate, 0, &UDynamicDelegateTestObject::ConstFunc5, (GIntArg, GStrArg), GPtrArg, GNameArg, bGBoolArg);
-		BIND_CALL_CHECK_DELEGATE(FTestDynamicDelegate2, 0, &UDynamicDelegateTestObject::ConstFunc5, (GIntArg, GStrArg), GPtrArg, GNameArg, bGBoolArg);
-		BIND_CALL_CHECK_DELEGATE(FTestDynamicMulticastDelegate2, 1, &UDynamicDelegateTestObject::ConstFunc5, (GIntArg, GStrArg), GPtrArg, GNameArg, bGBoolArg);
+		BIND_CALL_CHECK_DELEGATE(FTestDynamicMulticastDelegate2::FDelegate, 0, &UDynamicDelegateTestObject::Func5,      (GInt16Arg, GStrArg), GPtrArg, GNameArg, bGBoolArg);
+		BIND_CALL_CHECK_DELEGATE(FTestDynamicDelegate2,                     0, &UDynamicDelegateTestObject::Func5,      (GInt16Arg, GStrArg), GPtrArg, GNameArg, bGBoolArg);
+		BIND_CALL_CHECK_DELEGATE(FTestDynamicMulticastDelegate2,            1, &UDynamicDelegateTestObject::Func5,      (GInt16Arg, GStrArg), GPtrArg, GNameArg, bGBoolArg);
+		BIND_CALL_CHECK_DELEGATE(FTestDynamicMulticastDelegate2::FDelegate, 0, &UDynamicDelegateTestObject::ConstFunc5, (GInt16Arg, GStrArg), GPtrArg, GNameArg, bGBoolArg);
+		BIND_CALL_CHECK_DELEGATE(FTestDynamicDelegate2,                     0, &UDynamicDelegateTestObject::ConstFunc5, (GInt16Arg, GStrArg), GPtrArg, GNameArg, bGBoolArg);
+		BIND_CALL_CHECK_DELEGATE(FTestDynamicMulticastDelegate2,            1, &UDynamicDelegateTestObject::ConstFunc5, (GInt16Arg, GStrArg), GPtrArg, GNameArg, bGBoolArg);
 #if UE_USE_DYNAMIC_DELEGATE_PAYLOADS_RETVAL
-		BIND_CALL_CHECK_DELEGATE_RETVAL(GIntRet, FTestDynamicDelegateRetInt2, 0, &UDynamicDelegateTestObject::FuncRetInt5, (GIntArg, GStrArg), GPtrArg, GNameArg, bGBoolArg);
-		BIND_CALL_CHECK_DELEGATE_RETVAL(GIntRet, FTestDynamicDelegateRetInt2, 0, &UDynamicDelegateTestObject::ConstFuncRetInt5, (GIntArg, GStrArg), GPtrArg, GNameArg, bGBoolArg);
-		BIND_CALL_CHECK_DELEGATE_RETVAL(GStrRet, FTestDynamicDelegateRetStr2, 0, &UDynamicDelegateTestObject::FuncRetStr5, (GIntArg, GStrArg), GPtrArg, GNameArg, bGBoolArg);
-		BIND_CALL_CHECK_DELEGATE_RETVAL(GStrRet, FTestDynamicDelegateRetStr2, 0, &UDynamicDelegateTestObject::ConstFuncRetStr5, (GIntArg, GStrArg), GPtrArg, GNameArg, bGBoolArg);
+		BIND_CALL_CHECK_DELEGATE_RETVAL(GIntRet, FTestDynamicDelegateRetInt2, 0, &UDynamicDelegateTestObject::FuncRetInt5,      (GInt16Arg, GStrArg), GPtrArg, GNameArg, bGBoolArg);
+		BIND_CALL_CHECK_DELEGATE_RETVAL(GIntRet, FTestDynamicDelegateRetInt2, 0, &UDynamicDelegateTestObject::ConstFuncRetInt5, (GInt16Arg, GStrArg), GPtrArg, GNameArg, bGBoolArg);
+		BIND_CALL_CHECK_DELEGATE_RETVAL(GStrRet, FTestDynamicDelegateRetStr2, 0, &UDynamicDelegateTestObject::FuncRetStr5,      (GInt16Arg, GStrArg), GPtrArg, GNameArg, bGBoolArg);
+		BIND_CALL_CHECK_DELEGATE_RETVAL(GStrRet, FTestDynamicDelegateRetStr2, 0, &UDynamicDelegateTestObject::ConstFuncRetStr5, (GInt16Arg, GStrArg), GPtrArg, GNameArg, bGBoolArg);
 #endif
 
 		// 4 arg payload
-		BIND_CALL_CHECK_DELEGATE(FTestDynamicMulticastDelegate1::FDelegate, 0, &UDynamicDelegateTestObject::Func5, (GIntArg), GStrArg, GPtrArg, GNameArg, bGBoolArg);
-		BIND_CALL_CHECK_DELEGATE(FTestDynamicDelegate1, 0, &UDynamicDelegateTestObject::Func5, (GIntArg), GStrArg, GPtrArg, GNameArg, bGBoolArg);
-		BIND_CALL_CHECK_DELEGATE(FTestDynamicMulticastDelegate1, 1, &UDynamicDelegateTestObject::Func5, (GIntArg), GStrArg, GPtrArg, GNameArg, bGBoolArg);
-		BIND_CALL_CHECK_DELEGATE(FTestDynamicMulticastDelegate1::FDelegate, 0, &UDynamicDelegateTestObject::ConstFunc5, (GIntArg), GStrArg, GPtrArg, GNameArg, bGBoolArg);
-		BIND_CALL_CHECK_DELEGATE(FTestDynamicDelegate1, 0, &UDynamicDelegateTestObject::ConstFunc5, (GIntArg), GStrArg, GPtrArg, GNameArg, bGBoolArg);
-		BIND_CALL_CHECK_DELEGATE(FTestDynamicMulticastDelegate1, 1, &UDynamicDelegateTestObject::ConstFunc5, (GIntArg), GStrArg, GPtrArg, GNameArg, bGBoolArg);
+		BIND_CALL_CHECK_DELEGATE(FTestDynamicMulticastDelegate1::FDelegate, 0, &UDynamicDelegateTestObject::Func5,      (GInt16Arg), GStrArg, GPtrArg, GNameArg, bGBoolArg);
+		BIND_CALL_CHECK_DELEGATE(FTestDynamicDelegate1,                     0, &UDynamicDelegateTestObject::Func5,      (GInt16Arg), GStrArg, GPtrArg, GNameArg, bGBoolArg);
+		BIND_CALL_CHECK_DELEGATE(FTestDynamicMulticastDelegate1,            1, &UDynamicDelegateTestObject::Func5,      (GInt16Arg), GStrArg, GPtrArg, GNameArg, bGBoolArg);
+		BIND_CALL_CHECK_DELEGATE(FTestDynamicMulticastDelegate1::FDelegate, 0, &UDynamicDelegateTestObject::ConstFunc5, (GInt16Arg), GStrArg, GPtrArg, GNameArg, bGBoolArg);
+		BIND_CALL_CHECK_DELEGATE(FTestDynamicDelegate1,                     0, &UDynamicDelegateTestObject::ConstFunc5, (GInt16Arg), GStrArg, GPtrArg, GNameArg, bGBoolArg);
+		BIND_CALL_CHECK_DELEGATE(FTestDynamicMulticastDelegate1,            1, &UDynamicDelegateTestObject::ConstFunc5, (GInt16Arg), GStrArg, GPtrArg, GNameArg, bGBoolArg);
 #if UE_USE_DYNAMIC_DELEGATE_PAYLOADS_RETVAL
-		BIND_CALL_CHECK_DELEGATE_RETVAL(GIntRet, FTestDynamicDelegateRetInt1, 0, &UDynamicDelegateTestObject::FuncRetInt5, (GIntArg), GStrArg, GPtrArg, GNameArg, bGBoolArg);
-		BIND_CALL_CHECK_DELEGATE_RETVAL(GIntRet, FTestDynamicDelegateRetInt1, 0, &UDynamicDelegateTestObject::ConstFuncRetInt5, (GIntArg), GStrArg, GPtrArg, GNameArg, bGBoolArg);
-		BIND_CALL_CHECK_DELEGATE_RETVAL(GStrRet, FTestDynamicDelegateRetStr1, 0, &UDynamicDelegateTestObject::FuncRetStr5, (GIntArg), GStrArg, GPtrArg, GNameArg, bGBoolArg);
-		BIND_CALL_CHECK_DELEGATE_RETVAL(GStrRet, FTestDynamicDelegateRetStr1, 0, &UDynamicDelegateTestObject::ConstFuncRetStr5, (GIntArg), GStrArg, GPtrArg, GNameArg, bGBoolArg);
+		BIND_CALL_CHECK_DELEGATE_RETVAL(GIntRet, FTestDynamicDelegateRetInt1, 0, &UDynamicDelegateTestObject::FuncRetInt5,      (GInt16Arg), GStrArg, GPtrArg, GNameArg, bGBoolArg);
+		BIND_CALL_CHECK_DELEGATE_RETVAL(GIntRet, FTestDynamicDelegateRetInt1, 0, &UDynamicDelegateTestObject::ConstFuncRetInt5, (GInt16Arg), GStrArg, GPtrArg, GNameArg, bGBoolArg);
+		BIND_CALL_CHECK_DELEGATE_RETVAL(GStrRet, FTestDynamicDelegateRetStr1, 0, &UDynamicDelegateTestObject::FuncRetStr5,      (GInt16Arg), GStrArg, GPtrArg, GNameArg, bGBoolArg);
+		BIND_CALL_CHECK_DELEGATE_RETVAL(GStrRet, FTestDynamicDelegateRetStr1, 0, &UDynamicDelegateTestObject::ConstFuncRetStr5, (GInt16Arg), GStrArg, GPtrArg, GNameArg, bGBoolArg);
 #endif
 
 		// 5 arg payload
-		BIND_CALL_CHECK_DELEGATE(FTestDynamicMulticastDelegate0::FDelegate, 0, &UDynamicDelegateTestObject::Func5, (), GIntArg, GStrArg, GPtrArg, GNameArg, bGBoolArg);
-		BIND_CALL_CHECK_DELEGATE(FTestDynamicDelegate0, 0, &UDynamicDelegateTestObject::Func5, (), GIntArg, GStrArg, GPtrArg, GNameArg, bGBoolArg);
-		BIND_CALL_CHECK_DELEGATE(FTestDynamicMulticastDelegate0, 1, &UDynamicDelegateTestObject::Func5, (), GIntArg, GStrArg, GPtrArg, GNameArg, bGBoolArg);
-		BIND_CALL_CHECK_DELEGATE(FTestDynamicMulticastDelegate0::FDelegate, 0, &UDynamicDelegateTestObject::ConstFunc5, (), GIntArg, GStrArg, GPtrArg, GNameArg, bGBoolArg);
-		BIND_CALL_CHECK_DELEGATE(FTestDynamicDelegate0, 0, &UDynamicDelegateTestObject::ConstFunc5, (), GIntArg, GStrArg, GPtrArg, GNameArg, bGBoolArg);
-		BIND_CALL_CHECK_DELEGATE(FTestDynamicMulticastDelegate0, 1, &UDynamicDelegateTestObject::ConstFunc5, (), GIntArg, GStrArg, GPtrArg, GNameArg, bGBoolArg);
+		BIND_CALL_CHECK_DELEGATE(FTestDynamicMulticastDelegate0::FDelegate, 0, &UDynamicDelegateTestObject::Func5,      (), GInt16Arg, GStrArg, GPtrArg, GNameArg, bGBoolArg);
+		BIND_CALL_CHECK_DELEGATE(FTestDynamicDelegate0,                     0, &UDynamicDelegateTestObject::Func5,      (), GInt16Arg, GStrArg, GPtrArg, GNameArg, bGBoolArg);
+		BIND_CALL_CHECK_DELEGATE(FTestDynamicMulticastDelegate0,            1, &UDynamicDelegateTestObject::Func5,      (), GInt16Arg, GStrArg, GPtrArg, GNameArg, bGBoolArg);
+		BIND_CALL_CHECK_DELEGATE(FTestDynamicMulticastDelegate0::FDelegate, 0, &UDynamicDelegateTestObject::ConstFunc5, (), GInt16Arg, GStrArg, GPtrArg, GNameArg, bGBoolArg);
+		BIND_CALL_CHECK_DELEGATE(FTestDynamicDelegate0,                     0, &UDynamicDelegateTestObject::ConstFunc5, (), GInt16Arg, GStrArg, GPtrArg, GNameArg, bGBoolArg);
+		BIND_CALL_CHECK_DELEGATE(FTestDynamicMulticastDelegate0,            1, &UDynamicDelegateTestObject::ConstFunc5, (), GInt16Arg, GStrArg, GPtrArg, GNameArg, bGBoolArg);
 #if UE_USE_DYNAMIC_DELEGATE_PAYLOADS_RETVAL
-		BIND_CALL_CHECK_DELEGATE_RETVAL(GIntRet, FTestDynamicDelegateRetInt0, 0, &UDynamicDelegateTestObject::FuncRetInt5, (), GIntArg, GStrArg, GPtrArg, GNameArg, bGBoolArg);
-		BIND_CALL_CHECK_DELEGATE_RETVAL(GIntRet, FTestDynamicDelegateRetInt0, 0, &UDynamicDelegateTestObject::ConstFuncRetInt5, (), GIntArg, GStrArg, GPtrArg, GNameArg, bGBoolArg);
-		BIND_CALL_CHECK_DELEGATE_RETVAL(GStrRet, FTestDynamicDelegateRetStr0, 0, &UDynamicDelegateTestObject::FuncRetStr5, (), GIntArg, GStrArg, GPtrArg, GNameArg, bGBoolArg);
-		BIND_CALL_CHECK_DELEGATE_RETVAL(GStrRet, FTestDynamicDelegateRetStr0, 0, &UDynamicDelegateTestObject::ConstFuncRetStr5, (), GIntArg, GStrArg, GPtrArg, GNameArg, bGBoolArg);
+		BIND_CALL_CHECK_DELEGATE_RETVAL(GIntRet, FTestDynamicDelegateRetInt0, 0, &UDynamicDelegateTestObject::FuncRetInt5,      (), GInt16Arg, GStrArg, GPtrArg, GNameArg, bGBoolArg);
+		BIND_CALL_CHECK_DELEGATE_RETVAL(GIntRet, FTestDynamicDelegateRetInt0, 0, &UDynamicDelegateTestObject::ConstFuncRetInt5, (), GInt16Arg, GStrArg, GPtrArg, GNameArg, bGBoolArg);
+		BIND_CALL_CHECK_DELEGATE_RETVAL(GStrRet, FTestDynamicDelegateRetStr0, 0, &UDynamicDelegateTestObject::FuncRetStr5,      (), GInt16Arg, GStrArg, GPtrArg, GNameArg, bGBoolArg);
+		BIND_CALL_CHECK_DELEGATE_RETVAL(GStrRet, FTestDynamicDelegateRetStr0, 0, &UDynamicDelegateTestObject::ConstFuncRetStr5, (), GInt16Arg, GStrArg, GPtrArg, GNameArg, bGBoolArg);
 #endif
 #endif
 	}
@@ -839,152 +845,153 @@ void UDynamicDelegateTestObject::RunTests()
 	// 6 param function
 	{
 		// 0 arg payload
-		BIND_CALL_CHECK_DELEGATE(FTestDynamicMulticastDelegate6::FDelegate, 0, &UDynamicDelegateTestObject::Func6, (GIntArg, GStrArg, GPtrArg, GNameArg, bGBoolArg, GWeakObjPtrArg));
-		BIND_CALL_CHECK_DELEGATE(FTestDynamicDelegate6, 0, &UDynamicDelegateTestObject::Func6, (GIntArg, GStrArg, GPtrArg, GNameArg, bGBoolArg, GWeakObjPtrArg));
-		BIND_CALL_CHECK_DELEGATE(FTestDynamicMulticastDelegate6, 1, &UDynamicDelegateTestObject::Func6, (GIntArg, GStrArg, GPtrArg, GNameArg, bGBoolArg, GWeakObjPtrArg));
-		BIND_CALL_CHECK_DELEGATE(FTestDynamicMulticastDelegate6::FDelegate, 0, &UDynamicDelegateTestObject::ConstFunc6, (GIntArg, GStrArg, GPtrArg, GNameArg, bGBoolArg, GWeakObjPtrArg));
-		BIND_CALL_CHECK_DELEGATE(FTestDynamicDelegate6, 0, &UDynamicDelegateTestObject::ConstFunc6, (GIntArg, GStrArg, GPtrArg, GNameArg, bGBoolArg, GWeakObjPtrArg));
-		BIND_CALL_CHECK_DELEGATE(FTestDynamicMulticastDelegate6, 1, &UDynamicDelegateTestObject::ConstFunc6, (GIntArg, GStrArg, GPtrArg, GNameArg, bGBoolArg, GWeakObjPtrArg));
-		BIND_CALL_CHECK_DELEGATE_RETVAL(GIntRet, FTestDynamicDelegateRetInt6, 0, &UDynamicDelegateTestObject::FuncRetInt6, (GIntArg, GStrArg, GPtrArg, GNameArg, bGBoolArg, GWeakObjPtrArg));
-		BIND_CALL_CHECK_DELEGATE_RETVAL(GIntRet, FTestDynamicDelegateRetInt6, 0, &UDynamicDelegateTestObject::ConstFuncRetInt6, (GIntArg, GStrArg, GPtrArg, GNameArg, bGBoolArg, GWeakObjPtrArg));
-		BIND_CALL_CHECK_DELEGATE_RETVAL(GStrRet, FTestDynamicDelegateRetStr6, 0, &UDynamicDelegateTestObject::FuncRetStr6, (GIntArg, GStrArg, GPtrArg, GNameArg, bGBoolArg, GWeakObjPtrArg));
-		BIND_CALL_CHECK_DELEGATE_RETVAL(GStrRet, FTestDynamicDelegateRetStr6, 0, &UDynamicDelegateTestObject::ConstFuncRetStr6, (GIntArg, GStrArg, GPtrArg, GNameArg, bGBoolArg, GWeakObjPtrArg));
+		BIND_CALL_CHECK_DELEGATE(FTestDynamicMulticastDelegate6::FDelegate, 0, &UDynamicDelegateTestObject::Func6,      (GInt16Arg, GStrArg, GPtrArg, GNameArg, bGBoolArg, GWeakObjPtrArg));
+		BIND_CALL_CHECK_DELEGATE(FTestDynamicDelegate6,                     0, &UDynamicDelegateTestObject::Func6,      (GInt16Arg, GStrArg, GPtrArg, GNameArg, bGBoolArg, GWeakObjPtrArg));
+		BIND_CALL_CHECK_DELEGATE(FTestDynamicMulticastDelegate6,            1, &UDynamicDelegateTestObject::Func6,      (GInt16Arg, GStrArg, GPtrArg, GNameArg, bGBoolArg, GWeakObjPtrArg));
+		BIND_CALL_CHECK_DELEGATE(FTestDynamicMulticastDelegate6::FDelegate, 0, &UDynamicDelegateTestObject::ConstFunc6, (GInt16Arg, GStrArg, GPtrArg, GNameArg, bGBoolArg, GWeakObjPtrArg));
+		BIND_CALL_CHECK_DELEGATE(FTestDynamicDelegate6,                     0, &UDynamicDelegateTestObject::ConstFunc6, (GInt16Arg, GStrArg, GPtrArg, GNameArg, bGBoolArg, GWeakObjPtrArg));
+		BIND_CALL_CHECK_DELEGATE(FTestDynamicMulticastDelegate6,            1, &UDynamicDelegateTestObject::ConstFunc6, (GInt16Arg, GStrArg, GPtrArg, GNameArg, bGBoolArg, GWeakObjPtrArg));
+
+		BIND_CALL_CHECK_DELEGATE_RETVAL(GIntRet, FTestDynamicDelegateRetInt6, 0, &UDynamicDelegateTestObject::FuncRetInt6,      (GInt16Arg, GStrArg, GPtrArg, GNameArg, bGBoolArg, GWeakObjPtrArg));
+		BIND_CALL_CHECK_DELEGATE_RETVAL(GIntRet, FTestDynamicDelegateRetInt6, 0, &UDynamicDelegateTestObject::ConstFuncRetInt6, (GInt16Arg, GStrArg, GPtrArg, GNameArg, bGBoolArg, GWeakObjPtrArg));
+		BIND_CALL_CHECK_DELEGATE_RETVAL(GStrRet, FTestDynamicDelegateRetStr6, 0, &UDynamicDelegateTestObject::FuncRetStr6,      (GInt16Arg, GStrArg, GPtrArg, GNameArg, bGBoolArg, GWeakObjPtrArg));
+		BIND_CALL_CHECK_DELEGATE_RETVAL(GStrRet, FTestDynamicDelegateRetStr6, 0, &UDynamicDelegateTestObject::ConstFuncRetStr6, (GInt16Arg, GStrArg, GPtrArg, GNameArg, bGBoolArg, GWeakObjPtrArg));
 
 #if defined(UE_USE_DYNAMIC_DELEGATE_PAYLOADS) && UE_USE_DYNAMIC_DELEGATE_PAYLOADS
 		// 1 arg payload
-		BIND_CALL_CHECK_DELEGATE(FTestDynamicMulticastDelegate5::FDelegate, 0, &UDynamicDelegateTestObject::Func6, (GIntArg, GStrArg, GPtrArg, GNameArg, bGBoolArg), GWeakObjPtrArg);
-		BIND_CALL_CHECK_DELEGATE(FTestDynamicDelegate5, 0, &UDynamicDelegateTestObject::Func6, (GIntArg, GStrArg, GPtrArg, GNameArg, bGBoolArg), GWeakObjPtrArg);
-		BIND_CALL_CHECK_DELEGATE(FTestDynamicMulticastDelegate5, 1, &UDynamicDelegateTestObject::Func6, (GIntArg, GStrArg, GPtrArg, GNameArg, bGBoolArg), GWeakObjPtrArg);
-		BIND_CALL_CHECK_DELEGATE(FTestDynamicMulticastDelegate5::FDelegate, 0, &UDynamicDelegateTestObject::ConstFunc6, (GIntArg, GStrArg, GPtrArg, GNameArg, bGBoolArg), GWeakObjPtrArg);
-		BIND_CALL_CHECK_DELEGATE(FTestDynamicDelegate5, 0, &UDynamicDelegateTestObject::ConstFunc6, (GIntArg, GStrArg, GPtrArg, GNameArg, bGBoolArg), GWeakObjPtrArg);
-		BIND_CALL_CHECK_DELEGATE(FTestDynamicMulticastDelegate5, 1, &UDynamicDelegateTestObject::ConstFunc6, (GIntArg, GStrArg, GPtrArg, GNameArg, bGBoolArg), GWeakObjPtrArg);
+		BIND_CALL_CHECK_DELEGATE(FTestDynamicMulticastDelegate5::FDelegate, 0, &UDynamicDelegateTestObject::Func6,      (GInt16Arg, GStrArg, GPtrArg, GNameArg, bGBoolArg), GWeakObjPtrArg);
+		BIND_CALL_CHECK_DELEGATE(FTestDynamicDelegate5,                     0, &UDynamicDelegateTestObject::Func6,      (GInt16Arg, GStrArg, GPtrArg, GNameArg, bGBoolArg), GWeakObjPtrArg);
+		BIND_CALL_CHECK_DELEGATE(FTestDynamicMulticastDelegate5,            1, &UDynamicDelegateTestObject::Func6,      (GInt16Arg, GStrArg, GPtrArg, GNameArg, bGBoolArg), GWeakObjPtrArg);
+		BIND_CALL_CHECK_DELEGATE(FTestDynamicMulticastDelegate5::FDelegate, 0, &UDynamicDelegateTestObject::ConstFunc6, (GInt16Arg, GStrArg, GPtrArg, GNameArg, bGBoolArg), GWeakObjPtrArg);
+		BIND_CALL_CHECK_DELEGATE(FTestDynamicDelegate5,                     0, &UDynamicDelegateTestObject::ConstFunc6, (GInt16Arg, GStrArg, GPtrArg, GNameArg, bGBoolArg), GWeakObjPtrArg);
+		BIND_CALL_CHECK_DELEGATE(FTestDynamicMulticastDelegate5,            1, &UDynamicDelegateTestObject::ConstFunc6, (GInt16Arg, GStrArg, GPtrArg, GNameArg, bGBoolArg), GWeakObjPtrArg);
 #if UE_USE_DYNAMIC_DELEGATE_PAYLOADS_RETVAL
-		BIND_CALL_CHECK_DELEGATE_RETVAL(GIntRet, FTestDynamicDelegateRetInt5, 0, &UDynamicDelegateTestObject::FuncRetInt6, (GIntArg, GStrArg, GPtrArg, GNameArg, bGBoolArg), GWeakObjPtrArg);
-		BIND_CALL_CHECK_DELEGATE_RETVAL(GIntRet, FTestDynamicDelegateRetInt5, 0, &UDynamicDelegateTestObject::ConstFuncRetInt6, (GIntArg, GStrArg, GPtrArg, GNameArg, bGBoolArg), GWeakObjPtrArg);
-		BIND_CALL_CHECK_DELEGATE_RETVAL(GStrRet, FTestDynamicDelegateRetStr5, 0, &UDynamicDelegateTestObject::FuncRetStr6, (GIntArg, GStrArg, GPtrArg, GNameArg, bGBoolArg), GWeakObjPtrArg);
-		BIND_CALL_CHECK_DELEGATE_RETVAL(GStrRet, FTestDynamicDelegateRetStr5, 0, &UDynamicDelegateTestObject::ConstFuncRetStr6, (GIntArg, GStrArg, GPtrArg, GNameArg, bGBoolArg), GWeakObjPtrArg);
+		BIND_CALL_CHECK_DELEGATE_RETVAL(GIntRet, FTestDynamicDelegateRetInt5, 0, &UDynamicDelegateTestObject::FuncRetInt6,      (GInt16Arg, GStrArg, GPtrArg, GNameArg, bGBoolArg), GWeakObjPtrArg);
+		BIND_CALL_CHECK_DELEGATE_RETVAL(GIntRet, FTestDynamicDelegateRetInt5, 0, &UDynamicDelegateTestObject::ConstFuncRetInt6, (GInt16Arg, GStrArg, GPtrArg, GNameArg, bGBoolArg), GWeakObjPtrArg);
+		BIND_CALL_CHECK_DELEGATE_RETVAL(GStrRet, FTestDynamicDelegateRetStr5, 0, &UDynamicDelegateTestObject::FuncRetStr6,      (GInt16Arg, GStrArg, GPtrArg, GNameArg, bGBoolArg), GWeakObjPtrArg);
+		BIND_CALL_CHECK_DELEGATE_RETVAL(GStrRet, FTestDynamicDelegateRetStr5, 0, &UDynamicDelegateTestObject::ConstFuncRetStr6, (GInt16Arg, GStrArg, GPtrArg, GNameArg, bGBoolArg), GWeakObjPtrArg);
 #endif
 
 		// 2 arg payload
-		BIND_CALL_CHECK_DELEGATE(FTestDynamicMulticastDelegate4::FDelegate, 0, &UDynamicDelegateTestObject::Func6, (GIntArg, GStrArg, GPtrArg, GNameArg), bGBoolArg, GWeakObjPtrArg);
-		BIND_CALL_CHECK_DELEGATE(FTestDynamicDelegate4, 0, &UDynamicDelegateTestObject::Func6, (GIntArg, GStrArg, GPtrArg, GNameArg), bGBoolArg, GWeakObjPtrArg);
-		BIND_CALL_CHECK_DELEGATE(FTestDynamicMulticastDelegate4, 1, &UDynamicDelegateTestObject::Func6, (GIntArg, GStrArg, GPtrArg, GNameArg), bGBoolArg, GWeakObjPtrArg);
-		BIND_CALL_CHECK_DELEGATE(FTestDynamicMulticastDelegate4::FDelegate, 0, &UDynamicDelegateTestObject::ConstFunc6, (GIntArg, GStrArg, GPtrArg, GNameArg), bGBoolArg, GWeakObjPtrArg);
-		BIND_CALL_CHECK_DELEGATE(FTestDynamicDelegate4, 0, &UDynamicDelegateTestObject::ConstFunc6, (GIntArg, GStrArg, GPtrArg, GNameArg), bGBoolArg, GWeakObjPtrArg);
-		BIND_CALL_CHECK_DELEGATE(FTestDynamicMulticastDelegate4, 1, &UDynamicDelegateTestObject::ConstFunc6, (GIntArg, GStrArg, GPtrArg, GNameArg), bGBoolArg, GWeakObjPtrArg);
+		BIND_CALL_CHECK_DELEGATE(FTestDynamicMulticastDelegate4::FDelegate, 0, &UDynamicDelegateTestObject::Func6,      (GInt16Arg, GStrArg, GPtrArg, GNameArg), bGBoolArg, GWeakObjPtrArg);
+		BIND_CALL_CHECK_DELEGATE(FTestDynamicDelegate4,                     0, &UDynamicDelegateTestObject::Func6,      (GInt16Arg, GStrArg, GPtrArg, GNameArg), bGBoolArg, GWeakObjPtrArg);
+		BIND_CALL_CHECK_DELEGATE(FTestDynamicMulticastDelegate4,            1, &UDynamicDelegateTestObject::Func6,      (GInt16Arg, GStrArg, GPtrArg, GNameArg), bGBoolArg, GWeakObjPtrArg);
+		BIND_CALL_CHECK_DELEGATE(FTestDynamicMulticastDelegate4::FDelegate, 0, &UDynamicDelegateTestObject::ConstFunc6, (GInt16Arg, GStrArg, GPtrArg, GNameArg), bGBoolArg, GWeakObjPtrArg);
+		BIND_CALL_CHECK_DELEGATE(FTestDynamicDelegate4,                     0, &UDynamicDelegateTestObject::ConstFunc6, (GInt16Arg, GStrArg, GPtrArg, GNameArg), bGBoolArg, GWeakObjPtrArg);
+		BIND_CALL_CHECK_DELEGATE(FTestDynamicMulticastDelegate4,            1, &UDynamicDelegateTestObject::ConstFunc6, (GInt16Arg, GStrArg, GPtrArg, GNameArg), bGBoolArg, GWeakObjPtrArg);
 #if UE_USE_DYNAMIC_DELEGATE_PAYLOADS_RETVAL
-		BIND_CALL_CHECK_DELEGATE_RETVAL(GIntRet, FTestDynamicDelegateRetInt4, 0, &UDynamicDelegateTestObject::FuncRetInt6, (GIntArg, GStrArg, GPtrArg, GNameArg), bGBoolArg, GWeakObjPtrArg);
-		BIND_CALL_CHECK_DELEGATE_RETVAL(GIntRet, FTestDynamicDelegateRetInt4, 0, &UDynamicDelegateTestObject::ConstFuncRetInt6, (GIntArg, GStrArg, GPtrArg, GNameArg), bGBoolArg, GWeakObjPtrArg);
-		BIND_CALL_CHECK_DELEGATE_RETVAL(GStrRet, FTestDynamicDelegateRetStr4, 0, &UDynamicDelegateTestObject::FuncRetStr6, (GIntArg, GStrArg, GPtrArg, GNameArg), bGBoolArg, GWeakObjPtrArg);
-		BIND_CALL_CHECK_DELEGATE_RETVAL(GStrRet, FTestDynamicDelegateRetStr4, 0, &UDynamicDelegateTestObject::ConstFuncRetStr6, (GIntArg, GStrArg, GPtrArg, GNameArg), bGBoolArg, GWeakObjPtrArg);
+		BIND_CALL_CHECK_DELEGATE_RETVAL(GIntRet, FTestDynamicDelegateRetInt4, 0, &UDynamicDelegateTestObject::FuncRetInt6,      (GInt16Arg, GStrArg, GPtrArg, GNameArg), bGBoolArg, GWeakObjPtrArg);
+		BIND_CALL_CHECK_DELEGATE_RETVAL(GIntRet, FTestDynamicDelegateRetInt4, 0, &UDynamicDelegateTestObject::ConstFuncRetInt6, (GInt16Arg, GStrArg, GPtrArg, GNameArg), bGBoolArg, GWeakObjPtrArg);
+		BIND_CALL_CHECK_DELEGATE_RETVAL(GStrRet, FTestDynamicDelegateRetStr4, 0, &UDynamicDelegateTestObject::FuncRetStr6,      (GInt16Arg, GStrArg, GPtrArg, GNameArg), bGBoolArg, GWeakObjPtrArg);
+		BIND_CALL_CHECK_DELEGATE_RETVAL(GStrRet, FTestDynamicDelegateRetStr4, 0, &UDynamicDelegateTestObject::ConstFuncRetStr6, (GInt16Arg, GStrArg, GPtrArg, GNameArg), bGBoolArg, GWeakObjPtrArg);
 #endif
 
 		// 3 arg payload
-		BIND_CALL_CHECK_DELEGATE(FTestDynamicMulticastDelegate3::FDelegate, 0, &UDynamicDelegateTestObject::Func6, (GIntArg, GStrArg, GPtrArg), GNameArg, bGBoolArg, GWeakObjPtrArg);
-		BIND_CALL_CHECK_DELEGATE(FTestDynamicDelegate3, 0, &UDynamicDelegateTestObject::Func6, (GIntArg, GStrArg, GPtrArg), GNameArg, bGBoolArg, GWeakObjPtrArg);
-		BIND_CALL_CHECK_DELEGATE(FTestDynamicMulticastDelegate3, 1, &UDynamicDelegateTestObject::Func6, (GIntArg, GStrArg, GPtrArg), GNameArg, bGBoolArg, GWeakObjPtrArg);
-		BIND_CALL_CHECK_DELEGATE(FTestDynamicMulticastDelegate3::FDelegate, 0, &UDynamicDelegateTestObject::ConstFunc6, (GIntArg, GStrArg, GPtrArg), GNameArg, bGBoolArg, GWeakObjPtrArg);
-		BIND_CALL_CHECK_DELEGATE(FTestDynamicDelegate3, 0, &UDynamicDelegateTestObject::ConstFunc6, (GIntArg, GStrArg, GPtrArg), GNameArg, bGBoolArg, GWeakObjPtrArg);
-		BIND_CALL_CHECK_DELEGATE(FTestDynamicMulticastDelegate3, 1, &UDynamicDelegateTestObject::ConstFunc6, (GIntArg, GStrArg, GPtrArg), GNameArg, bGBoolArg, GWeakObjPtrArg);
+		BIND_CALL_CHECK_DELEGATE(FTestDynamicMulticastDelegate3::FDelegate, 0, &UDynamicDelegateTestObject::Func6,      (GInt16Arg, GStrArg, GPtrArg), GNameArg, bGBoolArg, GWeakObjPtrArg);
+		BIND_CALL_CHECK_DELEGATE(FTestDynamicDelegate3,                     0, &UDynamicDelegateTestObject::Func6,      (GInt16Arg, GStrArg, GPtrArg), GNameArg, bGBoolArg, GWeakObjPtrArg);
+		BIND_CALL_CHECK_DELEGATE(FTestDynamicMulticastDelegate3,            1, &UDynamicDelegateTestObject::Func6,      (GInt16Arg, GStrArg, GPtrArg), GNameArg, bGBoolArg, GWeakObjPtrArg);
+		BIND_CALL_CHECK_DELEGATE(FTestDynamicMulticastDelegate3::FDelegate, 0, &UDynamicDelegateTestObject::ConstFunc6, (GInt16Arg, GStrArg, GPtrArg), GNameArg, bGBoolArg, GWeakObjPtrArg);
+		BIND_CALL_CHECK_DELEGATE(FTestDynamicDelegate3,                     0, &UDynamicDelegateTestObject::ConstFunc6, (GInt16Arg, GStrArg, GPtrArg), GNameArg, bGBoolArg, GWeakObjPtrArg);
+		BIND_CALL_CHECK_DELEGATE(FTestDynamicMulticastDelegate3,            1, &UDynamicDelegateTestObject::ConstFunc6, (GInt16Arg, GStrArg, GPtrArg), GNameArg, bGBoolArg, GWeakObjPtrArg);
 #if UE_USE_DYNAMIC_DELEGATE_PAYLOADS_RETVAL
-		BIND_CALL_CHECK_DELEGATE_RETVAL(GIntRet, FTestDynamicDelegateRetInt3, 0, &UDynamicDelegateTestObject::FuncRetInt6, (GIntArg, GStrArg, GPtrArg), GNameArg, bGBoolArg, GWeakObjPtrArg);
-		BIND_CALL_CHECK_DELEGATE_RETVAL(GIntRet, FTestDynamicDelegateRetInt3, 0, &UDynamicDelegateTestObject::ConstFuncRetInt6, (GIntArg, GStrArg, GPtrArg), GNameArg, bGBoolArg, GWeakObjPtrArg);
-		BIND_CALL_CHECK_DELEGATE_RETVAL(GStrRet, FTestDynamicDelegateRetStr3, 0, &UDynamicDelegateTestObject::FuncRetStr6, (GIntArg, GStrArg, GPtrArg), GNameArg, bGBoolArg, GWeakObjPtrArg);
-		BIND_CALL_CHECK_DELEGATE_RETVAL(GStrRet, FTestDynamicDelegateRetStr3, 0, &UDynamicDelegateTestObject::ConstFuncRetStr6, (GIntArg, GStrArg, GPtrArg), GNameArg, bGBoolArg, GWeakObjPtrArg);
+		BIND_CALL_CHECK_DELEGATE_RETVAL(GIntRet, FTestDynamicDelegateRetInt3, 0, &UDynamicDelegateTestObject::FuncRetInt6,      (GInt16Arg, GStrArg, GPtrArg), GNameArg, bGBoolArg, GWeakObjPtrArg);
+		BIND_CALL_CHECK_DELEGATE_RETVAL(GIntRet, FTestDynamicDelegateRetInt3, 0, &UDynamicDelegateTestObject::ConstFuncRetInt6, (GInt16Arg, GStrArg, GPtrArg), GNameArg, bGBoolArg, GWeakObjPtrArg);
+		BIND_CALL_CHECK_DELEGATE_RETVAL(GStrRet, FTestDynamicDelegateRetStr3, 0, &UDynamicDelegateTestObject::FuncRetStr6,      (GInt16Arg, GStrArg, GPtrArg), GNameArg, bGBoolArg, GWeakObjPtrArg);
+		BIND_CALL_CHECK_DELEGATE_RETVAL(GStrRet, FTestDynamicDelegateRetStr3, 0, &UDynamicDelegateTestObject::ConstFuncRetStr6, (GInt16Arg, GStrArg, GPtrArg), GNameArg, bGBoolArg, GWeakObjPtrArg);
 #endif
 
 		// 4 arg payload
-		BIND_CALL_CHECK_DELEGATE(FTestDynamicMulticastDelegate2::FDelegate, 0, &UDynamicDelegateTestObject::Func6, (GIntArg, GStrArg), GPtrArg, GNameArg, bGBoolArg, GWeakObjPtrArg);
-		BIND_CALL_CHECK_DELEGATE(FTestDynamicDelegate2, 0, &UDynamicDelegateTestObject::Func6, (GIntArg, GStrArg), GPtrArg, GNameArg, bGBoolArg, GWeakObjPtrArg);
-		BIND_CALL_CHECK_DELEGATE(FTestDynamicMulticastDelegate2, 1, &UDynamicDelegateTestObject::Func6, (GIntArg, GStrArg), GPtrArg, GNameArg, bGBoolArg, GWeakObjPtrArg);
-		BIND_CALL_CHECK_DELEGATE(FTestDynamicMulticastDelegate2::FDelegate, 0, &UDynamicDelegateTestObject::ConstFunc6, (GIntArg, GStrArg), GPtrArg, GNameArg, bGBoolArg, GWeakObjPtrArg);
-		BIND_CALL_CHECK_DELEGATE(FTestDynamicDelegate2, 0, &UDynamicDelegateTestObject::ConstFunc6, (GIntArg, GStrArg), GPtrArg, GNameArg, bGBoolArg, GWeakObjPtrArg);
-		BIND_CALL_CHECK_DELEGATE(FTestDynamicMulticastDelegate2, 1, &UDynamicDelegateTestObject::ConstFunc6, (GIntArg, GStrArg), GPtrArg, GNameArg, bGBoolArg, GWeakObjPtrArg);
+		BIND_CALL_CHECK_DELEGATE(FTestDynamicMulticastDelegate2::FDelegate, 0, &UDynamicDelegateTestObject::Func6,      (GInt16Arg, GStrArg), GPtrArg, GNameArg, bGBoolArg, GWeakObjPtrArg);
+		BIND_CALL_CHECK_DELEGATE(FTestDynamicDelegate2,                     0, &UDynamicDelegateTestObject::Func6,      (GInt16Arg, GStrArg), GPtrArg, GNameArg, bGBoolArg, GWeakObjPtrArg);
+		BIND_CALL_CHECK_DELEGATE(FTestDynamicMulticastDelegate2,            1, &UDynamicDelegateTestObject::Func6,      (GInt16Arg, GStrArg), GPtrArg, GNameArg, bGBoolArg, GWeakObjPtrArg);
+		BIND_CALL_CHECK_DELEGATE(FTestDynamicMulticastDelegate2::FDelegate, 0, &UDynamicDelegateTestObject::ConstFunc6, (GInt16Arg, GStrArg), GPtrArg, GNameArg, bGBoolArg, GWeakObjPtrArg);
+		BIND_CALL_CHECK_DELEGATE(FTestDynamicDelegate2,                     0, &UDynamicDelegateTestObject::ConstFunc6, (GInt16Arg, GStrArg), GPtrArg, GNameArg, bGBoolArg, GWeakObjPtrArg);
+		BIND_CALL_CHECK_DELEGATE(FTestDynamicMulticastDelegate2,            1, &UDynamicDelegateTestObject::ConstFunc6, (GInt16Arg, GStrArg), GPtrArg, GNameArg, bGBoolArg, GWeakObjPtrArg);
 #if UE_USE_DYNAMIC_DELEGATE_PAYLOADS_RETVAL
-		BIND_CALL_CHECK_DELEGATE_RETVAL(GIntRet, FTestDynamicDelegateRetInt2, 0, &UDynamicDelegateTestObject::FuncRetInt6, (GIntArg, GStrArg), GPtrArg, GNameArg, bGBoolArg, GWeakObjPtrArg);
-		BIND_CALL_CHECK_DELEGATE_RETVAL(GIntRet, FTestDynamicDelegateRetInt2, 0, &UDynamicDelegateTestObject::ConstFuncRetInt6, (GIntArg, GStrArg), GPtrArg, GNameArg, bGBoolArg, GWeakObjPtrArg);
-		BIND_CALL_CHECK_DELEGATE_RETVAL(GStrRet, FTestDynamicDelegateRetStr2, 0, &UDynamicDelegateTestObject::FuncRetStr6, (GIntArg, GStrArg), GPtrArg, GNameArg, bGBoolArg, GWeakObjPtrArg);
-		BIND_CALL_CHECK_DELEGATE_RETVAL(GStrRet, FTestDynamicDelegateRetStr2, 0, &UDynamicDelegateTestObject::ConstFuncRetStr6, (GIntArg, GStrArg), GPtrArg, GNameArg, bGBoolArg, GWeakObjPtrArg);
+		BIND_CALL_CHECK_DELEGATE_RETVAL(GIntRet, FTestDynamicDelegateRetInt2, 0, &UDynamicDelegateTestObject::FuncRetInt6,      (GInt16Arg, GStrArg), GPtrArg, GNameArg, bGBoolArg, GWeakObjPtrArg);
+		BIND_CALL_CHECK_DELEGATE_RETVAL(GIntRet, FTestDynamicDelegateRetInt2, 0, &UDynamicDelegateTestObject::ConstFuncRetInt6, (GInt16Arg, GStrArg), GPtrArg, GNameArg, bGBoolArg, GWeakObjPtrArg);
+		BIND_CALL_CHECK_DELEGATE_RETVAL(GStrRet, FTestDynamicDelegateRetStr2, 0, &UDynamicDelegateTestObject::FuncRetStr6,      (GInt16Arg, GStrArg), GPtrArg, GNameArg, bGBoolArg, GWeakObjPtrArg);
+		BIND_CALL_CHECK_DELEGATE_RETVAL(GStrRet, FTestDynamicDelegateRetStr2, 0, &UDynamicDelegateTestObject::ConstFuncRetStr6, (GInt16Arg, GStrArg), GPtrArg, GNameArg, bGBoolArg, GWeakObjPtrArg);
 #endif
 
 		// 5 arg payload
-		BIND_CALL_CHECK_DELEGATE(FTestDynamicMulticastDelegate1::FDelegate, 0, &UDynamicDelegateTestObject::Func6, (GIntArg), GStrArg, GPtrArg, GNameArg, bGBoolArg, GWeakObjPtrArg);
-		BIND_CALL_CHECK_DELEGATE(FTestDynamicDelegate1, 0, &UDynamicDelegateTestObject::Func6, (GIntArg), GStrArg, GPtrArg, GNameArg, bGBoolArg, GWeakObjPtrArg);
-		BIND_CALL_CHECK_DELEGATE(FTestDynamicMulticastDelegate1, 1, &UDynamicDelegateTestObject::Func6, (GIntArg), GStrArg, GPtrArg, GNameArg, bGBoolArg, GWeakObjPtrArg);
-		BIND_CALL_CHECK_DELEGATE(FTestDynamicMulticastDelegate1::FDelegate, 0, &UDynamicDelegateTestObject::ConstFunc6, (GIntArg), GStrArg, GPtrArg, GNameArg, bGBoolArg, GWeakObjPtrArg);
-		BIND_CALL_CHECK_DELEGATE(FTestDynamicDelegate1, 0, &UDynamicDelegateTestObject::ConstFunc6, (GIntArg), GStrArg, GPtrArg, GNameArg, bGBoolArg, GWeakObjPtrArg);
-		BIND_CALL_CHECK_DELEGATE(FTestDynamicMulticastDelegate1, 1, &UDynamicDelegateTestObject::ConstFunc6, (GIntArg), GStrArg, GPtrArg, GNameArg, bGBoolArg, GWeakObjPtrArg);
+		BIND_CALL_CHECK_DELEGATE(FTestDynamicMulticastDelegate1::FDelegate, 0, &UDynamicDelegateTestObject::Func6,      (GInt16Arg), GStrArg, GPtrArg, GNameArg, bGBoolArg, GWeakObjPtrArg);
+		BIND_CALL_CHECK_DELEGATE(FTestDynamicDelegate1,                     0, &UDynamicDelegateTestObject::Func6,      (GInt16Arg), GStrArg, GPtrArg, GNameArg, bGBoolArg, GWeakObjPtrArg);
+		BIND_CALL_CHECK_DELEGATE(FTestDynamicMulticastDelegate1,            1, &UDynamicDelegateTestObject::Func6,      (GInt16Arg), GStrArg, GPtrArg, GNameArg, bGBoolArg, GWeakObjPtrArg);
+		BIND_CALL_CHECK_DELEGATE(FTestDynamicMulticastDelegate1::FDelegate, 0, &UDynamicDelegateTestObject::ConstFunc6, (GInt16Arg), GStrArg, GPtrArg, GNameArg, bGBoolArg, GWeakObjPtrArg);
+		BIND_CALL_CHECK_DELEGATE(FTestDynamicDelegate1,                     0, &UDynamicDelegateTestObject::ConstFunc6, (GInt16Arg), GStrArg, GPtrArg, GNameArg, bGBoolArg, GWeakObjPtrArg);
+		BIND_CALL_CHECK_DELEGATE(FTestDynamicMulticastDelegate1,            1, &UDynamicDelegateTestObject::ConstFunc6, (GInt16Arg), GStrArg, GPtrArg, GNameArg, bGBoolArg, GWeakObjPtrArg);
 #if UE_USE_DYNAMIC_DELEGATE_PAYLOADS_RETVAL
-		BIND_CALL_CHECK_DELEGATE_RETVAL(GIntRet, FTestDynamicDelegateRetInt1, 0, &UDynamicDelegateTestObject::FuncRetInt6, (GIntArg), GStrArg, GPtrArg, GNameArg, bGBoolArg, GWeakObjPtrArg);
-		BIND_CALL_CHECK_DELEGATE_RETVAL(GIntRet, FTestDynamicDelegateRetInt1, 0, &UDynamicDelegateTestObject::ConstFuncRetInt6, (GIntArg), GStrArg, GPtrArg, GNameArg, bGBoolArg, GWeakObjPtrArg);
-		BIND_CALL_CHECK_DELEGATE_RETVAL(GStrRet, FTestDynamicDelegateRetStr1, 0, &UDynamicDelegateTestObject::FuncRetStr6, (GIntArg), GStrArg, GPtrArg, GNameArg, bGBoolArg, GWeakObjPtrArg);
-		BIND_CALL_CHECK_DELEGATE_RETVAL(GStrRet, FTestDynamicDelegateRetStr1, 0, &UDynamicDelegateTestObject::ConstFuncRetStr6, (GIntArg), GStrArg, GPtrArg, GNameArg, bGBoolArg, GWeakObjPtrArg);
+		BIND_CALL_CHECK_DELEGATE_RETVAL(GIntRet, FTestDynamicDelegateRetInt1, 0, &UDynamicDelegateTestObject::FuncRetInt6,      (GInt16Arg), GStrArg, GPtrArg, GNameArg, bGBoolArg, GWeakObjPtrArg);
+		BIND_CALL_CHECK_DELEGATE_RETVAL(GIntRet, FTestDynamicDelegateRetInt1, 0, &UDynamicDelegateTestObject::ConstFuncRetInt6, (GInt16Arg), GStrArg, GPtrArg, GNameArg, bGBoolArg, GWeakObjPtrArg);
+		BIND_CALL_CHECK_DELEGATE_RETVAL(GStrRet, FTestDynamicDelegateRetStr1, 0, &UDynamicDelegateTestObject::FuncRetStr6,      (GInt16Arg), GStrArg, GPtrArg, GNameArg, bGBoolArg, GWeakObjPtrArg);
+		BIND_CALL_CHECK_DELEGATE_RETVAL(GStrRet, FTestDynamicDelegateRetStr1, 0, &UDynamicDelegateTestObject::ConstFuncRetStr6, (GInt16Arg), GStrArg, GPtrArg, GNameArg, bGBoolArg, GWeakObjPtrArg);
 #endif
 
 		// 6 arg payload
-		BIND_CALL_CHECK_DELEGATE(FTestDynamicMulticastDelegate0::FDelegate, 0, &UDynamicDelegateTestObject::Func6, (), GIntArg, GStrArg, GPtrArg, GNameArg, bGBoolArg, GWeakObjPtrArg);
-		BIND_CALL_CHECK_DELEGATE(FTestDynamicDelegate0, 0, &UDynamicDelegateTestObject::Func6, (), GIntArg, GStrArg, GPtrArg, GNameArg, bGBoolArg, GWeakObjPtrArg);
-		BIND_CALL_CHECK_DELEGATE(FTestDynamicMulticastDelegate0, 1, &UDynamicDelegateTestObject::Func6, (), GIntArg, GStrArg, GPtrArg, GNameArg, bGBoolArg, GWeakObjPtrArg);
-		BIND_CALL_CHECK_DELEGATE(FTestDynamicMulticastDelegate0::FDelegate, 0, &UDynamicDelegateTestObject::ConstFunc6, (), GIntArg, GStrArg, GPtrArg, GNameArg, bGBoolArg, GWeakObjPtrArg);
-		BIND_CALL_CHECK_DELEGATE(FTestDynamicDelegate0, 0, &UDynamicDelegateTestObject::ConstFunc6, (), GIntArg, GStrArg, GPtrArg, GNameArg, bGBoolArg, GWeakObjPtrArg);
-		BIND_CALL_CHECK_DELEGATE(FTestDynamicMulticastDelegate0, 1, &UDynamicDelegateTestObject::ConstFunc6, (), GIntArg, GStrArg, GPtrArg, GNameArg, bGBoolArg, GWeakObjPtrArg);
+		BIND_CALL_CHECK_DELEGATE(FTestDynamicMulticastDelegate0::FDelegate, 0, &UDynamicDelegateTestObject::Func6,      (), GInt16Arg, GStrArg, GPtrArg, GNameArg, bGBoolArg, GWeakObjPtrArg);
+		BIND_CALL_CHECK_DELEGATE(FTestDynamicDelegate0,                     0, &UDynamicDelegateTestObject::Func6,      (), GInt16Arg, GStrArg, GPtrArg, GNameArg, bGBoolArg, GWeakObjPtrArg);
+		BIND_CALL_CHECK_DELEGATE(FTestDynamicMulticastDelegate0,            1, &UDynamicDelegateTestObject::Func6,      (), GInt16Arg, GStrArg, GPtrArg, GNameArg, bGBoolArg, GWeakObjPtrArg);
+		BIND_CALL_CHECK_DELEGATE(FTestDynamicMulticastDelegate0::FDelegate, 0, &UDynamicDelegateTestObject::ConstFunc6, (), GInt16Arg, GStrArg, GPtrArg, GNameArg, bGBoolArg, GWeakObjPtrArg);
+		BIND_CALL_CHECK_DELEGATE(FTestDynamicDelegate0,                     0, &UDynamicDelegateTestObject::ConstFunc6, (), GInt16Arg, GStrArg, GPtrArg, GNameArg, bGBoolArg, GWeakObjPtrArg);
+		BIND_CALL_CHECK_DELEGATE(FTestDynamicMulticastDelegate0,            1, &UDynamicDelegateTestObject::ConstFunc6, (), GInt16Arg, GStrArg, GPtrArg, GNameArg, bGBoolArg, GWeakObjPtrArg);
 #if UE_USE_DYNAMIC_DELEGATE_PAYLOADS_RETVAL
-		BIND_CALL_CHECK_DELEGATE_RETVAL(GIntRet, FTestDynamicDelegateRetInt0, 0, &UDynamicDelegateTestObject::FuncRetInt6, (), GIntArg, GStrArg, GPtrArg, GNameArg, bGBoolArg, GWeakObjPtrArg);
-		BIND_CALL_CHECK_DELEGATE_RETVAL(GIntRet, FTestDynamicDelegateRetInt0, 0, &UDynamicDelegateTestObject::ConstFuncRetInt6, (), GIntArg, GStrArg, GPtrArg, GNameArg, bGBoolArg, GWeakObjPtrArg);
-		BIND_CALL_CHECK_DELEGATE_RETVAL(GStrRet, FTestDynamicDelegateRetStr0, 0, &UDynamicDelegateTestObject::FuncRetStr6, (), GIntArg, GStrArg, GPtrArg, GNameArg, bGBoolArg, GWeakObjPtrArg);
-		BIND_CALL_CHECK_DELEGATE_RETVAL(GStrRet, FTestDynamicDelegateRetStr0, 0, &UDynamicDelegateTestObject::ConstFuncRetStr6, (), GIntArg, GStrArg, GPtrArg, GNameArg, bGBoolArg, GWeakObjPtrArg);
+		BIND_CALL_CHECK_DELEGATE_RETVAL(GIntRet, FTestDynamicDelegateRetInt0, 0, &UDynamicDelegateTestObject::FuncRetInt6,      (), GInt16Arg, GStrArg, GPtrArg, GNameArg, bGBoolArg, GWeakObjPtrArg);
+		BIND_CALL_CHECK_DELEGATE_RETVAL(GIntRet, FTestDynamicDelegateRetInt0, 0, &UDynamicDelegateTestObject::ConstFuncRetInt6, (), GInt16Arg, GStrArg, GPtrArg, GNameArg, bGBoolArg, GWeakObjPtrArg);
+		BIND_CALL_CHECK_DELEGATE_RETVAL(GStrRet, FTestDynamicDelegateRetStr0, 0, &UDynamicDelegateTestObject::FuncRetStr6,      (), GInt16Arg, GStrArg, GPtrArg, GNameArg, bGBoolArg, GWeakObjPtrArg);
+		BIND_CALL_CHECK_DELEGATE_RETVAL(GStrRet, FTestDynamicDelegateRetStr0, 0, &UDynamicDelegateTestObject::ConstFuncRetStr6, (), GInt16Arg, GStrArg, GPtrArg, GNameArg, bGBoolArg, GWeakObjPtrArg);
 #endif
 #endif
 	}
 
 	// 1 out param function
 	{
-		BIND_CALL_CHECK_DELEGATE_OUT(FTestDynamicMulticastDelegateOut1::FDelegate, 0, &UDynamicDelegateTestObject::FuncOut1, (GIntArg), (GIntOut));
-		BIND_CALL_CHECK_DELEGATE_OUT(FTestDynamicDelegateOut1, 0, &UDynamicDelegateTestObject::FuncOut1, (GIntArg), (GIntOut));
-		BIND_CALL_CHECK_DELEGATE_OUT(FTestDynamicMulticastDelegateOut1, 1, &UDynamicDelegateTestObject::FuncOut1, (GIntArg), (GIntOut));
-		BIND_CALL_CHECK_DELEGATE_OUT(FTestDynamicMulticastDelegateOut1::FDelegate, 0, &UDynamicDelegateTestObject::ConstFuncOut1, (GIntArg), (GIntOut));
-		BIND_CALL_CHECK_DELEGATE_OUT(FTestDynamicDelegateOut1, 0, &UDynamicDelegateTestObject::ConstFuncOut1, (GIntArg), (GIntOut));
-		BIND_CALL_CHECK_DELEGATE_OUT(FTestDynamicMulticastDelegateOut1, 1, &UDynamicDelegateTestObject::ConstFuncOut1, (GIntArg), (GIntOut));
+		BIND_CALL_CHECK_DELEGATE_OUT(FTestDynamicMulticastDelegateOut1::FDelegate, 0, &UDynamicDelegateTestObject::FuncOut1,      (GInt16Arg), (GInt16Out));
+		BIND_CALL_CHECK_DELEGATE_OUT(FTestDynamicDelegateOut1,                     0, &UDynamicDelegateTestObject::FuncOut1,      (GInt16Arg), (GInt16Out));
+		BIND_CALL_CHECK_DELEGATE_OUT(FTestDynamicMulticastDelegateOut1,            1, &UDynamicDelegateTestObject::FuncOut1,      (GInt16Arg), (GInt16Out));
+		BIND_CALL_CHECK_DELEGATE_OUT(FTestDynamicMulticastDelegateOut1::FDelegate, 0, &UDynamicDelegateTestObject::ConstFuncOut1, (GInt16Arg), (GInt16Out));
+		BIND_CALL_CHECK_DELEGATE_OUT(FTestDynamicDelegateOut1,                     0, &UDynamicDelegateTestObject::ConstFuncOut1, (GInt16Arg), (GInt16Out));
+		BIND_CALL_CHECK_DELEGATE_OUT(FTestDynamicMulticastDelegateOut1,            1, &UDynamicDelegateTestObject::ConstFuncOut1, (GInt16Arg), (GInt16Out));
 	}
 
 	// 2 out param function
 	{
-		BIND_CALL_CHECK_DELEGATE_OUT(FTestDynamicMulticastDelegateOut2::FDelegate, 0, &UDynamicDelegateTestObject::FuncOut2, (GIntArg, GStrArg), (GIntOut, GStrOut));
-		BIND_CALL_CHECK_DELEGATE_OUT(FTestDynamicDelegateOut2, 0, &UDynamicDelegateTestObject::FuncOut2, (GIntArg, GStrArg), (GIntOut, GStrOut));
-		BIND_CALL_CHECK_DELEGATE_OUT(FTestDynamicMulticastDelegateOut2, 1, &UDynamicDelegateTestObject::FuncOut2, (GIntArg, GStrArg), (GIntOut, GStrOut));
-		BIND_CALL_CHECK_DELEGATE_OUT(FTestDynamicMulticastDelegateOut2::FDelegate, 0, &UDynamicDelegateTestObject::ConstFuncOut2, (GIntArg, GStrArg), (GIntOut, GStrOut));
-		BIND_CALL_CHECK_DELEGATE_OUT(FTestDynamicDelegateOut2, 0, &UDynamicDelegateTestObject::ConstFuncOut2, (GIntArg, GStrArg), (GIntOut, GStrOut));
-		BIND_CALL_CHECK_DELEGATE_OUT(FTestDynamicMulticastDelegateOut2, 1, &UDynamicDelegateTestObject::ConstFuncOut2, (GIntArg, GStrArg), (GIntOut, GStrOut));
+		BIND_CALL_CHECK_DELEGATE_OUT(FTestDynamicMulticastDelegateOut2::FDelegate, 0, &UDynamicDelegateTestObject::FuncOut2,      (GInt16Arg, GStrArg), (GInt16Out, GStrOut));
+		BIND_CALL_CHECK_DELEGATE_OUT(FTestDynamicDelegateOut2,                     0, &UDynamicDelegateTestObject::FuncOut2,      (GInt16Arg, GStrArg), (GInt16Out, GStrOut));
+		BIND_CALL_CHECK_DELEGATE_OUT(FTestDynamicMulticastDelegateOut2,            1, &UDynamicDelegateTestObject::FuncOut2,      (GInt16Arg, GStrArg), (GInt16Out, GStrOut));
+		BIND_CALL_CHECK_DELEGATE_OUT(FTestDynamicMulticastDelegateOut2::FDelegate, 0, &UDynamicDelegateTestObject::ConstFuncOut2, (GInt16Arg, GStrArg), (GInt16Out, GStrOut));
+		BIND_CALL_CHECK_DELEGATE_OUT(FTestDynamicDelegateOut2,                     0, &UDynamicDelegateTestObject::ConstFuncOut2, (GInt16Arg, GStrArg), (GInt16Out, GStrOut));
+		BIND_CALL_CHECK_DELEGATE_OUT(FTestDynamicMulticastDelegateOut2,            1, &UDynamicDelegateTestObject::ConstFuncOut2, (GInt16Arg, GStrArg), (GInt16Out, GStrOut));
 	}
 
 	// 3 out param function
 	{
-		BIND_CALL_CHECK_DELEGATE_OUT(FTestDynamicMulticastDelegateOut3::FDelegate, 0, &UDynamicDelegateTestObject::FuncOut3, (GIntArg, GStrArg, GNameArg), (GIntOut, GStrOut, GNameOut));
-		BIND_CALL_CHECK_DELEGATE_OUT(FTestDynamicDelegateOut3, 0, &UDynamicDelegateTestObject::FuncOut3, (GIntArg, GStrArg, GNameArg), (GIntOut, GStrOut, GNameOut));
-		BIND_CALL_CHECK_DELEGATE_OUT(FTestDynamicMulticastDelegateOut3, 1, &UDynamicDelegateTestObject::FuncOut3, (GIntArg, GStrArg, GNameArg), (GIntOut, GStrOut, GNameOut));
-		BIND_CALL_CHECK_DELEGATE_OUT(FTestDynamicMulticastDelegateOut3::FDelegate, 0, &UDynamicDelegateTestObject::ConstFuncOut3, (GIntArg, GStrArg, GNameArg), (GIntOut, GStrOut, GNameOut));
-		BIND_CALL_CHECK_DELEGATE_OUT(FTestDynamicDelegateOut3, 0, &UDynamicDelegateTestObject::ConstFuncOut3, (GIntArg, GStrArg, GNameArg), (GIntOut, GStrOut, GNameOut));
-		BIND_CALL_CHECK_DELEGATE_OUT(FTestDynamicMulticastDelegateOut3, 1, &UDynamicDelegateTestObject::ConstFuncOut3, (GIntArg, GStrArg, GNameArg), (GIntOut, GStrOut, GNameOut));
+		BIND_CALL_CHECK_DELEGATE_OUT(FTestDynamicMulticastDelegateOut3::FDelegate, 0, &UDynamicDelegateTestObject::FuncOut3,      (GInt16Arg, GStrArg, GNameArg), (GInt16Out, GStrOut, GNameOut));
+		BIND_CALL_CHECK_DELEGATE_OUT(FTestDynamicDelegateOut3,                     0, &UDynamicDelegateTestObject::FuncOut3,      (GInt16Arg, GStrArg, GNameArg), (GInt16Out, GStrOut, GNameOut));
+		BIND_CALL_CHECK_DELEGATE_OUT(FTestDynamicMulticastDelegateOut3,            1, &UDynamicDelegateTestObject::FuncOut3,      (GInt16Arg, GStrArg, GNameArg), (GInt16Out, GStrOut, GNameOut));
+		BIND_CALL_CHECK_DELEGATE_OUT(FTestDynamicMulticastDelegateOut3::FDelegate, 0, &UDynamicDelegateTestObject::ConstFuncOut3, (GInt16Arg, GStrArg, GNameArg), (GInt16Out, GStrOut, GNameOut));
+		BIND_CALL_CHECK_DELEGATE_OUT(FTestDynamicDelegateOut3,                     0, &UDynamicDelegateTestObject::ConstFuncOut3, (GInt16Arg, GStrArg, GNameArg), (GInt16Out, GStrOut, GNameOut));
+		BIND_CALL_CHECK_DELEGATE_OUT(FTestDynamicMulticastDelegateOut3,            1, &UDynamicDelegateTestObject::ConstFuncOut3, (GInt16Arg, GStrArg, GNameArg), (GInt16Out, GStrOut, GNameOut));
 	}
 
 	// 4 out param function
 	{
-		BIND_CALL_CHECK_DELEGATE_OUT(FTestDynamicMulticastDelegateOut4::FDelegate, 0, &UDynamicDelegateTestObject::FuncOut4, (GIntArg, GStrArg, GNameArg, bGBoolArg), (GIntOut, GStrOut, GNameOut, bGBoolOut));
-		BIND_CALL_CHECK_DELEGATE_OUT(FTestDynamicDelegateOut4, 0, &UDynamicDelegateTestObject::FuncOut4, (GIntArg, GStrArg, GNameArg, bGBoolArg), (GIntOut, GStrOut, GNameOut, bGBoolOut));
-		BIND_CALL_CHECK_DELEGATE_OUT(FTestDynamicMulticastDelegateOut4, 1, &UDynamicDelegateTestObject::FuncOut4, (GIntArg, GStrArg, GNameArg, bGBoolArg), (GIntOut, GStrOut, GNameOut, bGBoolOut));
-		BIND_CALL_CHECK_DELEGATE_OUT(FTestDynamicMulticastDelegateOut4::FDelegate, 0, &UDynamicDelegateTestObject::ConstFuncOut4, (GIntArg, GStrArg, GNameArg, bGBoolArg), (GIntOut, GStrOut, GNameOut, bGBoolOut));
-		BIND_CALL_CHECK_DELEGATE_OUT(FTestDynamicDelegateOut4, 0, &UDynamicDelegateTestObject::ConstFuncOut4, (GIntArg, GStrArg, GNameArg, bGBoolArg), (GIntOut, GStrOut, GNameOut, bGBoolOut));
-		BIND_CALL_CHECK_DELEGATE_OUT(FTestDynamicMulticastDelegateOut4, 1, &UDynamicDelegateTestObject::ConstFuncOut4, (GIntArg, GStrArg, GNameArg, bGBoolArg), (GIntOut, GStrOut, GNameOut, bGBoolOut));
+		BIND_CALL_CHECK_DELEGATE_OUT(FTestDynamicMulticastDelegateOut4::FDelegate, 0, &UDynamicDelegateTestObject::FuncOut4,      (GInt16Arg, GStrArg, GNameArg, bGBoolArg), (GInt16Out, GStrOut, GNameOut, bGBoolOut));
+		BIND_CALL_CHECK_DELEGATE_OUT(FTestDynamicDelegateOut4,                     0, &UDynamicDelegateTestObject::FuncOut4,      (GInt16Arg, GStrArg, GNameArg, bGBoolArg), (GInt16Out, GStrOut, GNameOut, bGBoolOut));
+		BIND_CALL_CHECK_DELEGATE_OUT(FTestDynamicMulticastDelegateOut4,            1, &UDynamicDelegateTestObject::FuncOut4,      (GInt16Arg, GStrArg, GNameArg, bGBoolArg), (GInt16Out, GStrOut, GNameOut, bGBoolOut));
+		BIND_CALL_CHECK_DELEGATE_OUT(FTestDynamicMulticastDelegateOut4::FDelegate, 0, &UDynamicDelegateTestObject::ConstFuncOut4, (GInt16Arg, GStrArg, GNameArg, bGBoolArg), (GInt16Out, GStrOut, GNameOut, bGBoolOut));
+		BIND_CALL_CHECK_DELEGATE_OUT(FTestDynamicDelegateOut4,                     0, &UDynamicDelegateTestObject::ConstFuncOut4, (GInt16Arg, GStrArg, GNameArg, bGBoolArg), (GInt16Out, GStrOut, GNameOut, bGBoolOut));
+		BIND_CALL_CHECK_DELEGATE_OUT(FTestDynamicMulticastDelegateOut4,            1, &UDynamicDelegateTestObject::ConstFuncOut4, (GInt16Arg, GStrArg, GNameArg, bGBoolArg), (GInt16Out, GStrOut, GNameOut, bGBoolOut));
 	}
 
 	// 5 out param function
 	{
-		BIND_CALL_CHECK_DELEGATE_OUT(FTestDynamicMulticastDelegateOut5::FDelegate, 0, &UDynamicDelegateTestObject::FuncOut5, (GIntArg, GStrArg, GNameArg, bGBoolArg, GWeakObjPtrArg), (GIntOut, GStrOut, GNameOut, bGBoolOut, GWeakObjPtrOut));
-		BIND_CALL_CHECK_DELEGATE_OUT(FTestDynamicDelegateOut5, 0, &UDynamicDelegateTestObject::FuncOut5, (GIntArg, GStrArg, GNameArg, bGBoolArg, GWeakObjPtrArg), (GIntOut, GStrOut, GNameOut, bGBoolOut, GWeakObjPtrOut));
-		BIND_CALL_CHECK_DELEGATE_OUT(FTestDynamicMulticastDelegateOut5, 1, &UDynamicDelegateTestObject::FuncOut5, (GIntArg, GStrArg, GNameArg, bGBoolArg, GWeakObjPtrArg), (GIntOut, GStrOut, GNameOut, bGBoolOut, GWeakObjPtrOut));
-		BIND_CALL_CHECK_DELEGATE_OUT(FTestDynamicMulticastDelegateOut5::FDelegate, 0, &UDynamicDelegateTestObject::ConstFuncOut5, (GIntArg, GStrArg, GNameArg, bGBoolArg, GWeakObjPtrArg), (GIntOut, GStrOut, GNameOut, bGBoolOut, GWeakObjPtrOut));
-		BIND_CALL_CHECK_DELEGATE_OUT(FTestDynamicDelegateOut5, 0, &UDynamicDelegateTestObject::ConstFuncOut5, (GIntArg, GStrArg, GNameArg, bGBoolArg, GWeakObjPtrArg), (GIntOut, GStrOut, GNameOut, bGBoolOut, GWeakObjPtrOut));
-		BIND_CALL_CHECK_DELEGATE_OUT(FTestDynamicMulticastDelegateOut5, 1, &UDynamicDelegateTestObject::ConstFuncOut5, (GIntArg, GStrArg, GNameArg, bGBoolArg, GWeakObjPtrArg), (GIntOut, GStrOut, GNameOut, bGBoolOut, GWeakObjPtrOut));
+		BIND_CALL_CHECK_DELEGATE_OUT(FTestDynamicMulticastDelegateOut5::FDelegate, 0, &UDynamicDelegateTestObject::FuncOut5,      (GInt16Arg, GStrArg, GNameArg, bGBoolArg, GWeakObjPtrArg), (GInt16Out, GStrOut, GNameOut, bGBoolOut, GWeakObjPtrOut));
+		BIND_CALL_CHECK_DELEGATE_OUT(FTestDynamicDelegateOut5,                     0, &UDynamicDelegateTestObject::FuncOut5,      (GInt16Arg, GStrArg, GNameArg, bGBoolArg, GWeakObjPtrArg), (GInt16Out, GStrOut, GNameOut, bGBoolOut, GWeakObjPtrOut));
+		BIND_CALL_CHECK_DELEGATE_OUT(FTestDynamicMulticastDelegateOut5,            1, &UDynamicDelegateTestObject::FuncOut5,      (GInt16Arg, GStrArg, GNameArg, bGBoolArg, GWeakObjPtrArg), (GInt16Out, GStrOut, GNameOut, bGBoolOut, GWeakObjPtrOut));
+		BIND_CALL_CHECK_DELEGATE_OUT(FTestDynamicMulticastDelegateOut5::FDelegate, 0, &UDynamicDelegateTestObject::ConstFuncOut5, (GInt16Arg, GStrArg, GNameArg, bGBoolArg, GWeakObjPtrArg), (GInt16Out, GStrOut, GNameOut, bGBoolOut, GWeakObjPtrOut));
+		BIND_CALL_CHECK_DELEGATE_OUT(FTestDynamicDelegateOut5,                     0, &UDynamicDelegateTestObject::ConstFuncOut5, (GInt16Arg, GStrArg, GNameArg, bGBoolArg, GWeakObjPtrArg), (GInt16Out, GStrOut, GNameOut, bGBoolOut, GWeakObjPtrOut));
+		BIND_CALL_CHECK_DELEGATE_OUT(FTestDynamicMulticastDelegateOut5,            1, &UDynamicDelegateTestObject::ConstFuncOut5, (GInt16Arg, GStrArg, GNameArg, bGBoolArg, GWeakObjPtrArg), (GInt16Out, GStrOut, GNameOut, bGBoolOut, GWeakObjPtrOut));
 	}
 
 	// Unscoped enum
diff --git a/Engine/Source/Runtime/CoreUObject/Public/UObject/DynamicDelegateTestObject.h b/Engine/Source/Runtime/CoreUObject/Public/UObject/DynamicDelegateTestObject.h
--- a/Engine/Source/Runtime/CoreUObject/Public/UObject/DynamicDelegateTestObject.h
+++ b/Engine/Source/Runtime/CoreUObject/Public/UObject/DynamicDelegateTestObject.h
@@ -10,49 +10,49 @@
 #include "DynamicDelegateTestObject.generated.h"
 
 DECLARE_DYNAMIC_DELEGATE            (FTestDynamicDelegate0);
-DECLARE_DYNAMIC_DELEGATE_OneParam   (FTestDynamicDelegate1, int, Int);
-DECLARE_DYNAMIC_DELEGATE_TwoParams  (FTestDynamicDelegate2, int, Int, FString, Str);
-DECLARE_DYNAMIC_DELEGATE_ThreeParams(FTestDynamicDelegate3, int, Int, FString, Str, const UDynamicDelegateTestObject*, Ptr);
-DECLARE_DYNAMIC_DELEGATE_FourParams (FTestDynamicDelegate4, int, Int, FString, Str, const UDynamicDelegateTestObject*, Ptr, FName, Name);
-DECLARE_DYNAMIC_DELEGATE_FiveParams (FTestDynamicDelegate5, int, Int, FString, Str, const UDynamicDelegateTestObject*, Ptr, FName, Name, bool, bBool);
-DECLARE_DYNAMIC_DELEGATE_SixParams  (FTestDynamicDelegate6, int, Int, FString, Str, const UDynamicDelegateTestObject*, Ptr, FName, Name, bool, bBool, TWeakObjectPtr<UDynamicDelegateTestObject>, WeakPtr);
+DECLARE_DYNAMIC_DELEGATE_OneParam   (FTestDynamicDelegate1, int16, Int);
+DECLARE_DYNAMIC_DELEGATE_TwoParams  (FTestDynamicDelegate2, int16, Int, FString, Str);
+DECLARE_DYNAMIC_DELEGATE_ThreeParams(FTestDynamicDelegate3, int16, Int, FString, Str, const UDynamicDelegateTestObject*, Ptr);
+DECLARE_DYNAMIC_DELEGATE_FourParams (FTestDynamicDelegate4, int16, Int, FString, Str, const UDynamicDelegateTestObject*, Ptr, FName, Name);
+DECLARE_DYNAMIC_DELEGATE_FiveParams (FTestDynamicDelegate5, int16, Int, FString, Str, const UDynamicDelegateTestObject*, Ptr, FName, Name, bool, bBool);
+DECLARE_DYNAMIC_DELEGATE_SixParams  (FTestDynamicDelegate6, int16, Int, FString, Str, const UDynamicDelegateTestObject*, Ptr, FName, Name, bool, bBool, TWeakObjectPtr<UDynamicDelegateTestObject>, WeakPtr);
 
 DECLARE_DYNAMIC_MULTICAST_DELEGATE            (FTestDynamicMulticastDelegate0);
-DECLARE_DYNAMIC_MULTICAST_DELEGATE_OneParam   (FTestDynamicMulticastDelegate1, int, Int);
-DECLARE_DYNAMIC_MULTICAST_DELEGATE_TwoParams  (FTestDynamicMulticastDelegate2, int, Int, FString, Str);
-DECLARE_DYNAMIC_MULTICAST_DELEGATE_ThreeParams(FTestDynamicMulticastDelegate3, int, Int, FString, Str, const UDynamicDelegateTestObject*, Ptr);
-DECLARE_DYNAMIC_MULTICAST_DELEGATE_FourParams (FTestDynamicMulticastDelegate4, int, Int, FString, Str, const UDynamicDelegateTestObject*, Ptr, FName, Name);
-DECLARE_DYNAMIC_MULTICAST_DELEGATE_FiveParams (FTestDynamicMulticastDelegate5, int, Int, FString, Str, const UDynamicDelegateTestObject*, Ptr, FName, Name, bool, bBool);
-DECLARE_DYNAMIC_MULTICAST_DELEGATE_SixParams  (FTestDynamicMulticastDelegate6, int, Int, FString, Str, const UDynamicDelegateTestObject*, Ptr, FName, Name, bool, bBool, TWeakObjectPtr<UDynamicDelegateTestObject>, WeakPtr);
+DECLARE_DYNAMIC_MULTICAST_DELEGATE_OneParam   (FTestDynamicMulticastDelegate1, int16, Int);
+DECLARE_DYNAMIC_MULTICAST_DELEGATE_TwoParams  (FTestDynamicMulticastDelegate2, int16, Int, FString, Str);
+DECLARE_DYNAMIC_MULTICAST_DELEGATE_ThreeParams(FTestDynamicMulticastDelegate3, int16, Int, FString, Str, const UDynamicDelegateTestObject*, Ptr);
+DECLARE_DYNAMIC_MULTICAST_DELEGATE_FourParams (FTestDynamicMulticastDelegate4, int16, Int, FString, Str, const UDynamicDelegateTestObject*, Ptr, FName, Name);
+DECLARE_DYNAMIC_MULTICAST_DELEGATE_FiveParams (FTestDynamicMulticastDelegate5, int16, Int, FString, Str, const UDynamicDelegateTestObject*, Ptr, FName, Name, bool, bBool);
+DECLARE_DYNAMIC_MULTICAST_DELEGATE_SixParams  (FTestDynamicMulticastDelegate6, int16, Int, FString, Str, const UDynamicDelegateTestObject*, Ptr, FName, Name, bool, bBool, TWeakObjectPtr<UDynamicDelegateTestObject>, WeakPtr);
 
 DECLARE_DYNAMIC_DELEGATE_RetVal            (int, FTestDynamicDelegateRetInt0);
-DECLARE_DYNAMIC_DELEGATE_RetVal_OneParam   (int, FTestDynamicDelegateRetInt1, int, Int);
-DECLARE_DYNAMIC_DELEGATE_RetVal_TwoParams  (int, FTestDynamicDelegateRetInt2, int, Int, FString, Str);
-DECLARE_DYNAMIC_DELEGATE_RetVal_ThreeParams(int, FTestDynamicDelegateRetInt3, int, Int, FString, Str, const UDynamicDelegateTestObject*, Ptr);
-DECLARE_DYNAMIC_DELEGATE_RetVal_FourParams (int, FTestDynamicDelegateRetInt4, int, Int, FString, Str, const UDynamicDelegateTestObject*, Ptr, FName, Name);
-DECLARE_DYNAMIC_DELEGATE_RetVal_FiveParams (int, FTestDynamicDelegateRetInt5, int, Int, FString, Str, const UDynamicDelegateTestObject*, Ptr, FName, Name, bool, bBool);
-DECLARE_DYNAMIC_DELEGATE_RetVal_SixParams  (int, FTestDynamicDelegateRetInt6, int, Int, FString, Str, const UDynamicDelegateTestObject*, Ptr, FName, Name, bool, bBool, TWeakObjectPtr<UDynamicDelegateTestObject>, WeakPtr);
+DECLARE_DYNAMIC_DELEGATE_RetVal_OneParam   (int, FTestDynamicDelegateRetInt1, int16, Int);
+DECLARE_DYNAMIC_DELEGATE_RetVal_TwoParams  (int, FTestDynamicDelegateRetInt2, int16, Int, FString, Str);
+DECLARE_DYNAMIC_DELEGATE_RetVal_ThreeParams(int, FTestDynamicDelegateRetInt3, int16, Int, FString, Str, const UDynamicDelegateTestObject*, Ptr);
+DECLARE_DYNAMIC_DELEGATE_RetVal_FourParams (int, FTestDynamicDelegateRetInt4, int16, Int, FString, Str, const UDynamicDelegateTestObject*, Ptr, FName, Name);
+DECLARE_DYNAMIC_DELEGATE_RetVal_FiveParams (int, FTestDynamicDelegateRetInt5, int16, Int, FString, Str, const UDynamicDelegateTestObject*, Ptr, FName, Name, bool, bBool);
+DECLARE_DYNAMIC_DELEGATE_RetVal_SixParams  (int, FTestDynamicDelegateRetInt6, int16, Int, FString, Str, const UDynamicDelegateTestObject*, Ptr, FName, Name, bool, bBool, TWeakObjectPtr<UDynamicDelegateTestObject>, WeakPtr);
 
 DECLARE_DYNAMIC_DELEGATE_RetVal            (FString, FTestDynamicDelegateRetStr0);
-DECLARE_DYNAMIC_DELEGATE_RetVal_OneParam   (FString, FTestDynamicDelegateRetStr1, int, Int);
-DECLARE_DYNAMIC_DELEGATE_RetVal_TwoParams  (FString, FTestDynamicDelegateRetStr2, int, Int, FString, Str);
-DECLARE_DYNAMIC_DELEGATE_RetVal_ThreeParams(FString, FTestDynamicDelegateRetStr3, int, Int, FString, Str, const UDynamicDelegateTestObject*, Ptr);
-DECLARE_DYNAMIC_DELEGATE_RetVal_FourParams (FString, FTestDynamicDelegateRetStr4, int, Int, FString, Str, const UDynamicDelegateTestObject*, Ptr, FName, Name);
-DECLARE_DYNAMIC_DELEGATE_RetVal_FiveParams (FString, FTestDynamicDelegateRetStr5, int, Int, FString, Str, const UDynamicDelegateTestObject*, Ptr, FName, Name, bool, bBool);
-DECLARE_DYNAMIC_DELEGATE_RetVal_SixParams  (FString, FTestDynamicDelegateRetStr6, int, Int, FString, Str, const UDynamicDelegateTestObject*, Ptr, FName, Name, bool, bBool, TWeakObjectPtr<UDynamicDelegateTestObject>, WeakPtr);
-
-DECLARE_DYNAMIC_DELEGATE_OneParam   (FTestDynamicDelegateOut1, int&, Int);
-DECLARE_DYNAMIC_DELEGATE_TwoParams  (FTestDynamicDelegateOut2, int&, Int, FString&, Str);
-DECLARE_DYNAMIC_DELEGATE_ThreeParams(FTestDynamicDelegateOut3, int&, Int, FString&, Str, FName&, Name);
-DECLARE_DYNAMIC_DELEGATE_FourParams (FTestDynamicDelegateOut4, int&, Int, FString&, Str, FName&, Name, bool&, bBool);
-DECLARE_DYNAMIC_DELEGATE_FiveParams (FTestDynamicDelegateOut5, int&, Int, FString&, Str, FName&, Name, bool&, bBool, TWeakObjectPtr<UDynamicDelegateTestObject>&, WeakPtr);
+DECLARE_DYNAMIC_DELEGATE_RetVal_OneParam   (FString, FTestDynamicDelegateRetStr1, int16, Int);
+DECLARE_DYNAMIC_DELEGATE_RetVal_TwoParams  (FString, FTestDynamicDelegateRetStr2, int16, Int, FString, Str);
+DECLARE_DYNAMIC_DELEGATE_RetVal_ThreeParams(FString, FTestDynamicDelegateRetStr3, int16, Int, FString, Str, const UDynamicDelegateTestObject*, Ptr);
+DECLARE_DYNAMIC_DELEGATE_RetVal_FourParams (FString, FTestDynamicDelegateRetStr4, int16, Int, FString, Str, const UDynamicDelegateTestObject*, Ptr, FName, Name);
+DECLARE_DYNAMIC_DELEGATE_RetVal_FiveParams (FString, FTestDynamicDelegateRetStr5, int16, Int, FString, Str, const UDynamicDelegateTestObject*, Ptr, FName, Name, bool, bBool);
+DECLARE_DYNAMIC_DELEGATE_RetVal_SixParams  (FString, FTestDynamicDelegateRetStr6, int16, Int, FString, Str, const UDynamicDelegateTestObject*, Ptr, FName, Name, bool, bBool, TWeakObjectPtr<UDynamicDelegateTestObject>, WeakPtr);
+
+DECLARE_DYNAMIC_DELEGATE_OneParam   (FTestDynamicDelegateOut1, int16&, Int);
+DECLARE_DYNAMIC_DELEGATE_TwoParams  (FTestDynamicDelegateOut2, int16&, Int, FString&, Str);
+DECLARE_DYNAMIC_DELEGATE_ThreeParams(FTestDynamicDelegateOut3, int16&, Int, FString&, Str, FName&, Name);
+DECLARE_DYNAMIC_DELEGATE_FourParams (FTestDynamicDelegateOut4, int16&, Int, FString&, Str, FName&, Name, bool&, bBool);
+DECLARE_DYNAMIC_DELEGATE_FiveParams (FTestDynamicDelegateOut5, int16&, Int, FString&, Str, FName&, Name, bool&, bBool, TWeakObjectPtr<UDynamicDelegateTestObject>&, WeakPtr);
 
 // Pointer out params aren't currently supported
-DECLARE_DYNAMIC_MULTICAST_DELEGATE_OneParam   (FTestDynamicMulticastDelegateOut1, int&, Int);
-DECLARE_DYNAMIC_MULTICAST_DELEGATE_TwoParams  (FTestDynamicMulticastDelegateOut2, int&, Int, FString&, Str);
-DECLARE_DYNAMIC_MULTICAST_DELEGATE_ThreeParams(FTestDynamicMulticastDelegateOut3, int&, Int, FString&, Str, FName&, Name);
-DECLARE_DYNAMIC_MULTICAST_DELEGATE_FourParams (FTestDynamicMulticastDelegateOut4, int&, Int, FString&, Str, FName&, Name, bool&, bBool);
-DECLARE_DYNAMIC_MULTICAST_DELEGATE_FiveParams (FTestDynamicMulticastDelegateOut5, int&, Int, FString&, Str, FName&, Name, bool&, bBool, TWeakObjectPtr<UDynamicDelegateTestObject>&, WeakPtr);
+DECLARE_DYNAMIC_MULTICAST_DELEGATE_OneParam   (FTestDynamicMulticastDelegateOut1, int16&, Int);
+DECLARE_DYNAMIC_MULTICAST_DELEGATE_TwoParams  (FTestDynamicMulticastDelegateOut2, int16&, Int, FString&, Str);
+DECLARE_DYNAMIC_MULTICAST_DELEGATE_ThreeParams(FTestDynamicMulticastDelegateOut3, int16&, Int, FString&, Str, FName&, Name);
+DECLARE_DYNAMIC_MULTICAST_DELEGATE_FourParams (FTestDynamicMulticastDelegateOut4, int16&, Int, FString&, Str, FName&, Name, bool&, bBool);
+DECLARE_DYNAMIC_MULTICAST_DELEGATE_FiveParams (FTestDynamicMulticastDelegateOut5, int16&, Int, FString&, Str, FName&, Name, bool&, bBool, TWeakObjectPtr<UDynamicDelegateTestObject>&, WeakPtr);
 
 UENUM()
 namespace ETestUnscopedEnum
@@ -75,66 +75,66 @@ public:
 	UDynamicDelegateTestObject();
 
 	UFUNCTION() void Func0();
-	UFUNCTION() void Func1(int Int);
-	UFUNCTION() void Func2(int Int, FString Str);
-	UFUNCTION() void Func3(int Int, FString Str, const UDynamicDelegateTestObject* Ptr);
-	UFUNCTION() void Func4(int Int, FString Str, const UDynamicDelegateTestObject* Ptr, FName Name);
-	UFUNCTION() void Func5(int Int, FString Str, const UDynamicDelegateTestObject* Ptr, FName Name, bool bBool);
-	UFUNCTION() void Func6(int Int, FString Str, const UDynamicDelegateTestObject* Ptr, FName Name, bool bBool, TWeakObjectPtr<UDynamicDelegateTestObject> WeakPtr);
+	UFUNCTION() void Func1(int16 Int);
+	UFUNCTION() void Func2(int16 Int, FString Str);
+	UFUNCTION() void Func3(int16 Int, FString Str, const UDynamicDelegateTestObject* Ptr);
+	UFUNCTION() void Func4(int16 Int, FString Str, const UDynamicDelegateTestObject* Ptr, FName Name);
+	UFUNCTION() void Func5(int16 Int, FString Str, const UDynamicDelegateTestObject* Ptr, FName Name, bool bBool);
+	UFUNCTION() void Func6(int16 Int, FString Str, const UDynamicDelegateTestObject* Ptr, FName Name, bool bBool, TWeakObjectPtr<UDynamicDelegateTestObject> WeakPtr);
 
 	UFUNCTION() void ConstFunc0() const;
-	UFUNCTION() void ConstFunc1(int Int) const;
-	UFUNCTION() void ConstFunc2(int Int, FString Str) const;
-	UFUNCTION() void ConstFunc3(int Int, FString Str, const UDynamicDelegateTestObject* Ptr) const;
-	UFUNCTION() void ConstFunc4(int Int, FString Str, const UDynamicDelegateTestObject* Ptr, FName Name) const;
-	UFUNCTION() void ConstFunc5(int Int, FString Str, const UDynamicDelegateTestObject* Ptr, FName Name, bool bBool) const;
-	UFUNCTION() void ConstFunc6(int Int, FString Str, const UDynamicDelegateTestObject* Ptr, FName Name, bool bBool, TWeakObjectPtr<UDynamicDelegateTestObject> WeakPtr) const;
+	UFUNCTION() void ConstFunc1(int16 Int) const;
+	UFUNCTION() void ConstFunc2(int16 Int, FString Str) const;
+	UFUNCTION() void ConstFunc3(int16 Int, FString Str, const UDynamicDelegateTestObject* Ptr) const;
+	UFUNCTION() void ConstFunc4(int16 Int, FString Str, const UDynamicDelegateTestObject* Ptr, FName Name) const;
+	UFUNCTION() void ConstFunc5(int16 Int, FString Str, const UDynamicDelegateTestObject* Ptr, FName Name, bool bBool) const;
+	UFUNCTION() void ConstFunc6(int16 Int, FString Str, const UDynamicDelegateTestObject* Ptr, FName Name, bool bBool, TWeakObjectPtr<UDynamicDelegateTestObject> WeakPtr) const;
 
 	UFUNCTION() int FuncRetInt0();
-	UFUNCTION() int FuncRetInt1(int Int);
-	UFUNCTION() int FuncRetInt2(int Int, FString Str);
-	UFUNCTION() int FuncRetInt3(int Int, FString Str, const UDynamicDelegateTestObject* Ptr);
-	UFUNCTION() int FuncRetInt4(int Int, FString Str, const UDynamicDelegateTestObject* Ptr, FName Name);
-	UFUNCTION() int FuncRetInt5(int Int, FString Str, const UDynamicDelegateTestObject* Ptr, FName Name, bool bBool);
-	UFUNCTION() int FuncRetInt6(int Int, FString Str, const UDynamicDelegateTestObject* Ptr, FName Name, bool bBool, TWeakObjectPtr<UDynamicDelegateTestObject> WeakPtr);
+	UFUNCTION() int FuncRetInt1(int16 Int);
+	UFUNCTION() int FuncRetInt2(int16 Int, FString Str);
+	UFUNCTION() int FuncRetInt3(int16 Int, FString Str, const UDynamicDelegateTestObject* Ptr);
+	UFUNCTION() int FuncRetInt4(int16 Int, FString Str, const UDynamicDelegateTestObject* Ptr, FName Name);
+	UFUNCTION() int FuncRetInt5(int16 Int, FString Str, const UDynamicDelegateTestObject* Ptr, FName Name, bool bBool);
+	UFUNCTION() int FuncRetInt6(int16 Int, FString Str, const UDynamicDelegateTestObject* Ptr, FName Name, bool bBool, TWeakObjectPtr<UDynamicDelegateTestObject> WeakPtr);
 
 	UFUNCTION() FString FuncRetStr0();
-	UFUNCTION() FString FuncRetStr1(int Int);
-	UFUNCTION() FString FuncRetStr2(int Int, FString Str);
-	UFUNCTION() FString FuncRetStr3(int Int, FString Str, const UDynamicDelegateTestObject* Ptr);
-	UFUNCTION() FString FuncRetStr4(int Int, FString Str, const UDynamicDelegateTestObject* Ptr, FName Name);
-	UFUNCTION() FString FuncRetStr5(int Int, FString Str, const UDynamicDelegateTestObject* Ptr, FName Name, bool bBool);
-	UFUNCTION() FString FuncRetStr6(int Int, FString Str, const UDynamicDelegateTestObject* Ptr, FName Name, bool bBool, TWeakObjectPtr<UDynamicDelegateTestObject> WeakPtr);
+	UFUNCTION() FString FuncRetStr1(int16 Int);
+	UFUNCTION() FString FuncRetStr2(int16 Int, FString Str);
+	UFUNCTION() FString FuncRetStr3(int16 Int, FString Str, const UDynamicDelegateTestObject* Ptr);
+	UFUNCTION() FString FuncRetStr4(int16 Int, FString Str, const UDynamicDelegateTestObject* Ptr, FName Name);
+	UFUNCTION() FString FuncRetStr5(int16 Int, FString Str, const UDynamicDelegateTestObject* Ptr, FName Name, bool bBool);
+	UFUNCTION() FString FuncRetStr6(int16 Int, FString Str, const UDynamicDelegateTestObject* Ptr, FName Name, bool bBool, TWeakObjectPtr<UDynamicDelegateTestObject> WeakPtr);
 
 	UFUNCTION() int ConstFuncRetInt0() const;
-	UFUNCTION() int ConstFuncRetInt1(int Int) const;
-	UFUNCTION() int ConstFuncRetInt2(int Int, FString Str) const;
-	UFUNCTION() int ConstFuncRetInt3(int Int, FString Str, const UDynamicDelegateTestObject* Ptr) const;
-	UFUNCTION() int ConstFuncRetInt4(int Int, FString Str, const UDynamicDelegateTestObject* Ptr, FName Name) const;
-	UFUNCTION() int ConstFuncRetInt5(int Int, FString Str, const UDynamicDelegateTestObject* Ptr, FName Name, bool bBool) const;
-	UFUNCTION() int ConstFuncRetInt6(int Int, FString Str, const UDynamicDelegateTestObject* Ptr, FName Name, bool bBool, TWeakObjectPtr<UDynamicDelegateTestObject> WeakPtr) const;
+	UFUNCTION() int ConstFuncRetInt1(int16 Int) const;
+	UFUNCTION() int ConstFuncRetInt2(int16 Int, FString Str) const;
+	UFUNCTION() int ConstFuncRetInt3(int16 Int, FString Str, const UDynamicDelegateTestObject* Ptr) const;
+	UFUNCTION() int ConstFuncRetInt4(int16 Int, FString Str, const UDynamicDelegateTestObject* Ptr, FName Name) const;
+	UFUNCTION() int ConstFuncRetInt5(int16 Int, FString Str, const UDynamicDelegateTestObject* Ptr, FName Name, bool bBool) const;
+	UFUNCTION() int ConstFuncRetInt6(int16 Int, FString Str, const UDynamicDelegateTestObject* Ptr, FName Name, bool bBool, TWeakObjectPtr<UDynamicDelegateTestObject> WeakPtr) const;
 
 	UFUNCTION() FString ConstFuncRetStr0() const;
-	UFUNCTION() FString ConstFuncRetStr1(int Int) const;
-	UFUNCTION() FString ConstFuncRetStr2(int Int, FString Str) const;
-	UFUNCTION() FString ConstFuncRetStr3(int Int, FString Str, const UDynamicDelegateTestObject* Ptr) const;
-	UFUNCTION() FString ConstFuncRetStr4(int Int, FString Str, const UDynamicDelegateTestObject* Ptr, FName Name) const;
-	UFUNCTION() FString ConstFuncRetStr5(int Int, FString Str, const UDynamicDelegateTestObject* Ptr, FName Name, bool bBool) const;
-	UFUNCTION() FString ConstFuncRetStr6(int Int, FString Str, const UDynamicDelegateTestObject* Ptr, FName Name, bool bBool, TWeakObjectPtr<UDynamicDelegateTestObject> WeakPtr) const;
+	UFUNCTION() FString ConstFuncRetStr1(int16 Int) const;
+	UFUNCTION() FString ConstFuncRetStr2(int16 Int, FString Str) const;
+	UFUNCTION() FString ConstFuncRetStr3(int16 Int, FString Str, const UDynamicDelegateTestObject* Ptr) const;
+	UFUNCTION() FString ConstFuncRetStr4(int16 Int, FString Str, const UDynamicDelegateTestObject* Ptr, FName Name) const;
+	UFUNCTION() FString ConstFuncRetStr5(int16 Int, FString Str, const UDynamicDelegateTestObject* Ptr, FName Name, bool bBool) const;
+	UFUNCTION() FString ConstFuncRetStr6(int16 Int, FString Str, const UDynamicDelegateTestObject* Ptr, FName Name, bool bBool, TWeakObjectPtr<UDynamicDelegateTestObject> WeakPtr) const;
 
 	// Pointer out params aren't currently supported
-	UFUNCTION() void FuncOut1(int& Int);
-	UFUNCTION() void FuncOut2(int& Int, FString& Str);
-	UFUNCTION() void FuncOut3(int& Int, FString& Str, FName& Name);
-	UFUNCTION() void FuncOut4(int& Int, FString& Str, FName& Name, bool& bBool);
-	UFUNCTION() void FuncOut5(int& Int, FString& Str, FName& Name, bool& bBool, TWeakObjectPtr<UDynamicDelegateTestObject>& WeakPtr);
+	UFUNCTION() void FuncOut1(int16& Int);
+	UFUNCTION() void FuncOut2(int16& Int, FString& Str);
+	UFUNCTION() void FuncOut3(int16& Int, FString& Str, FName& Name);
+	UFUNCTION() void FuncOut4(int16& Int, FString& Str, FName& Name, bool& bBool);
+	UFUNCTION() void FuncOut5(int16& Int, FString& Str, FName& Name, bool& bBool, TWeakObjectPtr<UDynamicDelegateTestObject>& WeakPtr);
 
 	// Pointer out params aren't currently supported
-	UFUNCTION() void ConstFuncOut1(int& Int) const;
-	UFUNCTION() void ConstFuncOut2(int& Int, FString& Str) const;
-	UFUNCTION() void ConstFuncOut3(int& Int, FString& Str, FName& Name) const;
-	UFUNCTION() void ConstFuncOut4(int& Int, FString& Str, FName& Name, bool& bBool) const;
-	UFUNCTION() void ConstFuncOut5(int& Int, FString& Str, FName& Name, bool& bBool, TWeakObjectPtr<UDynamicDelegateTestObject>& WeakPtr) const;
+	UFUNCTION() void ConstFuncOut1(int16& Int) const;
+	UFUNCTION() void ConstFuncOut2(int16& Int, FString& Str) const;
+	UFUNCTION() void ConstFuncOut3(int16& Int, FString& Str, FName& Name) const;
+	UFUNCTION() void ConstFuncOut4(int16& Int, FString& Str, FName& Name, bool& bBool) const;
+	UFUNCTION() void ConstFuncOut5(int16& Int, FString& Str, FName& Name, bool& bBool, TWeakObjectPtr<UDynamicDelegateTestObject>& WeakPtr) const;
 
 	UFUNCTION() void TestUnscopedEnum(int16 Int1, ETestUnscopedEnum::Type E, int16 Int2);
 
