[ZenLoader]
- Fix blueprint class that could end up being postloaded before they were compiled when sync loads happened during serialize

#Synced-CL 52534307

diff --git a/Engine/Source/Runtime/CoreUObject/Private/Blueprint/BlueprintSupport.cpp b/Engine/Source/Runtime/CoreUObject/Private/Blueprint/BlueprintSupport.cpp
--- a/Engine/Source/Runtime/CoreUObject/Private/Blueprint/BlueprintSupport.cpp
+++ b/Engine/Source/Runtime/CoreUObject/Private/Blueprint/BlueprintSupport.cpp
@@ -469,8 +469,9 @@ void FBlueprintSupport::ValidateNoExternalRefsToSkeletons()
  ******************************************************************************/
 
 #if WITH_EDITOR
-UClass* FScopedClassDependencyGather::BatchAuthorityClass = nullptr;
-TArray<UClass*> FScopedClassDependencyGather::BatchClassDependencies;
+thread_local UClass* FScopedClassDependencyGather::BatchAuthorityClass = nullptr;
+thread_local TArray<UClass*> FScopedClassDependencyGather::BatchClassDependencies;
+thread_local FScopedClassDependencyGather* FScopedClassDependencyGather::AuthoritativeScope = nullptr;
 
 FScopedClassDependencyGather::FScopedClassDependencyGather(UClass* ClassToGather, FUObjectSerializeContext* InLoadContext)
 	: bAuthoritativeClass(false)
@@ -485,6 +486,8 @@ FScopedClassDependencyGather::FScopedClassDependencyGather(UClass* ClassToGather
 			BatchAuthorityClass = ClassToGather;
 			BatchClassDependencies.Empty();
 			bAuthoritativeClass = true;
+
+			AuthoritativeScope = this;
 		}
 		else
 		{
@@ -504,6 +507,18 @@ FScopedClassDependencyGather::~FScopedClassDependencyGather()
 
 		BatchAuthorityClass = nullptr;
 	}
+
+	if (AuthoritativeScope == this)
+	{
+		AuthoritativeScope = nullptr;
+
+		OnScopeExitDelegate.Broadcast();
+	}
+}
+
+FScopedClassDependencyGather* FScopedClassDependencyGather::GetAuthoritativeScope()
+{
+	return AuthoritativeScope;
 }
 
 TArray<UClass*> const& FScopedClassDependencyGather::GetCachedDependencies()
diff --git a/Engine/Source/Runtime/CoreUObject/Private/Serialization/AsyncLoading2.cpp b/Engine/Source/Runtime/CoreUObject/Private/Serialization/AsyncLoading2.cpp
--- a/Engine/Source/Runtime/CoreUObject/Private/Serialization/AsyncLoading2.cpp
+++ b/Engine/Source/Runtime/CoreUObject/Private/Serialization/AsyncLoading2.cpp
@@ -24,6 +24,7 @@
 #include "Misc/CoreStats.h"
 #include "HAL/IConsoleManager.h"
 #include "Misc/PackageAccessTrackingOps.h"
+#include "Blueprint/BlueprintSupport.h"
 #include "Misc/CoreDelegates.h"
 #include "Misc/CommandLine.h"
 #include "Misc/App.h"
@@ -9871,6 +9872,27 @@ void FAsyncPackage2::ConditionalReleasePartialRequests(FAsyncLoadingThreadState2
 					);
 					AsyncLoadingThread.RemovePendingRequests(ThreadState, { RequestId });
 
+#if WITH_EDITOR
+					// If the package we release hasn't been postloaded yet, we'll prevent it from being postloaded until
+					// its blueprint classes have been compiled if we find any.
+					FScopedClassDependencyGather* ClassDependencyGather = FScopedClassDependencyGather::GetAuthoritativeScope();
+					if (ClassDependencyGather != nullptr && AsyncPackageLoadingState < EAsyncPackageLoadingState2::DeferredPostLoad)
+					{
+						UE_LOG(LogStreaming, Display, TEXT("Package %s postload will not proceed until its blueprint has been compiled"),
+							*Desc.UPackageName.ToString()
+						);
+
+						GetPackageNode(EEventLoadNode2::Package_DeferredPostLoad).AddBarrier();
+
+						ClassDependencyGather->OnScopeExitDelegate.AddLambda(
+							[this, &ThreadState]()
+							{
+								GetPackageNode(EEventLoadNode2::Package_DeferredPostLoad).ReleaseBarrier(&ThreadState);
+							}
+						);
+					}
+#endif
+
 					switch (AsyncPackageLoadingState)
 					{
 					case EAsyncPackageLoadingState2::ExportsDone:
diff --git a/Engine/Source/Runtime/CoreUObject/Public/Blueprint/BlueprintSupport.h b/Engine/Source/Runtime/CoreUObject/Public/Blueprint/BlueprintSupport.h
--- a/Engine/Source/Runtime/CoreUObject/Public/Blueprint/BlueprintSupport.h
+++ b/Engine/Source/Runtime/CoreUObject/Public/Blueprint/BlueprintSupport.h
@@ -180,15 +180,29 @@ public:
 	 */
 	static COREUOBJECT_API TArray<UClass*> const& GetCachedDependencies();
 
+	/**
+	 * Returns the current authoratitative scope on the stack if one exist.
+	 */
+	static COREUOBJECT_API FScopedClassDependencyGather* GetAuthoritativeScope();
+
+	DECLARE_TS_MULTICAST_DELEGATE(FOnScopeExit);
+
+	/**
+	 * Delegate called when the scope exits;
+	 */
+	FOnScopeExit OnScopeExitDelegate;
 private:
 	/** Whether or not this dependency gather is the dependency authoritative class, and thus should process all dependencies in the destructor */
 	bool bAuthoritativeClass;
 
 	/** The current class that is gathering potential dependencies in this scope */
-	static UClass* BatchAuthorityClass;
+	static thread_local UClass* BatchAuthorityClass;
 
 	/** List of dependencies (i.e. UClasses that have been newly instantiated) in the scope of this dependency gather */
-	static TArray<UClass*> BatchClassDependencies;
+	static thread_local TArray<UClass*> BatchClassDependencies;
+
+	/** Contains the authoritative scope on the current thread */
+	static thread_local FScopedClassDependencyGather* AuthoritativeScope;
 
 	/** Current load context */
 	FUObjectSerializeContext* LoadContext;
