[Linker]
- Optimize by only trying to flush when there is an actual async request pending on the linker

#Synced-CL 52534307

diff --git a/Engine/Source/Runtime/CoreUObject/Private/UObject/Linker.cpp b/Engine/Source/Runtime/CoreUObject/Private/UObject/Linker.cpp
--- a/Engine/Source/Runtime/CoreUObject/Private/UObject/Linker.cpp
+++ b/Engine/Source/Runtime/CoreUObject/Private/UObject/Linker.cpp
@@ -813,24 +813,31 @@ void ConditionalFlushAsyncLoadingForLinkers(TConstArrayView<FLinkerLoad*> InLink
 	// Flushing the new requests will also flush any other pending async requests for the same linkers.
 	if (InLinkers.Num() > 0 && IsAsyncLoading())
 	{
-		UE_LOG(LogLinker, Log, TEXT("Conditionally flushing loading for linker(s) (%s)"), *GetPathNameSafe(InLinkers[0]->LinkerRoot));
-
 		TArray<int32, TInlineAllocator<4>> RequestIds;
 		for (FLinkerLoad* Linker : InLinkers)
 		{
-			int32 Request = INDEX_NONE;
-			if (!Linker->GetInstancingContext().IsInstanced())
-			{
-				Request = LoadPackageAsync(Linker->GetPackagePath(), FLoadPackageAsyncOptionalParams{ .PackagePriority = MAX_int32 });
-			}
-			else
+			if (Linker->AsyncRoot != nullptr)
 			{
-				FName RemappedPackageName = Linker->GetInstancingContext().RemapPackage(Linker->GetPackagePath().GetPackageFName());
-				Request = LoadPackageAsync(Linker->GetPackagePath(), FLoadPackageAsyncOptionalParams{ .CustomPackageName = RemappedPackageName, .PackagePriority = MAX_int32, .InstancingContext = &Linker->GetInstancingContext() });
+				int32 Request = INDEX_NONE;
+				if (!Linker->GetInstancingContext().IsInstanced())
+				{
+					Request = LoadPackageAsync(Linker->GetPackagePath(), FLoadPackageAsyncOptionalParams{ .PackagePriority = MAX_int32 });
+				}
+				else
+				{
+					FName RemappedPackageName = Linker->GetInstancingContext().RemapPackage(Linker->GetPackagePath().GetPackageFName());
+					Request = LoadPackageAsync(Linker->GetPackagePath(), FLoadPackageAsyncOptionalParams{ .CustomPackageName = RemappedPackageName, .PackagePriority = MAX_int32, .InstancingContext = &Linker->GetInstancingContext() });
+				}
+				RequestIds.Add(Request);
 			}
-			RequestIds.Add(Request);
 		}
-		FlushAsyncLoading(RequestIds);
+
+		if (RequestIds.Num())
+		{
+			UE_LOG(LogLinker, Log, TEXT("Conditionally flushing loading for linker(s) (%s)"), *GetPathNameSafe(InLinkers[0]->LinkerRoot));
+
+			FlushAsyncLoading(RequestIds);
+		}
 	}
 }
 
