Supported unscoped enums in dynamic delegate payloads.
Added tests for this.

diff --git a/Engine/Source/Runtime/Core/Public/UObject/ScriptDelegates.h b/Engine/Source/Runtime/Core/Public/UObject/ScriptDelegates.h
--- a/Engine/Source/Runtime/Core/Public/UObject/ScriptDelegates.h
+++ b/Engine/Source/Runtime/Core/Public/UObject/ScriptDelegates.h
@@ -43,6 +43,7 @@
 
 #if UE_USE_DYNAMIC_DELEGATE_PAYLOADS
 	#include "Concepts/Arithmetic.h"
+	#include "Concepts/UnscopedEnum.h"
 	#include "Concepts/UObject.h"
 	#include "Templates/Tuple.h"
 	#include "Traits/FunctionArity.h"
@@ -120,14 +121,19 @@ namespace UE::Core::Private
 	// - Numeric types.
 	// - void (only for return types, not parameters, but C++ will stop you having a void parameter before getting here)
 	// - Strings and FNames.
-	template <typename WeakPtrType, UE::CUObject T>    struct TDynamicDelegateParamPayloadCaptureType<WeakPtrType, T*>          { using Type = TImplicitWeakObjectPtr<T, WeakPtrType>; };
-	template <typename WeakPtrType, UE::CArithmetic T> struct TDynamicDelegateParamPayloadCaptureType<WeakPtrType, T>           { using Type = T; };
-	template <typename WeakPtrType, CWeakPointerish T> struct TDynamicDelegateParamPayloadCaptureType<WeakPtrType, T>           { using Type = T; };
-	template <typename WeakPtrType>                    struct TDynamicDelegateParamPayloadCaptureType<WeakPtrType, void>        { using Type = void; };
-	template <typename WeakPtrType>                    struct TDynamicDelegateParamPayloadCaptureType<WeakPtrType, FAnsiString> { using Type = FAnsiString; };
-	template <typename WeakPtrType>                    struct TDynamicDelegateParamPayloadCaptureType<WeakPtrType, FString>     { using Type = FString; };
-	template <typename WeakPtrType>                    struct TDynamicDelegateParamPayloadCaptureType<WeakPtrType, FUtf8String> { using Type = FUtf8String; };
-	template <typename WeakPtrType>                    struct TDynamicDelegateParamPayloadCaptureType<WeakPtrType, FName>       { using Type = FName; };
+	template <typename WeakPtrType, UE::CUObject T>      struct TDynamicDelegateParamPayloadCaptureType<WeakPtrType, T*>          { using Type = TImplicitWeakObjectPtr<T, WeakPtrType>; };
+	template <typename WeakPtrType, UE::CArithmetic T>   struct TDynamicDelegateParamPayloadCaptureType<WeakPtrType, T>           { using Type = T; };
+	template <typename WeakPtrType, CWeakPointerish T>   struct TDynamicDelegateParamPayloadCaptureType<WeakPtrType, T>           { using Type = T; };
+	template <typename WeakPtrType>                      struct TDynamicDelegateParamPayloadCaptureType<WeakPtrType, void>        { using Type = void; };
+	template <typename WeakPtrType>                      struct TDynamicDelegateParamPayloadCaptureType<WeakPtrType, FAnsiString> { using Type = FAnsiString; };
+	template <typename WeakPtrType>                      struct TDynamicDelegateParamPayloadCaptureType<WeakPtrType, FString>     { using Type = FString; };
+	template <typename WeakPtrType>                      struct TDynamicDelegateParamPayloadCaptureType<WeakPtrType, FUtf8String> { using Type = FUtf8String; };
+	template <typename WeakPtrType>                      struct TDynamicDelegateParamPayloadCaptureType<WeakPtrType, FName>       { using Type = FName; };
+#if UE_DYNAMIC_DELEGATE_UNSCOPED_ENUM_HACK
+	template <typename WeakPtrType, UE::CUnscopedEnum T> struct TDynamicDelegateParamPayloadCaptureType<WeakPtrType, T>           { using Type = TEnumAsByte<T>; };
+#else
+	template <typename WeakPtrType, UE::CUnscopedEnum T> struct TDynamicDelegateParamPayloadCaptureType<WeakPtrType, T>           { using Type = T; };
+#endif
 
 	template <typename WeakPtrType, typename T>
 	using TDynamicDelegateParamPayloadCaptureType_T = typename TDynamicDelegateParamPayloadCaptureType<WeakPtrType, T>::Type;
@@ -207,8 +213,12 @@ namespace UE::Core::Private
 			// ... otherwise we need to combine Params and Payload into a single struct
 			else
 			{
+#if UE_DYNAMIC_DELEGATE_UNSCOPED_ENUM_HACK
+				using FullTupleType   = TTuple<std::conditional_t<UE::CUnscopedEnum<std::decay_t<ParamTypes>>, TEnumAsByte<std::decay_t<ParamTypes>>, std::decay_t<ParamTypes>>...>;
+#else
 				using FullTupleType   = TTuple<std::decay_t<ParamTypes>...>;
-				using ParamsTupleType = TTuplePopBack_T<TTuple<std::decay_t<ParamTypes>...>, sizeof...(VarTypes)>;
+#endif
+				using ParamsTupleType = TTuplePopBack_T<FullTupleType, sizeof...(VarTypes)>;
 
 				((ParamsTupleType*)Parameters)->ApplyBefore(
 					[&Payload = Payload, ObjectPtr, Function](auto&&... Params)
diff --git a/Engine/Source/Runtime/CoreUObject/Public/UObject/DynamicDelegateTestObject.cpp b/Engine/Source/Runtime/CoreUObject/Public/UObject/DynamicDelegateTestObject.cpp
--- a/Engine/Source/Runtime/CoreUObject/Public/UObject/DynamicDelegateTestObject.cpp
+++ b/Engine/Source/Runtime/CoreUObject/Public/UObject/DynamicDelegateTestObject.cpp
@@ -15,6 +15,9 @@ FAnsiString GCalledFunctionName;
 int16 GInt16Arg = 5;
 int16 GInt16Out = 97;
 
+ETestUnscopedEnum::Type GUnscopedEnumArg = ETestUnscopedEnum::E2;
+ETestUnscopedEnum::Type GUnscopedEnumOut = ETestUnscopedEnum::E1;
+
 FString GStrArg = "Hello";
 FString GStrOut = "Goodbye";
 
@@ -51,59 +54,73 @@ void UDynamicDelegateTestObject::Func1(int16 Int)
 	CHECK(Int == GInt16Arg);
 }
 
-void UDynamicDelegateTestObject::Func2(int16 Int, FString Str)
+void UDynamicDelegateTestObject::Func2(int16 Int, ETestUnscopedEnum::Type UnscopedEnum)
 {
 	CHECK(GCalledFunctionName.IsEmpty());
 	GCalledFunctionName = "Func2";
 
-	CHECK(Int == GInt16Arg);
-	CHECK(Str == GStrArg);
+	CHECK(Int          == GInt16Arg);
+	CHECK(UnscopedEnum == GUnscopedEnumArg);
 }
 
-void UDynamicDelegateTestObject::Func3(int16 Int, FString Str, const UDynamicDelegateTestObject* Ptr)
+void UDynamicDelegateTestObject::Func3(int16 Int, ETestUnscopedEnum::Type UnscopedEnum, FString Str)
 {
 	CHECK(GCalledFunctionName.IsEmpty());
 	GCalledFunctionName = "Func3";
 
-	CHECK(Int == GInt16Arg);
-	CHECK(Str == GStrArg);
-	CHECK(Ptr == GPtrArg);
+	CHECK(Int          == GInt16Arg);
+	CHECK(UnscopedEnum == GUnscopedEnumArg);
+	CHECK(Str          == GStrArg);
 }
 
-void UDynamicDelegateTestObject::Func4(int16 Int, FString Str, const UDynamicDelegateTestObject* Ptr, FName Name)
+void UDynamicDelegateTestObject::Func4(int16 Int, ETestUnscopedEnum::Type UnscopedEnum, FString Str, const UDynamicDelegateTestObject* Ptr)
 {
 	CHECK(GCalledFunctionName.IsEmpty());
 	GCalledFunctionName = "Func4";
 
-	CHECK(Int  == GInt16Arg);
-	CHECK(Str  == GStrArg);
-	CHECK(Ptr  == GPtrArg);
-	CHECK(Name == GNameArg);
+	CHECK(Int          == GInt16Arg);
+	CHECK(UnscopedEnum == GUnscopedEnumArg);
+	CHECK(Str          == GStrArg);
+	CHECK(Ptr          == GPtrArg);
 }
 
-void UDynamicDelegateTestObject::Func5(int16 Int, FString Str, const UDynamicDelegateTestObject* Ptr, FName Name, bool bBool)
+void UDynamicDelegateTestObject::Func5(int16 Int, ETestUnscopedEnum::Type UnscopedEnum, FString Str, const UDynamicDelegateTestObject* Ptr, FName Name)
 {
 	CHECK(GCalledFunctionName.IsEmpty());
 	GCalledFunctionName = "Func5";
 
-	CHECK(Int   == GInt16Arg);
-	CHECK(Str   == GStrArg);
-	CHECK(Ptr   == GPtrArg);
-	CHECK(Name  == GNameArg);
-	CHECK(bBool == bGBoolArg);
+	CHECK(Int          == GInt16Arg);
+	CHECK(UnscopedEnum == GUnscopedEnumArg);
+	CHECK(Str          == GStrArg);
+	CHECK(Ptr          == GPtrArg);
+	CHECK(Name         == GNameArg);
 }
 
-void UDynamicDelegateTestObject::Func6(int16 Int, FString Str, const UDynamicDelegateTestObject* Ptr, FName Name, bool bBool, TWeakObjectPtr<UDynamicDelegateTestObject> WeakObjPtr)
+void UDynamicDelegateTestObject::Func6(int16 Int, ETestUnscopedEnum::Type UnscopedEnum, FString Str, const UDynamicDelegateTestObject* Ptr, FName Name, bool bBool)
 {
 	CHECK(GCalledFunctionName.IsEmpty());
 	GCalledFunctionName = "Func6";
 
-	CHECK(Int        == GInt16Arg);
-	CHECK(Str        == GStrArg);
-	CHECK(Ptr        == GPtrArg);
-	CHECK(Name       == GNameArg);
-	CHECK(bBool      == bGBoolArg);
-	CHECK(WeakObjPtr == GWeakObjPtrArg);
+	CHECK(Int          == GInt16Arg);
+	CHECK(UnscopedEnum == GUnscopedEnumArg);
+	CHECK(Str          == GStrArg);
+	CHECK(Ptr          == GPtrArg);
+	CHECK(Name         == GNameArg);
+	CHECK(bBool        == bGBoolArg);
+}
+
+void UDynamicDelegateTestObject::Func7(int16 Int, ETestUnscopedEnum::Type UnscopedEnum, FString Str, const UDynamicDelegateTestObject* Ptr, FName Name, bool bBool, TWeakObjectPtr<UDynamicDelegateTestObject> WeakObjPtr)
+{
+	CHECK(GCalledFunctionName.IsEmpty());
+	GCalledFunctionName = "Func7";
+
+	CHECK(Int          == GInt16Arg);
+	CHECK(UnscopedEnum == GUnscopedEnumArg);
+	CHECK(Str          == GStrArg);
+	CHECK(Ptr          == GPtrArg);
+	CHECK(Name         == GNameArg);
+	CHECK(bBool        == bGBoolArg);
+	CHECK(WeakObjPtr   == GWeakObjPtrArg);
 }
 
 void UDynamicDelegateTestObject::ConstFunc0() const
@@ -120,59 +137,73 @@ void UDynamicDelegateTestObject::ConstFunc1(int16 Int) const
 	CHECK(Int == GInt16Arg);
 }
 
-void UDynamicDelegateTestObject::ConstFunc2(int16 Int, FString Str) const
+void UDynamicDelegateTestObject::ConstFunc2(int16 Int, ETestUnscopedEnum::Type UnscopedEnum) const
 {
 	CHECK(GCalledFunctionName.IsEmpty());
 	GCalledFunctionName = "ConstFunc2";
 
-	CHECK(Int == GInt16Arg);
-	CHECK(Str == GStrArg);
+	CHECK(Int          == GInt16Arg);
+	CHECK(UnscopedEnum == GUnscopedEnumArg);
 }
 
-void UDynamicDelegateTestObject::ConstFunc3(int16 Int, FString Str, const UDynamicDelegateTestObject* Ptr) const
+void UDynamicDelegateTestObject::ConstFunc3(int16 Int, ETestUnscopedEnum::Type UnscopedEnum, FString Str) const
 {
 	CHECK(GCalledFunctionName.IsEmpty());
 	GCalledFunctionName = "ConstFunc3";
 
-	CHECK(Int == GInt16Arg);
-	CHECK(Str == GStrArg);
-	CHECK(Ptr == GPtrArg);
+	CHECK(Int          == GInt16Arg);
+	CHECK(UnscopedEnum == GUnscopedEnumArg);
+	CHECK(Str          == GStrArg);
 }
 
-void UDynamicDelegateTestObject::ConstFunc4(int16 Int, FString Str, const UDynamicDelegateTestObject* Ptr, FName Name) const
+void UDynamicDelegateTestObject::ConstFunc4(int16 Int, ETestUnscopedEnum::Type UnscopedEnum, FString Str, const UDynamicDelegateTestObject* Ptr) const
 {
 	CHECK(GCalledFunctionName.IsEmpty());
 	GCalledFunctionName = "ConstFunc4";
 
-	CHECK(Int  == GInt16Arg);
-	CHECK(Str  == GStrArg);
-	CHECK(Ptr  == GPtrArg);
-	CHECK(Name == GNameArg);
+	CHECK(Int          == GInt16Arg);
+	CHECK(UnscopedEnum == GUnscopedEnumArg);
+	CHECK(Str          == GStrArg);
+	CHECK(Ptr          == GPtrArg);
 }
 
-void UDynamicDelegateTestObject::ConstFunc5(int16 Int, FString Str, const UDynamicDelegateTestObject* Ptr, FName Name, bool bBool) const
+void UDynamicDelegateTestObject::ConstFunc5(int16 Int, ETestUnscopedEnum::Type UnscopedEnum, FString Str, const UDynamicDelegateTestObject* Ptr, FName Name) const
 {
 	CHECK(GCalledFunctionName.IsEmpty());
 	GCalledFunctionName = "ConstFunc5";
 
-	CHECK(Int   == GInt16Arg);
-	CHECK(Str   == GStrArg);
-	CHECK(Ptr   == GPtrArg);
-	CHECK(Name  == GNameArg);
-	CHECK(bBool == bGBoolArg);
+	CHECK(Int          == GInt16Arg);
+	CHECK(UnscopedEnum == GUnscopedEnumArg);
+	CHECK(Str          == GStrArg);
+	CHECK(Ptr          == GPtrArg);
+	CHECK(Name         == GNameArg);
 }
 
-void UDynamicDelegateTestObject::ConstFunc6(int16 Int, FString Str, const UDynamicDelegateTestObject* Ptr, FName Name, bool bBool, TWeakObjectPtr<UDynamicDelegateTestObject> WeakObjPtr) const
+void UDynamicDelegateTestObject::ConstFunc6(int16 Int, ETestUnscopedEnum::Type UnscopedEnum, FString Str, const UDynamicDelegateTestObject* Ptr, FName Name, bool bBool) const
 {
 	CHECK(GCalledFunctionName.IsEmpty());
 	GCalledFunctionName = "ConstFunc6";
 
-	CHECK(Int        == GInt16Arg);
-	CHECK(Str        == GStrArg);
-	CHECK(Ptr        == GPtrArg);
-	CHECK(Name       == GNameArg);
-	CHECK(bBool      == bGBoolArg);
-	CHECK(WeakObjPtr == GWeakObjPtrArg);
+	CHECK(Int          == GInt16Arg);
+	CHECK(UnscopedEnum == GUnscopedEnumArg);
+	CHECK(Str          == GStrArg);
+	CHECK(Ptr          == GPtrArg);
+	CHECK(Name         == GNameArg);
+	CHECK(bBool        == bGBoolArg);
+}
+
+void UDynamicDelegateTestObject::ConstFunc7(int16 Int, ETestUnscopedEnum::Type UnscopedEnum, FString Str, const UDynamicDelegateTestObject* Ptr, FName Name, bool bBool, TWeakObjectPtr<UDynamicDelegateTestObject> WeakObjPtr) const
+{
+	CHECK(GCalledFunctionName.IsEmpty());
+	GCalledFunctionName = "ConstFunc7";
+
+	CHECK(Int          == GInt16Arg);
+	CHECK(UnscopedEnum == GUnscopedEnumArg);
+	CHECK(Str          == GStrArg);
+	CHECK(Ptr          == GPtrArg);
+	CHECK(Name         == GNameArg);
+	CHECK(bBool        == bGBoolArg);
+	CHECK(WeakObjPtr   == GWeakObjPtrArg);
 }
 
 int UDynamicDelegateTestObject::FuncRetInt0()
@@ -189,41 +220,48 @@ int UDynamicDelegateTestObject::FuncRetInt1(int16 Int)
 	return GIntRet;
 }
 
-int UDynamicDelegateTestObject::FuncRetInt2(int16 Int, FString Str)
+int UDynamicDelegateTestObject::FuncRetInt2(int16 Int, ETestUnscopedEnum::Type UnscopedEnum)
 {
-	Func2(Int, Str);
+	Func2(Int, UnscopedEnum);
 	GCalledFunctionName = "FuncRetInt2";
 	return GIntRet;
 }
 
-int UDynamicDelegateTestObject::FuncRetInt3(int16 Int, FString Str, const UDynamicDelegateTestObject* Ptr)
+int UDynamicDelegateTestObject::FuncRetInt3(int16 Int, ETestUnscopedEnum::Type UnscopedEnum, FString Str)
 {
-	Func3(Int, Str, Ptr);
+	Func3(Int, UnscopedEnum, Str);
 	GCalledFunctionName = "FuncRetInt3";
 	return GIntRet;
 }
 
-int UDynamicDelegateTestObject::FuncRetInt4(int16 Int, FString Str, const UDynamicDelegateTestObject* Ptr, FName Name)
+int UDynamicDelegateTestObject::FuncRetInt4(int16 Int, ETestUnscopedEnum::Type UnscopedEnum, FString Str, const UDynamicDelegateTestObject* Ptr)
 {
-	Func4(Int, Str, Ptr, Name);
+	Func4(Int, UnscopedEnum, Str, Ptr);
 	GCalledFunctionName = "FuncRetInt4";
 	return GIntRet;
 }
 
-int UDynamicDelegateTestObject::FuncRetInt5(int16 Int, FString Str, const UDynamicDelegateTestObject* Ptr, FName Name, bool bBool)
+int UDynamicDelegateTestObject::FuncRetInt5(int16 Int, ETestUnscopedEnum::Type UnscopedEnum, FString Str, const UDynamicDelegateTestObject* Ptr, FName Name)
 {
-	Func5(Int, Str, Ptr, Name, bBool);
+	Func5(Int, UnscopedEnum, Str, Ptr, Name);
 	GCalledFunctionName = "FuncRetInt5";
 	return GIntRet;
 }
 
-int UDynamicDelegateTestObject::FuncRetInt6(int16 Int, FString Str, const UDynamicDelegateTestObject* Ptr, FName Name, bool bBool, TWeakObjectPtr<UDynamicDelegateTestObject> WeakObjPtr)
+int UDynamicDelegateTestObject::FuncRetInt6(int16 Int, ETestUnscopedEnum::Type UnscopedEnum, FString Str, const UDynamicDelegateTestObject* Ptr, FName Name, bool bBool)
 {
-	Func6(Int, Str, Ptr, Name, bBool, WeakObjPtr);
+	Func6(Int, UnscopedEnum, Str, Ptr, Name, bBool);
 	GCalledFunctionName = "FuncRetInt6";
 	return GIntRet;
 }
 
+int UDynamicDelegateTestObject::FuncRetInt7(int16 Int, ETestUnscopedEnum::Type UnscopedEnum, FString Str, const UDynamicDelegateTestObject* Ptr, FName Name, bool bBool, TWeakObjectPtr<UDynamicDelegateTestObject> WeakObjPtr)
+{
+	Func7(Int, UnscopedEnum, Str, Ptr, Name, bBool, WeakObjPtr);
+	GCalledFunctionName = "FuncRetInt7";
+	return GIntRet;
+}
+
 FString UDynamicDelegateTestObject::FuncRetStr0()
 {
 	Func0();
@@ -238,41 +276,48 @@ FString UDynamicDelegateTestObject::FuncRetStr1(int16 Int)
 	return GStrRet;
 }
 
-FString UDynamicDelegateTestObject::FuncRetStr2(int16 Int, FString Str)
+FString UDynamicDelegateTestObject::FuncRetStr2(int16 Int, ETestUnscopedEnum::Type UnscopedEnum)
 {
-	Func2(Int, Str);
+	Func2(Int, UnscopedEnum);
 	GCalledFunctionName = "FuncRetStr2";
 	return GStrRet;
 }
 
-FString UDynamicDelegateTestObject::FuncRetStr3(int16 Int, FString Str, const UDynamicDelegateTestObject* Ptr)
+FString UDynamicDelegateTestObject::FuncRetStr3(int16 Int, ETestUnscopedEnum::Type UnscopedEnum, FString Str)
 {
-	Func3(Int, Str, Ptr);
+	Func3(Int, UnscopedEnum, Str);
 	GCalledFunctionName = "FuncRetStr3";
 	return GStrRet;
 }
 
-FString UDynamicDelegateTestObject::FuncRetStr4(int16 Int, FString Str, const UDynamicDelegateTestObject* Ptr, FName Name)
+FString UDynamicDelegateTestObject::FuncRetStr4(int16 Int, ETestUnscopedEnum::Type UnscopedEnum, FString Str, const UDynamicDelegateTestObject* Ptr)
 {
-	Func4(Int, Str, Ptr, Name);
+	Func4(Int, UnscopedEnum, Str, Ptr);
 	GCalledFunctionName = "FuncRetStr4";
 	return GStrRet;
 }
 
-FString UDynamicDelegateTestObject::FuncRetStr5(int16 Int, FString Str, const UDynamicDelegateTestObject* Ptr, FName Name, bool bBool)
+FString UDynamicDelegateTestObject::FuncRetStr5(int16 Int, ETestUnscopedEnum::Type UnscopedEnum, FString Str, const UDynamicDelegateTestObject* Ptr, FName Name)
 {
-	Func5(Int, Str, Ptr, Name, bBool);
+	Func5(Int, UnscopedEnum, Str, Ptr, Name);
 	GCalledFunctionName = "FuncRetStr5";
 	return GStrRet;
 }
 
-FString UDynamicDelegateTestObject::FuncRetStr6(int16 Int, FString Str, const UDynamicDelegateTestObject* Ptr, FName Name, bool bBool, TWeakObjectPtr<UDynamicDelegateTestObject> WeakObjPtr)
+FString UDynamicDelegateTestObject::FuncRetStr6(int16 Int, ETestUnscopedEnum::Type UnscopedEnum, FString Str, const UDynamicDelegateTestObject* Ptr, FName Name, bool bBool)
 {
-	Func6(Int, Str, Ptr, Name, bBool, WeakObjPtr);
+	Func6(Int, UnscopedEnum, Str, Ptr, Name, bBool);
 	GCalledFunctionName = "FuncRetStr6";
 	return GStrRet;
 }
 
+FString UDynamicDelegateTestObject::FuncRetStr7(int16 Int, ETestUnscopedEnum::Type UnscopedEnum, FString Str, const UDynamicDelegateTestObject* Ptr, FName Name, bool bBool, TWeakObjectPtr<UDynamicDelegateTestObject> WeakObjPtr)
+{
+	Func7(Int, UnscopedEnum, Str, Ptr, Name, bBool, WeakObjPtr);
+	GCalledFunctionName = "FuncRetStr7";
+	return GStrRet;
+}
+
 int UDynamicDelegateTestObject::ConstFuncRetInt0() const
 {
 	ConstFunc0();
@@ -287,41 +332,48 @@ int UDynamicDelegateTestObject::ConstFuncRetInt1(int16 Int) const
 	return GIntRet;
 }
 
-int UDynamicDelegateTestObject::ConstFuncRetInt2(int16 Int, FString Str) const
+int UDynamicDelegateTestObject::ConstFuncRetInt2(int16 Int, ETestUnscopedEnum::Type UnscopedEnum) const
 {
-	ConstFunc2(Int, Str);
+	ConstFunc2(Int, UnscopedEnum);
 	GCalledFunctionName = "ConstFuncRetInt2";
 	return GIntRet;
 }
 
-int UDynamicDelegateTestObject::ConstFuncRetInt3(int16 Int, FString Str, const UDynamicDelegateTestObject* Ptr) const
+int UDynamicDelegateTestObject::ConstFuncRetInt3(int16 Int, ETestUnscopedEnum::Type UnscopedEnum, FString Str) const
 {
-	ConstFunc3(Int, Str, Ptr);
+	ConstFunc3(Int, UnscopedEnum, Str);
 	GCalledFunctionName = "ConstFuncRetInt3";
 	return GIntRet;
 }
 
-int UDynamicDelegateTestObject::ConstFuncRetInt4(int16 Int, FString Str, const UDynamicDelegateTestObject* Ptr, FName Name) const
+int UDynamicDelegateTestObject::ConstFuncRetInt4(int16 Int, ETestUnscopedEnum::Type UnscopedEnum, FString Str, const UDynamicDelegateTestObject* Ptr) const
 {
-	ConstFunc4(Int, Str, Ptr, Name);
+	ConstFunc4(Int, UnscopedEnum, Str, Ptr);
 	GCalledFunctionName = "ConstFuncRetInt4";
 	return GIntRet;
 }
 
-int UDynamicDelegateTestObject::ConstFuncRetInt5(int16 Int, FString Str, const UDynamicDelegateTestObject* Ptr, FName Name, bool bBool) const
+int UDynamicDelegateTestObject::ConstFuncRetInt5(int16 Int, ETestUnscopedEnum::Type UnscopedEnum, FString Str, const UDynamicDelegateTestObject* Ptr, FName Name) const
 {
-	ConstFunc5(Int, Str, Ptr, Name, bBool);
+	ConstFunc5(Int, UnscopedEnum, Str, Ptr, Name);
 	GCalledFunctionName = "ConstFuncRetInt5";
 	return GIntRet;
 }
 
-int UDynamicDelegateTestObject::ConstFuncRetInt6(int16 Int, FString Str, const UDynamicDelegateTestObject* Ptr, FName Name, bool bBool, TWeakObjectPtr<UDynamicDelegateTestObject> WeakObjPtr) const
+int UDynamicDelegateTestObject::ConstFuncRetInt6(int16 Int, ETestUnscopedEnum::Type UnscopedEnum, FString Str, const UDynamicDelegateTestObject* Ptr, FName Name, bool bBool) const
 {
-	ConstFunc6(Int, Str, Ptr, Name, bBool, WeakObjPtr);
+	ConstFunc6(Int, UnscopedEnum, Str, Ptr, Name, bBool);
 	GCalledFunctionName = "ConstFuncRetInt6";
 	return GIntRet;
 }
 
+int UDynamicDelegateTestObject::ConstFuncRetInt7(int16 Int, ETestUnscopedEnum::Type UnscopedEnum, FString Str, const UDynamicDelegateTestObject* Ptr, FName Name, bool bBool, TWeakObjectPtr<UDynamicDelegateTestObject> WeakObjPtr) const
+{
+	ConstFunc7(Int, UnscopedEnum, Str, Ptr, Name, bBool, WeakObjPtr);
+	GCalledFunctionName = "ConstFuncRetInt7";
+	return GIntRet;
+}
+
 FString UDynamicDelegateTestObject::ConstFuncRetStr0() const
 {
 	ConstFunc0();
@@ -336,41 +388,48 @@ FString UDynamicDelegateTestObject::ConstFuncRetStr1(int16 Int) const
 	return GStrRet;
 }
 
-FString UDynamicDelegateTestObject::ConstFuncRetStr2(int16 Int, FString Str) const
+FString UDynamicDelegateTestObject::ConstFuncRetStr2(int16 Int, ETestUnscopedEnum::Type UnscopedEnum) const
 {
-	ConstFunc2(Int, Str);
+	ConstFunc2(Int, UnscopedEnum);
 	GCalledFunctionName = "ConstFuncRetStr2";
 	return GStrRet;
 }
 
-FString UDynamicDelegateTestObject::ConstFuncRetStr3(int16 Int, FString Str, const UDynamicDelegateTestObject* Ptr) const
+FString UDynamicDelegateTestObject::ConstFuncRetStr3(int16 Int, ETestUnscopedEnum::Type UnscopedEnum, FString Str) const
 {
-	ConstFunc3(Int, Str, Ptr);
+	ConstFunc3(Int, UnscopedEnum, Str);
 	GCalledFunctionName = "ConstFuncRetStr3";
 	return GStrRet;
 }
 
-FString UDynamicDelegateTestObject::ConstFuncRetStr4(int16 Int, FString Str, const UDynamicDelegateTestObject* Ptr, FName Name) const
+FString UDynamicDelegateTestObject::ConstFuncRetStr4(int16 Int, ETestUnscopedEnum::Type UnscopedEnum, FString Str, const UDynamicDelegateTestObject* Ptr) const
 {
-	ConstFunc4(Int, Str, Ptr, Name);
+	ConstFunc4(Int, UnscopedEnum, Str, Ptr);
 	GCalledFunctionName = "ConstFuncRetStr4";
 	return GStrRet;
 }
 
-FString UDynamicDelegateTestObject::ConstFuncRetStr5(int16 Int, FString Str, const UDynamicDelegateTestObject* Ptr, FName Name, bool bBool) const
+FString UDynamicDelegateTestObject::ConstFuncRetStr5(int16 Int, ETestUnscopedEnum::Type UnscopedEnum, FString Str, const UDynamicDelegateTestObject* Ptr, FName Name) const
 {
-	ConstFunc5(Int, Str, Ptr, Name, bBool);
+	ConstFunc5(Int, UnscopedEnum, Str, Ptr, Name);
 	GCalledFunctionName = "ConstFuncRetStr5";
 	return GStrRet;
 }
 
-FString UDynamicDelegateTestObject::ConstFuncRetStr6(int16 Int, FString Str, const UDynamicDelegateTestObject* Ptr, FName Name, bool bBool, TWeakObjectPtr<UDynamicDelegateTestObject> WeakObjPtr) const
+FString UDynamicDelegateTestObject::ConstFuncRetStr6(int16 Int, ETestUnscopedEnum::Type UnscopedEnum, FString Str, const UDynamicDelegateTestObject* Ptr, FName Name, bool bBool) const
 {
-	ConstFunc6(Int, Str, Ptr, Name, bBool, WeakObjPtr);
+	ConstFunc6(Int, UnscopedEnum, Str, Ptr, Name, bBool);
 	GCalledFunctionName = "ConstFuncRetStr6";
 	return GStrRet;
 }
 
+FString UDynamicDelegateTestObject::ConstFuncRetStr7(int16 Int, ETestUnscopedEnum::Type UnscopedEnum, FString Str, const UDynamicDelegateTestObject* Ptr, FName Name, bool bBool, TWeakObjectPtr<UDynamicDelegateTestObject> WeakObjPtr) const
+{
+	ConstFunc7(Int, UnscopedEnum, Str, Ptr, Name, bBool, WeakObjPtr);
+	GCalledFunctionName = "ConstFuncRetStr7";
+	return GStrRet;
+}
+
 void UDynamicDelegateTestObject::FuncOut1(int16& Int)
 {
 	CHECK(GCalledFunctionName.IsEmpty());
@@ -379,34 +438,42 @@ void UDynamicDelegateTestObject::FuncOut1(int16& Int)
 	Int = GInt16Out;
 }
 
-void UDynamicDelegateTestObject::FuncOut2(int16& Int, FString& Str)
+void UDynamicDelegateTestObject::FuncOut2(int16& Int, TEnumAsByte<ETestUnscopedEnum::Type>& UnscopedEnum)
 {
 	FuncOut1(Int);
 	GCalledFunctionName = "FuncOut2";
+	CHECK(UnscopedEnum == GUnscopedEnumArg);
+	UnscopedEnum = GUnscopedEnumOut;
+}
+
+void UDynamicDelegateTestObject::FuncOut3(int16& Int, TEnumAsByte<ETestUnscopedEnum::Type>& UnscopedEnum, FString& Str)
+{
+	FuncOut2(Int, UnscopedEnum);
+	GCalledFunctionName = "FuncOut3";
 	CHECK(Str == GStrArg);
 	Str = GStrOut;
 }
 
-void UDynamicDelegateTestObject::FuncOut3(int16& Int, FString& Str, FName& Name)
+void UDynamicDelegateTestObject::FuncOut4(int16& Int, TEnumAsByte<ETestUnscopedEnum::Type>& UnscopedEnum, FString& Str, FName& Name)
 {
-	FuncOut2(Int, Str);
-	GCalledFunctionName = "FuncOut3";
+	FuncOut3(Int, UnscopedEnum, Str);
+	GCalledFunctionName = "FuncOut4";
 	CHECK(Name == GNameArg);
 	Name = GNameOut;
 }
 
-void UDynamicDelegateTestObject::FuncOut4(int16& Int, FString& Str, FName& Name, bool& bBool)
+void UDynamicDelegateTestObject::FuncOut5(int16& Int, TEnumAsByte<ETestUnscopedEnum::Type>& UnscopedEnum, FString& Str, FName& Name, bool& bBool)
 {
-	FuncOut3(Int, Str, Name);
-	GCalledFunctionName = "FuncOut4";
+	FuncOut4(Int, UnscopedEnum, Str, Name);
+	GCalledFunctionName = "FuncOut5";
 	CHECK(bBool == bGBoolArg);
 	bBool = bGBoolOut;
 }
 
-void UDynamicDelegateTestObject::FuncOut5(int16& Int, FString& Str, FName& Name, bool& bBool, TWeakObjectPtr<UDynamicDelegateTestObject>& WeakObjPtr)
+void UDynamicDelegateTestObject::FuncOut6(int16& Int, TEnumAsByte<ETestUnscopedEnum::Type>& UnscopedEnum, FString& Str, FName& Name, bool& bBool, TWeakObjectPtr<UDynamicDelegateTestObject>& WeakObjPtr)
 {
-	FuncOut4(Int, Str, Name, bBool);
-	GCalledFunctionName = "FuncOut5";
+	FuncOut5(Int, UnscopedEnum, Str, Name, bBool);
+	GCalledFunctionName = "FuncOut6";
 	CHECK(WeakObjPtr == GWeakObjPtrArg);
 	WeakObjPtr = GWeakObjPtrOut;
 }
@@ -419,51 +486,69 @@ void UDynamicDelegateTestObject::ConstFuncOut1(int16& Int) const
 	Int = GInt16Out;
 }
 
-void UDynamicDelegateTestObject::ConstFuncOut2(int16& Int, FString& Str) const
+void UDynamicDelegateTestObject::ConstFuncOut2(int16& Int, TEnumAsByte<ETestUnscopedEnum::Type>& UnscopedEnum) const
 {
 	ConstFuncOut1(Int);
 	GCalledFunctionName = "ConstFuncOut2";
+	CHECK(UnscopedEnum == GUnscopedEnumArg);
+	UnscopedEnum = GUnscopedEnumOut;
+}
+
+void UDynamicDelegateTestObject::ConstFuncOut3(int16& Int, TEnumAsByte<ETestUnscopedEnum::Type>& UnscopedEnum, FString& Str) const
+{
+	ConstFuncOut2(Int, UnscopedEnum);
+	GCalledFunctionName = "ConstFuncOut3";
 	CHECK(Str == GStrArg);
 	Str = GStrOut;
 }
 
-void UDynamicDelegateTestObject::ConstFuncOut3(int16& Int, FString& Str, FName& Name) const
+void UDynamicDelegateTestObject::ConstFuncOut4(int16& Int, TEnumAsByte<ETestUnscopedEnum::Type>& UnscopedEnum, FString& Str, FName& Name) const
 {
-	ConstFuncOut2(Int, Str);
-	GCalledFunctionName = "ConstFuncOut3";
+	ConstFuncOut3(Int, UnscopedEnum, Str);
+	GCalledFunctionName = "ConstFuncOut4";
 	CHECK(Name == GNameArg);
 	Name = GNameOut;
 }
 
-void UDynamicDelegateTestObject::ConstFuncOut4(int16& Int, FString& Str, FName& Name, bool& bBool) const
+void UDynamicDelegateTestObject::ConstFuncOut5(int16& Int, TEnumAsByte<ETestUnscopedEnum::Type>& UnscopedEnum, FString& Str, FName& Name, bool& bBool) const
 {
-	ConstFuncOut3(Int, Str, Name);
-	GCalledFunctionName = "ConstFuncOut4";
+	ConstFuncOut4(Int, UnscopedEnum, Str, Name);
+	GCalledFunctionName = "ConstFuncOut5";
 	CHECK(bBool == bGBoolArg);
 	bBool = bGBoolOut;
 }
 
-void UDynamicDelegateTestObject::ConstFuncOut5(int16& Int, FString& Str, FName& Name, bool& bBool, TWeakObjectPtr<UDynamicDelegateTestObject>& WeakObjPtr) const
+void UDynamicDelegateTestObject::ConstFuncOut6(int16& Int, TEnumAsByte<ETestUnscopedEnum::Type>& UnscopedEnum, FString& Str, FName& Name, bool& bBool, TWeakObjectPtr<UDynamicDelegateTestObject>& WeakObjPtr) const
 {
-	ConstFuncOut4(Int, Str, Name, bBool);
-	GCalledFunctionName = "ConstFuncOut5";
+	ConstFuncOut5(Int, UnscopedEnum, Str, Name, bBool);
+	GCalledFunctionName = "ConstFuncOut6";
 	CHECK(WeakObjPtr == GWeakObjPtrArg);
 	WeakObjPtr = GWeakObjPtrOut;
 }
 
-void UDynamicDelegateTestObject::TestUnscopedEnum(int16 Int1, ETestUnscopedEnum::Type E, int16 Int2)
+// Forwards an unscoped enum
+template <typename T>
+UE_REWRITE decltype(auto) MoveToOutArg(T&& Ref)
 {
-	CHECK(Int1 == 5);
-	CHECK(E == ETestUnscopedEnum::E2);
-	CHECK(Int2 == 10);
+	using DecayedT = std::decay_t<T>;
+
+	if constexpr (UE::CUnscopedEnum<DecayedT>)
+	{
+		return TEnumAsByte<DecayedT>(Ref);
+	}
+	else
+	{
+		return Ref;
+	}
 }
 
-#define DEFINE_ARG_VARIABLES_1(Arg1Val)                                              auto Arg1 = Arg1Val;
-#define DEFINE_ARG_VARIABLES_2(Arg1Val, Arg2Val)                                     auto Arg1 = Arg1Val; auto Arg2 = Arg2Val;
-#define DEFINE_ARG_VARIABLES_3(Arg1Val, Arg2Val, Arg3Val)                            auto Arg1 = Arg1Val; auto Arg2 = Arg2Val; auto Arg3 = Arg3Val;
-#define DEFINE_ARG_VARIABLES_4(Arg1Val, Arg2Val, Arg3Val, Arg4Val)                   auto Arg1 = Arg1Val; auto Arg2 = Arg2Val; auto Arg3 = Arg3Val; auto Arg4 = Arg4Val;
-#define DEFINE_ARG_VARIABLES_5(Arg1Val, Arg2Val, Arg3Val, Arg4Val, Arg5Val)          auto Arg1 = Arg1Val; auto Arg2 = Arg2Val; auto Arg3 = Arg3Val; auto Arg4 = Arg4Val; auto Arg5 = Arg5Val;
-#define DEFINE_ARG_VARIABLES_6(Arg1Val, Arg2Val, Arg3Val, Arg4Val, Arg5Val, Arg6Val) auto Arg1 = Arg1Val; auto Arg2 = Arg2Val; auto Arg3 = Arg3Val; auto Arg4 = Arg4Val; auto Arg5 = Arg5Val; auto Arg6 = Arg6Val;
+#define DEFINE_ARG_VARIABLES_1(Arg1Val)                                                       auto Arg1 = MoveToOutArg(Arg1Val);
+#define DEFINE_ARG_VARIABLES_2(Arg1Val, Arg2Val)                                              auto Arg1 = MoveToOutArg(Arg1Val); auto Arg2 = MoveToOutArg(Arg2Val);
+#define DEFINE_ARG_VARIABLES_3(Arg1Val, Arg2Val, Arg3Val)                                     auto Arg1 = MoveToOutArg(Arg1Val); auto Arg2 = MoveToOutArg(Arg2Val); auto Arg3 = MoveToOutArg(Arg3Val);
+#define DEFINE_ARG_VARIABLES_4(Arg1Val, Arg2Val, Arg3Val, Arg4Val)                            auto Arg1 = MoveToOutArg(Arg1Val); auto Arg2 = MoveToOutArg(Arg2Val); auto Arg3 = MoveToOutArg(Arg3Val); auto Arg4 = MoveToOutArg(Arg4Val);
+#define DEFINE_ARG_VARIABLES_5(Arg1Val, Arg2Val, Arg3Val, Arg4Val, Arg5Val)                   auto Arg1 = MoveToOutArg(Arg1Val); auto Arg2 = MoveToOutArg(Arg2Val); auto Arg3 = MoveToOutArg(Arg3Val); auto Arg4 = MoveToOutArg(Arg4Val); auto Arg5 = MoveToOutArg(Arg5Val);
+#define DEFINE_ARG_VARIABLES_6(Arg1Val, Arg2Val, Arg3Val, Arg4Val, Arg5Val, Arg6Val)          auto Arg1 = MoveToOutArg(Arg1Val); auto Arg2 = MoveToOutArg(Arg2Val); auto Arg3 = MoveToOutArg(Arg3Val); auto Arg4 = MoveToOutArg(Arg4Val); auto Arg5 = MoveToOutArg(Arg5Val); auto Arg6 = MoveToOutArg(Arg6Val);
+#define DEFINE_ARG_VARIABLES_7(Arg1Val, Arg2Val, Arg3Val, Arg4Val, Arg5Val, Arg6Val, Arg7Val) auto Arg1 = MoveToOutArg(Arg1Val); auto Arg2 = MoveToOutArg(Arg2Val); auto Arg3 = MoveToOutArg(Arg3Val); auto Arg4 = MoveToOutArg(Arg4Val); auto Arg5 = MoveToOutArg(Arg5Val); auto Arg6 = MoveToOutArg(Arg6Val); auto Arg7 = MoveToOutArg(Arg7Val);
 
 #define EXPAND_ARG_VARIABLES_1( ...) (Arg1)
 #define EXPAND_ARG_VARIABLES_2( ...) (Arg1, Arg2)
@@ -471,13 +556,15 @@ void UDynamicDelegateTestObject::TestUnscopedEnum(int16 Int1, ETestUnscopedEnum:
 #define EXPAND_ARG_VARIABLES_4( ...) (Arg1, Arg2, Arg3, Arg4)
 #define EXPAND_ARG_VARIABLES_5( ...) (Arg1, Arg2, Arg3, Arg4, Arg5)
 #define EXPAND_ARG_VARIABLES_6( ...) (Arg1, Arg2, Arg3, Arg4, Arg5, Arg6)
+#define EXPAND_ARG_VARIABLES_7( ...) (Arg1, Arg2, Arg3, Arg4, Arg5, Arg6, Arg7)
 
-#define CHECK_ARG_VARIABLES_1(Arg1Val)                                              CHECK(Arg1 == Arg1Val);
-#define CHECK_ARG_VARIABLES_2(Arg1Val, Arg2Val)                                     CHECK(Arg1 == Arg1Val); CHECK(Arg2 == Arg2Val);
-#define CHECK_ARG_VARIABLES_3(Arg1Val, Arg2Val, Arg3Val)                            CHECK(Arg1 == Arg1Val); CHECK(Arg2 == Arg2Val); CHECK(Arg3 == Arg3Val);
-#define CHECK_ARG_VARIABLES_4(Arg1Val, Arg2Val, Arg3Val, Arg4Val)                   CHECK(Arg1 == Arg1Val); CHECK(Arg2 == Arg2Val); CHECK(Arg3 == Arg3Val); CHECK(Arg4 == Arg4Val);
-#define CHECK_ARG_VARIABLES_5(Arg1Val, Arg2Val, Arg3Val, Arg4Val, Arg5Val)          CHECK(Arg1 == Arg1Val); CHECK(Arg2 == Arg2Val); CHECK(Arg3 == Arg3Val); CHECK(Arg4 == Arg4Val); CHECK(Arg5 == Arg5Val);
-#define CHECK_ARG_VARIABLES_6(Arg1Val, Arg2Val, Arg3Val, Arg4Val, Arg5Val, Arr6Val) CHECK(Arg1 == Arg1Val); CHECK(Arg2 == Arg2Val); CHECK(Arg3 == Arg3Val); CHECK(Arg4 == Arg4Val); CHECK(Arg5 == Arg5Val); CHECK(Arg6 == Arg6Val);
+#define CHECK_ARG_VARIABLES_1(Arg1Val)                                                       CHECK(Arg1 == Arg1Val);
+#define CHECK_ARG_VARIABLES_2(Arg1Val, Arg2Val)                                              CHECK(Arg1 == Arg1Val); CHECK(Arg2 == Arg2Val);
+#define CHECK_ARG_VARIABLES_3(Arg1Val, Arg2Val, Arg3Val)                                     CHECK(Arg1 == Arg1Val); CHECK(Arg2 == Arg2Val); CHECK(Arg3 == Arg3Val);
+#define CHECK_ARG_VARIABLES_4(Arg1Val, Arg2Val, Arg3Val, Arg4Val)                            CHECK(Arg1 == Arg1Val); CHECK(Arg2 == Arg2Val); CHECK(Arg3 == Arg3Val); CHECK(Arg4 == Arg4Val);
+#define CHECK_ARG_VARIABLES_5(Arg1Val, Arg2Val, Arg3Val, Arg4Val, Arg5Val)                   CHECK(Arg1 == Arg1Val); CHECK(Arg2 == Arg2Val); CHECK(Arg3 == Arg3Val); CHECK(Arg4 == Arg4Val); CHECK(Arg5 == Arg5Val);
+#define CHECK_ARG_VARIABLES_6(Arg1Val, Arg2Val, Arg3Val, Arg4Val, Arg5Val, Arg6Val)          CHECK(Arg1 == Arg1Val); CHECK(Arg2 == Arg2Val); CHECK(Arg3 == Arg3Val); CHECK(Arg4 == Arg4Val); CHECK(Arg5 == Arg5Val); CHECK(Arg6 == Arg6Val);
+#define CHECK_ARG_VARIABLES_7(Arg1Val, Arg2Val, Arg3Val, Arg4Val, Arg5Val, Arg6Val, Arg7Val) CHECK(Arg1 == Arg1Val); CHECK(Arg2 == Arg2Val); CHECK(Arg3 == Arg3Val); CHECK(Arg4 == Arg4Val); CHECK(Arg5 == Arg5Val); CHECK(Arg6 == Arg6Val); CHECK(Arg7 == Arg7Val);
 
 #define DEFINE_ARG_VARIABLES(...) UE_APPEND_VA_ARG_COUNT(DEFINE_ARG_VARIABLES_, __VA_ARGS__)(__VA_ARGS__)
 #define EXPAND_ARG_VARIABLES(...) UE_APPEND_VA_ARG_COUNT(EXPAND_ARG_VARIABLES_, __VA_ARGS__)(ObjectPtr, FuncPtr, __VA_ARGS__)
@@ -577,45 +664,45 @@ void UDynamicDelegateTestObject::RunTests()
 	// 2 param function
 	{
 		// 0 arg payload
-		BIND_CALL_CHECK_DELEGATE(FTestDynamicMulticastDelegate2::FDelegate, 0, &UDynamicDelegateTestObject::Func2,      (GInt16Arg, GStrArg));
-		BIND_CALL_CHECK_DELEGATE(FTestDynamicDelegate2,                     0, &UDynamicDelegateTestObject::Func2,      (GInt16Arg, GStrArg));
-		BIND_CALL_CHECK_DELEGATE(FTestDynamicMulticastDelegate2,            1, &UDynamicDelegateTestObject::Func2,      (GInt16Arg, GStrArg));
-		BIND_CALL_CHECK_DELEGATE(FTestDynamicMulticastDelegate2::FDelegate, 0, &UDynamicDelegateTestObject::ConstFunc2, (GInt16Arg, GStrArg));
-		BIND_CALL_CHECK_DELEGATE(FTestDynamicDelegate2,                     0, &UDynamicDelegateTestObject::ConstFunc2, (GInt16Arg, GStrArg));
-		BIND_CALL_CHECK_DELEGATE(FTestDynamicMulticastDelegate2,            1, &UDynamicDelegateTestObject::ConstFunc2, (GInt16Arg, GStrArg));
-
-		BIND_CALL_CHECK_DELEGATE_RETVAL(GIntRet, FTestDynamicDelegateRetInt2, 0, &UDynamicDelegateTestObject::FuncRetInt2,      (GInt16Arg, GStrArg));
-		BIND_CALL_CHECK_DELEGATE_RETVAL(GIntRet, FTestDynamicDelegateRetInt2, 0, &UDynamicDelegateTestObject::ConstFuncRetInt2, (GInt16Arg, GStrArg));
-		BIND_CALL_CHECK_DELEGATE_RETVAL(GStrRet, FTestDynamicDelegateRetStr2, 0, &UDynamicDelegateTestObject::FuncRetStr2,      (GInt16Arg, GStrArg));
-		BIND_CALL_CHECK_DELEGATE_RETVAL(GStrRet, FTestDynamicDelegateRetStr2, 0, &UDynamicDelegateTestObject::ConstFuncRetStr2, (GInt16Arg, GStrArg));
+		BIND_CALL_CHECK_DELEGATE(FTestDynamicMulticastDelegate2::FDelegate, 0, &UDynamicDelegateTestObject::Func2,      (GInt16Arg, GUnscopedEnumArg));
+		BIND_CALL_CHECK_DELEGATE(FTestDynamicDelegate2,                     0, &UDynamicDelegateTestObject::Func2,      (GInt16Arg, GUnscopedEnumArg));
+		BIND_CALL_CHECK_DELEGATE(FTestDynamicMulticastDelegate2,            1, &UDynamicDelegateTestObject::Func2,      (GInt16Arg, GUnscopedEnumArg));
+		BIND_CALL_CHECK_DELEGATE(FTestDynamicMulticastDelegate2::FDelegate, 0, &UDynamicDelegateTestObject::ConstFunc2, (GInt16Arg, GUnscopedEnumArg));
+		BIND_CALL_CHECK_DELEGATE(FTestDynamicDelegate2,                     0, &UDynamicDelegateTestObject::ConstFunc2, (GInt16Arg, GUnscopedEnumArg));
+		BIND_CALL_CHECK_DELEGATE(FTestDynamicMulticastDelegate2,            1, &UDynamicDelegateTestObject::ConstFunc2, (GInt16Arg, GUnscopedEnumArg));
+
+		BIND_CALL_CHECK_DELEGATE_RETVAL(GIntRet, FTestDynamicDelegateRetInt2, 0, &UDynamicDelegateTestObject::FuncRetInt2,      (GInt16Arg, GUnscopedEnumArg));
+		BIND_CALL_CHECK_DELEGATE_RETVAL(GIntRet, FTestDynamicDelegateRetInt2, 0, &UDynamicDelegateTestObject::ConstFuncRetInt2, (GInt16Arg, GUnscopedEnumArg));
+		BIND_CALL_CHECK_DELEGATE_RETVAL(GStrRet, FTestDynamicDelegateRetStr2, 0, &UDynamicDelegateTestObject::FuncRetStr2,      (GInt16Arg, GUnscopedEnumArg));
+		BIND_CALL_CHECK_DELEGATE_RETVAL(GStrRet, FTestDynamicDelegateRetStr2, 0, &UDynamicDelegateTestObject::ConstFuncRetStr2, (GInt16Arg, GUnscopedEnumArg));
 
 #if defined(UE_USE_DYNAMIC_DELEGATE_PAYLOADS) && UE_USE_DYNAMIC_DELEGATE_PAYLOADS
 		// 1 arg payload
-		BIND_CALL_CHECK_DELEGATE(FTestDynamicMulticastDelegate1::FDelegate, 0, &UDynamicDelegateTestObject::Func2,      (GInt16Arg), GStrArg);
-		BIND_CALL_CHECK_DELEGATE(FTestDynamicDelegate1,                     0, &UDynamicDelegateTestObject::Func2,      (GInt16Arg), GStrArg);
-		BIND_CALL_CHECK_DELEGATE(FTestDynamicMulticastDelegate1,            1, &UDynamicDelegateTestObject::Func2,      (GInt16Arg), GStrArg);
-		BIND_CALL_CHECK_DELEGATE(FTestDynamicMulticastDelegate1::FDelegate, 0, &UDynamicDelegateTestObject::ConstFunc2, (GInt16Arg), GStrArg);
-		BIND_CALL_CHECK_DELEGATE(FTestDynamicDelegate1,                     0, &UDynamicDelegateTestObject::ConstFunc2, (GInt16Arg), GStrArg);
-		BIND_CALL_CHECK_DELEGATE(FTestDynamicMulticastDelegate1,            1, &UDynamicDelegateTestObject::ConstFunc2, (GInt16Arg), GStrArg);
+		BIND_CALL_CHECK_DELEGATE(FTestDynamicMulticastDelegate1::FDelegate, 0, &UDynamicDelegateTestObject::Func2,      (GInt16Arg), GUnscopedEnumArg);
+		BIND_CALL_CHECK_DELEGATE(FTestDynamicDelegate1,                     0, &UDynamicDelegateTestObject::Func2,      (GInt16Arg), GUnscopedEnumArg);
+		BIND_CALL_CHECK_DELEGATE(FTestDynamicMulticastDelegate1,            1, &UDynamicDelegateTestObject::Func2,      (GInt16Arg), GUnscopedEnumArg);
+		BIND_CALL_CHECK_DELEGATE(FTestDynamicMulticastDelegate1::FDelegate, 0, &UDynamicDelegateTestObject::ConstFunc2, (GInt16Arg), GUnscopedEnumArg);
+		BIND_CALL_CHECK_DELEGATE(FTestDynamicDelegate1,                     0, &UDynamicDelegateTestObject::ConstFunc2, (GInt16Arg), GUnscopedEnumArg);
+		BIND_CALL_CHECK_DELEGATE(FTestDynamicMulticastDelegate1,            1, &UDynamicDelegateTestObject::ConstFunc2, (GInt16Arg), GUnscopedEnumArg);
 #if UE_USE_DYNAMIC_DELEGATE_PAYLOADS_RETVAL
-		BIND_CALL_CHECK_DELEGATE_RETVAL(GIntRet, FTestDynamicDelegateRetInt1, 0, &UDynamicDelegateTestObject::FuncRetInt2,      (GInt16Arg), GStrArg);
-		BIND_CALL_CHECK_DELEGATE_RETVAL(GIntRet, FTestDynamicDelegateRetInt1, 0, &UDynamicDelegateTestObject::ConstFuncRetInt2, (GInt16Arg), GStrArg);
-		BIND_CALL_CHECK_DELEGATE_RETVAL(GStrRet, FTestDynamicDelegateRetStr1, 0, &UDynamicDelegateTestObject::FuncRetStr2,      (GInt16Arg), GStrArg);
-		BIND_CALL_CHECK_DELEGATE_RETVAL(GStrRet, FTestDynamicDelegateRetStr1, 0, &UDynamicDelegateTestObject::ConstFuncRetStr2, (GInt16Arg), GStrArg);
+		BIND_CALL_CHECK_DELEGATE_RETVAL(GIntRet, FTestDynamicDelegateRetInt1, 0, &UDynamicDelegateTestObject::FuncRetInt2,      (GInt16Arg), GUnscopedEnumArg);
+		BIND_CALL_CHECK_DELEGATE_RETVAL(GIntRet, FTestDynamicDelegateRetInt1, 0, &UDynamicDelegateTestObject::ConstFuncRetInt2, (GInt16Arg), GUnscopedEnumArg);
+		BIND_CALL_CHECK_DELEGATE_RETVAL(GStrRet, FTestDynamicDelegateRetStr1, 0, &UDynamicDelegateTestObject::FuncRetStr2,      (GInt16Arg), GUnscopedEnumArg);
+		BIND_CALL_CHECK_DELEGATE_RETVAL(GStrRet, FTestDynamicDelegateRetStr1, 0, &UDynamicDelegateTestObject::ConstFuncRetStr2, (GInt16Arg), GUnscopedEnumArg);
 #endif
 
 		// 2 arg payload
-		BIND_CALL_CHECK_DELEGATE(FTestDynamicMulticastDelegate0::FDelegate, 0, &UDynamicDelegateTestObject::Func2,      (), GInt16Arg, GStrArg);
-		BIND_CALL_CHECK_DELEGATE(FTestDynamicDelegate0,                     0, &UDynamicDelegateTestObject::Func2,      (), GInt16Arg, GStrArg);
-		BIND_CALL_CHECK_DELEGATE(FTestDynamicMulticastDelegate0,            1, &UDynamicDelegateTestObject::Func2,      (), GInt16Arg, GStrArg);
-		BIND_CALL_CHECK_DELEGATE(FTestDynamicMulticastDelegate0::FDelegate, 0, &UDynamicDelegateTestObject::ConstFunc2, (), GInt16Arg, GStrArg);
-		BIND_CALL_CHECK_DELEGATE(FTestDynamicDelegate0,                     0, &UDynamicDelegateTestObject::ConstFunc2, (), GInt16Arg, GStrArg);
-		BIND_CALL_CHECK_DELEGATE(FTestDynamicMulticastDelegate0,            1, &UDynamicDelegateTestObject::ConstFunc2, (), GInt16Arg, GStrArg);
+		BIND_CALL_CHECK_DELEGATE(FTestDynamicMulticastDelegate0::FDelegate, 0, &UDynamicDelegateTestObject::Func2,      (), GInt16Arg, GUnscopedEnumArg);
+		BIND_CALL_CHECK_DELEGATE(FTestDynamicDelegate0,                     0, &UDynamicDelegateTestObject::Func2,      (), GInt16Arg, GUnscopedEnumArg);
+		BIND_CALL_CHECK_DELEGATE(FTestDynamicMulticastDelegate0,            1, &UDynamicDelegateTestObject::Func2,      (), GInt16Arg, GUnscopedEnumArg);
+		BIND_CALL_CHECK_DELEGATE(FTestDynamicMulticastDelegate0::FDelegate, 0, &UDynamicDelegateTestObject::ConstFunc2, (), GInt16Arg, GUnscopedEnumArg);
+		BIND_CALL_CHECK_DELEGATE(FTestDynamicDelegate0,                     0, &UDynamicDelegateTestObject::ConstFunc2, (), GInt16Arg, GUnscopedEnumArg);
+		BIND_CALL_CHECK_DELEGATE(FTestDynamicMulticastDelegate0,            1, &UDynamicDelegateTestObject::ConstFunc2, (), GInt16Arg, GUnscopedEnumArg);
 #if UE_USE_DYNAMIC_DELEGATE_PAYLOADS_RETVAL
-		BIND_CALL_CHECK_DELEGATE_RETVAL(GIntRet, FTestDynamicDelegateRetInt0, 0, &UDynamicDelegateTestObject::FuncRetInt2,      (), GInt16Arg, GStrArg);
-		BIND_CALL_CHECK_DELEGATE_RETVAL(GIntRet, FTestDynamicDelegateRetInt0, 0, &UDynamicDelegateTestObject::ConstFuncRetInt2, (), GInt16Arg, GStrArg);
-		BIND_CALL_CHECK_DELEGATE_RETVAL(GStrRet, FTestDynamicDelegateRetStr0, 0, &UDynamicDelegateTestObject::FuncRetStr2,      (), GInt16Arg, GStrArg);
-		BIND_CALL_CHECK_DELEGATE_RETVAL(GStrRet, FTestDynamicDelegateRetStr0, 0, &UDynamicDelegateTestObject::ConstFuncRetStr2, (), GInt16Arg, GStrArg);
+		BIND_CALL_CHECK_DELEGATE_RETVAL(GIntRet, FTestDynamicDelegateRetInt0, 0, &UDynamicDelegateTestObject::FuncRetInt2,      (), GInt16Arg, GUnscopedEnumArg);
+		BIND_CALL_CHECK_DELEGATE_RETVAL(GIntRet, FTestDynamicDelegateRetInt0, 0, &UDynamicDelegateTestObject::ConstFuncRetInt2, (), GInt16Arg, GUnscopedEnumArg);
+		BIND_CALL_CHECK_DELEGATE_RETVAL(GStrRet, FTestDynamicDelegateRetStr0, 0, &UDynamicDelegateTestObject::FuncRetStr2,      (), GInt16Arg, GUnscopedEnumArg);
+		BIND_CALL_CHECK_DELEGATE_RETVAL(GStrRet, FTestDynamicDelegateRetStr0, 0, &UDynamicDelegateTestObject::ConstFuncRetStr2, (), GInt16Arg, GUnscopedEnumArg);
 #endif
 #endif
 	}
@@ -623,59 +710,59 @@ void UDynamicDelegateTestObject::RunTests()
 	// 3 param function
 	{
 		// 0 arg payload
-		BIND_CALL_CHECK_DELEGATE(FTestDynamicMulticastDelegate3::FDelegate, 0, &UDynamicDelegateTestObject::Func3,      (GInt16Arg, GStrArg, GPtrArg));
-		BIND_CALL_CHECK_DELEGATE(FTestDynamicDelegate3,                     0, &UDynamicDelegateTestObject::Func3,      (GInt16Arg, GStrArg, GPtrArg));
-		BIND_CALL_CHECK_DELEGATE(FTestDynamicMulticastDelegate3,            1, &UDynamicDelegateTestObject::Func3,      (GInt16Arg, GStrArg, GPtrArg));
-		BIND_CALL_CHECK_DELEGATE(FTestDynamicMulticastDelegate3::FDelegate, 0, &UDynamicDelegateTestObject::ConstFunc3, (GInt16Arg, GStrArg, GPtrArg));
-		BIND_CALL_CHECK_DELEGATE(FTestDynamicDelegate3,                     0, &UDynamicDelegateTestObject::ConstFunc3, (GInt16Arg, GStrArg, GPtrArg));
-		BIND_CALL_CHECK_DELEGATE(FTestDynamicMulticastDelegate3,            1, &UDynamicDelegateTestObject::ConstFunc3, (GInt16Arg, GStrArg, GPtrArg));
-
-		BIND_CALL_CHECK_DELEGATE_RETVAL(GIntRet, FTestDynamicDelegateRetInt3, 0, &UDynamicDelegateTestObject::FuncRetInt3,      (GInt16Arg, GStrArg, GPtrArg));
-		BIND_CALL_CHECK_DELEGATE_RETVAL(GIntRet, FTestDynamicDelegateRetInt3, 0, &UDynamicDelegateTestObject::ConstFuncRetInt3, (GInt16Arg, GStrArg, GPtrArg));
-		BIND_CALL_CHECK_DELEGATE_RETVAL(GStrRet, FTestDynamicDelegateRetStr3, 0, &UDynamicDelegateTestObject::FuncRetStr3,      (GInt16Arg, GStrArg, GPtrArg));
-		BIND_CALL_CHECK_DELEGATE_RETVAL(GStrRet, FTestDynamicDelegateRetStr3, 0, &UDynamicDelegateTestObject::ConstFuncRetStr3, (GInt16Arg, GStrArg, GPtrArg));
+		BIND_CALL_CHECK_DELEGATE(FTestDynamicMulticastDelegate3::FDelegate, 0, &UDynamicDelegateTestObject::Func3,      (GInt16Arg, GUnscopedEnumArg, GStrArg));
+		BIND_CALL_CHECK_DELEGATE(FTestDynamicDelegate3,                     0, &UDynamicDelegateTestObject::Func3,      (GInt16Arg, GUnscopedEnumArg, GStrArg));
+		BIND_CALL_CHECK_DELEGATE(FTestDynamicMulticastDelegate3,            1, &UDynamicDelegateTestObject::Func3,      (GInt16Arg, GUnscopedEnumArg, GStrArg));
+		BIND_CALL_CHECK_DELEGATE(FTestDynamicMulticastDelegate3::FDelegate, 0, &UDynamicDelegateTestObject::ConstFunc3, (GInt16Arg, GUnscopedEnumArg, GStrArg));
+		BIND_CALL_CHECK_DELEGATE(FTestDynamicDelegate3,                     0, &UDynamicDelegateTestObject::ConstFunc3, (GInt16Arg, GUnscopedEnumArg, GStrArg));
+		BIND_CALL_CHECK_DELEGATE(FTestDynamicMulticastDelegate3,            1, &UDynamicDelegateTestObject::ConstFunc3, (GInt16Arg, GUnscopedEnumArg, GStrArg));
+
+		BIND_CALL_CHECK_DELEGATE_RETVAL(GIntRet, FTestDynamicDelegateRetInt3, 0, &UDynamicDelegateTestObject::FuncRetInt3,      (GInt16Arg, GUnscopedEnumArg, GStrArg));
+		BIND_CALL_CHECK_DELEGATE_RETVAL(GIntRet, FTestDynamicDelegateRetInt3, 0, &UDynamicDelegateTestObject::ConstFuncRetInt3, (GInt16Arg, GUnscopedEnumArg, GStrArg));
+		BIND_CALL_CHECK_DELEGATE_RETVAL(GStrRet, FTestDynamicDelegateRetStr3, 0, &UDynamicDelegateTestObject::FuncRetStr3,      (GInt16Arg, GUnscopedEnumArg, GStrArg));
+		BIND_CALL_CHECK_DELEGATE_RETVAL(GStrRet, FTestDynamicDelegateRetStr3, 0, &UDynamicDelegateTestObject::ConstFuncRetStr3, (GInt16Arg, GUnscopedEnumArg, GStrArg));
 
 #if defined(UE_USE_DYNAMIC_DELEGATE_PAYLOADS) && UE_USE_DYNAMIC_DELEGATE_PAYLOADS
 		// 1 arg payload
-		BIND_CALL_CHECK_DELEGATE(FTestDynamicMulticastDelegate2::FDelegate, 0, &UDynamicDelegateTestObject::Func3,      (GInt16Arg, GStrArg), GPtrArg);
-		BIND_CALL_CHECK_DELEGATE(FTestDynamicDelegate2,                     0, &UDynamicDelegateTestObject::Func3,      (GInt16Arg, GStrArg), GPtrArg);
-		BIND_CALL_CHECK_DELEGATE(FTestDynamicMulticastDelegate2,            1, &UDynamicDelegateTestObject::Func3,      (GInt16Arg, GStrArg), GPtrArg);
-		BIND_CALL_CHECK_DELEGATE(FTestDynamicMulticastDelegate2::FDelegate, 0, &UDynamicDelegateTestObject::ConstFunc3, (GInt16Arg, GStrArg), GPtrArg);
-		BIND_CALL_CHECK_DELEGATE(FTestDynamicDelegate2,                     0, &UDynamicDelegateTestObject::ConstFunc3, (GInt16Arg, GStrArg), GPtrArg);
-		BIND_CALL_CHECK_DELEGATE(FTestDynamicMulticastDelegate2,            1, &UDynamicDelegateTestObject::ConstFunc3, (GInt16Arg, GStrArg), GPtrArg);
+		BIND_CALL_CHECK_DELEGATE(FTestDynamicMulticastDelegate2::FDelegate, 0, &UDynamicDelegateTestObject::Func3,      (GInt16Arg, GUnscopedEnumArg), GStrArg);
+		BIND_CALL_CHECK_DELEGATE(FTestDynamicDelegate2,                     0, &UDynamicDelegateTestObject::Func3,      (GInt16Arg, GUnscopedEnumArg), GStrArg);
+		BIND_CALL_CHECK_DELEGATE(FTestDynamicMulticastDelegate2,            1, &UDynamicDelegateTestObject::Func3,      (GInt16Arg, GUnscopedEnumArg), GStrArg);
+		BIND_CALL_CHECK_DELEGATE(FTestDynamicMulticastDelegate2::FDelegate, 0, &UDynamicDelegateTestObject::ConstFunc3, (GInt16Arg, GUnscopedEnumArg), GStrArg);
+		BIND_CALL_CHECK_DELEGATE(FTestDynamicDelegate2,                     0, &UDynamicDelegateTestObject::ConstFunc3, (GInt16Arg, GUnscopedEnumArg), GStrArg);
+		BIND_CALL_CHECK_DELEGATE(FTestDynamicMulticastDelegate2,            1, &UDynamicDelegateTestObject::ConstFunc3, (GInt16Arg, GUnscopedEnumArg), GStrArg);
 #if UE_USE_DYNAMIC_DELEGATE_PAYLOADS_RETVAL
-		BIND_CALL_CHECK_DELEGATE_RETVAL(GIntRet, FTestDynamicDelegateRetInt2, 0, &UDynamicDelegateTestObject::FuncRetInt3,      (GInt16Arg, GStrArg), GPtrArg);
-		BIND_CALL_CHECK_DELEGATE_RETVAL(GIntRet, FTestDynamicDelegateRetInt2, 0, &UDynamicDelegateTestObject::ConstFuncRetInt3, (GInt16Arg, GStrArg), GPtrArg);
-		BIND_CALL_CHECK_DELEGATE_RETVAL(GStrRet, FTestDynamicDelegateRetStr2, 0, &UDynamicDelegateTestObject::FuncRetStr3,      (GInt16Arg, GStrArg), GPtrArg);
-		BIND_CALL_CHECK_DELEGATE_RETVAL(GStrRet, FTestDynamicDelegateRetStr2, 0, &UDynamicDelegateTestObject::ConstFuncRetStr3, (GInt16Arg, GStrArg), GPtrArg);
+		BIND_CALL_CHECK_DELEGATE_RETVAL(GIntRet, FTestDynamicDelegateRetInt2, 0, &UDynamicDelegateTestObject::FuncRetInt3,      (GInt16Arg, GUnscopedEnumArg), GStrArg);
+		BIND_CALL_CHECK_DELEGATE_RETVAL(GIntRet, FTestDynamicDelegateRetInt2, 0, &UDynamicDelegateTestObject::ConstFuncRetInt3, (GInt16Arg, GUnscopedEnumArg), GStrArg);
+		BIND_CALL_CHECK_DELEGATE_RETVAL(GStrRet, FTestDynamicDelegateRetStr2, 0, &UDynamicDelegateTestObject::FuncRetStr3,      (GInt16Arg, GUnscopedEnumArg), GStrArg);
+		BIND_CALL_CHECK_DELEGATE_RETVAL(GStrRet, FTestDynamicDelegateRetStr2, 0, &UDynamicDelegateTestObject::ConstFuncRetStr3, (GInt16Arg, GUnscopedEnumArg), GStrArg);
 #endif
 
 		// 2 arg payload
-		BIND_CALL_CHECK_DELEGATE(FTestDynamicMulticastDelegate1::FDelegate, 0, &UDynamicDelegateTestObject::Func3,      (GInt16Arg), GStrArg, GPtrArg);
-		BIND_CALL_CHECK_DELEGATE(FTestDynamicDelegate1,                     0, &UDynamicDelegateTestObject::Func3,      (GInt16Arg), GStrArg, GPtrArg);
-		BIND_CALL_CHECK_DELEGATE(FTestDynamicMulticastDelegate1,            1, &UDynamicDelegateTestObject::Func3,      (GInt16Arg), GStrArg, GPtrArg);
-		BIND_CALL_CHECK_DELEGATE(FTestDynamicMulticastDelegate1::FDelegate, 0, &UDynamicDelegateTestObject::ConstFunc3, (GInt16Arg), GStrArg, GPtrArg);
-		BIND_CALL_CHECK_DELEGATE(FTestDynamicDelegate1,                     0, &UDynamicDelegateTestObject::ConstFunc3, (GInt16Arg), GStrArg, GPtrArg);
-		BIND_CALL_CHECK_DELEGATE(FTestDynamicMulticastDelegate1,            1, &UDynamicDelegateTestObject::ConstFunc3, (GInt16Arg), GStrArg, GPtrArg);
+		BIND_CALL_CHECK_DELEGATE(FTestDynamicMulticastDelegate1::FDelegate, 0, &UDynamicDelegateTestObject::Func3,      (GInt16Arg), GUnscopedEnumArg, GStrArg);
+		BIND_CALL_CHECK_DELEGATE(FTestDynamicDelegate1,                     0, &UDynamicDelegateTestObject::Func3,      (GInt16Arg), GUnscopedEnumArg, GStrArg);
+		BIND_CALL_CHECK_DELEGATE(FTestDynamicMulticastDelegate1,            1, &UDynamicDelegateTestObject::Func3,      (GInt16Arg), GUnscopedEnumArg, GStrArg);
+		BIND_CALL_CHECK_DELEGATE(FTestDynamicMulticastDelegate1::FDelegate, 0, &UDynamicDelegateTestObject::ConstFunc3, (GInt16Arg), GUnscopedEnumArg, GStrArg);
+		BIND_CALL_CHECK_DELEGATE(FTestDynamicDelegate1,                     0, &UDynamicDelegateTestObject::ConstFunc3, (GInt16Arg), GUnscopedEnumArg, GStrArg);
+		BIND_CALL_CHECK_DELEGATE(FTestDynamicMulticastDelegate1,            1, &UDynamicDelegateTestObject::ConstFunc3, (GInt16Arg), GUnscopedEnumArg, GStrArg);
 #if UE_USE_DYNAMIC_DELEGATE_PAYLOADS_RETVAL
-		BIND_CALL_CHECK_DELEGATE_RETVAL(GIntRet, FTestDynamicDelegateRetInt1, 0, &UDynamicDelegateTestObject::FuncRetInt3,      (GInt16Arg), GStrArg, GPtrArg);
-		BIND_CALL_CHECK_DELEGATE_RETVAL(GIntRet, FTestDynamicDelegateRetInt1, 0, &UDynamicDelegateTestObject::ConstFuncRetInt3, (GInt16Arg), GStrArg, GPtrArg);
-		BIND_CALL_CHECK_DELEGATE_RETVAL(GStrRet, FTestDynamicDelegateRetStr1, 0, &UDynamicDelegateTestObject::FuncRetStr3,      (GInt16Arg), GStrArg, GPtrArg);
-		BIND_CALL_CHECK_DELEGATE_RETVAL(GStrRet, FTestDynamicDelegateRetStr1, 0, &UDynamicDelegateTestObject::ConstFuncRetStr3, (GInt16Arg), GStrArg, GPtrArg);
+		BIND_CALL_CHECK_DELEGATE_RETVAL(GIntRet, FTestDynamicDelegateRetInt1, 0, &UDynamicDelegateTestObject::FuncRetInt3,      (GInt16Arg), GUnscopedEnumArg, GStrArg);
+		BIND_CALL_CHECK_DELEGATE_RETVAL(GIntRet, FTestDynamicDelegateRetInt1, 0, &UDynamicDelegateTestObject::ConstFuncRetInt3, (GInt16Arg), GUnscopedEnumArg, GStrArg);
+		BIND_CALL_CHECK_DELEGATE_RETVAL(GStrRet, FTestDynamicDelegateRetStr1, 0, &UDynamicDelegateTestObject::FuncRetStr3,      (GInt16Arg), GUnscopedEnumArg, GStrArg);
+		BIND_CALL_CHECK_DELEGATE_RETVAL(GStrRet, FTestDynamicDelegateRetStr1, 0, &UDynamicDelegateTestObject::ConstFuncRetStr3, (GInt16Arg), GUnscopedEnumArg, GStrArg);
 #endif
 
 		// 3 arg payload
-		BIND_CALL_CHECK_DELEGATE(FTestDynamicMulticastDelegate0::FDelegate, 0, &UDynamicDelegateTestObject::Func3,      (), GInt16Arg, GStrArg, GPtrArg);
-		BIND_CALL_CHECK_DELEGATE(FTestDynamicDelegate0,                     0, &UDynamicDelegateTestObject::Func3,      (), GInt16Arg, GStrArg, GPtrArg);
-		BIND_CALL_CHECK_DELEGATE(FTestDynamicMulticastDelegate0,            1, &UDynamicDelegateTestObject::Func3,      (), GInt16Arg, GStrArg, GPtrArg);
-		BIND_CALL_CHECK_DELEGATE(FTestDynamicMulticastDelegate0::FDelegate, 0, &UDynamicDelegateTestObject::ConstFunc3, (), GInt16Arg, GStrArg, GPtrArg);
-		BIND_CALL_CHECK_DELEGATE(FTestDynamicDelegate0,                     0, &UDynamicDelegateTestObject::ConstFunc3, (), GInt16Arg, GStrArg, GPtrArg);
-		BIND_CALL_CHECK_DELEGATE(FTestDynamicMulticastDelegate0,            1, &UDynamicDelegateTestObject::ConstFunc3, (), GInt16Arg, GStrArg, GPtrArg);
+		BIND_CALL_CHECK_DELEGATE(FTestDynamicMulticastDelegate0::FDelegate, 0, &UDynamicDelegateTestObject::Func3,      (), GInt16Arg, GUnscopedEnumArg, GStrArg);
+		BIND_CALL_CHECK_DELEGATE(FTestDynamicDelegate0,                     0, &UDynamicDelegateTestObject::Func3,      (), GInt16Arg, GUnscopedEnumArg, GStrArg);
+		BIND_CALL_CHECK_DELEGATE(FTestDynamicMulticastDelegate0,            1, &UDynamicDelegateTestObject::Func3,      (), GInt16Arg, GUnscopedEnumArg, GStrArg);
+		BIND_CALL_CHECK_DELEGATE(FTestDynamicMulticastDelegate0::FDelegate, 0, &UDynamicDelegateTestObject::ConstFunc3, (), GInt16Arg, GUnscopedEnumArg, GStrArg);
+		BIND_CALL_CHECK_DELEGATE(FTestDynamicDelegate0,                     0, &UDynamicDelegateTestObject::ConstFunc3, (), GInt16Arg, GUnscopedEnumArg, GStrArg);
+		BIND_CALL_CHECK_DELEGATE(FTestDynamicMulticastDelegate0,            1, &UDynamicDelegateTestObject::ConstFunc3, (), GInt16Arg, GUnscopedEnumArg, GStrArg);
 #if UE_USE_DYNAMIC_DELEGATE_PAYLOADS_RETVAL
-		BIND_CALL_CHECK_DELEGATE_RETVAL(GIntRet, FTestDynamicDelegateRetInt0, 0, &UDynamicDelegateTestObject::FuncRetInt3,      (), GInt16Arg, GStrArg, GPtrArg);
-		BIND_CALL_CHECK_DELEGATE_RETVAL(GIntRet, FTestDynamicDelegateRetInt0, 0, &UDynamicDelegateTestObject::ConstFuncRetInt3, (), GInt16Arg, GStrArg, GPtrArg);
-		BIND_CALL_CHECK_DELEGATE_RETVAL(GStrRet, FTestDynamicDelegateRetStr0, 0, &UDynamicDelegateTestObject::FuncRetStr3,      (), GInt16Arg, GStrArg, GPtrArg);
-		BIND_CALL_CHECK_DELEGATE_RETVAL(GStrRet, FTestDynamicDelegateRetStr0, 0, &UDynamicDelegateTestObject::ConstFuncRetStr3, (), GInt16Arg, GStrArg, GPtrArg);
+		BIND_CALL_CHECK_DELEGATE_RETVAL(GIntRet, FTestDynamicDelegateRetInt0, 0, &UDynamicDelegateTestObject::FuncRetInt3,      (), GInt16Arg, GUnscopedEnumArg, GStrArg);
+		BIND_CALL_CHECK_DELEGATE_RETVAL(GIntRet, FTestDynamicDelegateRetInt0, 0, &UDynamicDelegateTestObject::ConstFuncRetInt3, (), GInt16Arg, GUnscopedEnumArg, GStrArg);
+		BIND_CALL_CHECK_DELEGATE_RETVAL(GStrRet, FTestDynamicDelegateRetStr0, 0, &UDynamicDelegateTestObject::FuncRetStr3,      (), GInt16Arg, GUnscopedEnumArg, GStrArg);
+		BIND_CALL_CHECK_DELEGATE_RETVAL(GStrRet, FTestDynamicDelegateRetStr0, 0, &UDynamicDelegateTestObject::ConstFuncRetStr3, (), GInt16Arg, GUnscopedEnumArg, GStrArg);
 #endif
 #endif
 	}
@@ -683,73 +770,73 @@ void UDynamicDelegateTestObject::RunTests()
 	// 4 param function
 	{
 		// 0 arg payload
-		BIND_CALL_CHECK_DELEGATE(FTestDynamicMulticastDelegate4::FDelegate, 0, &UDynamicDelegateTestObject::Func4,      (GInt16Arg, GStrArg, GPtrArg, GNameArg));
-		BIND_CALL_CHECK_DELEGATE(FTestDynamicDelegate4,                     0, &UDynamicDelegateTestObject::Func4,      (GInt16Arg, GStrArg, GPtrArg, GNameArg));
-		BIND_CALL_CHECK_DELEGATE(FTestDynamicMulticastDelegate4,            1, &UDynamicDelegateTestObject::Func4,      (GInt16Arg, GStrArg, GPtrArg, GNameArg));
-		BIND_CALL_CHECK_DELEGATE(FTestDynamicMulticastDelegate4::FDelegate, 0, &UDynamicDelegateTestObject::ConstFunc4, (GInt16Arg, GStrArg, GPtrArg, GNameArg));
-		BIND_CALL_CHECK_DELEGATE(FTestDynamicDelegate4,                     0, &UDynamicDelegateTestObject::ConstFunc4, (GInt16Arg, GStrArg, GPtrArg, GNameArg));
-		BIND_CALL_CHECK_DELEGATE(FTestDynamicMulticastDelegate4,            1, &UDynamicDelegateTestObject::ConstFunc4, (GInt16Arg, GStrArg, GPtrArg, GNameArg));
-
-		BIND_CALL_CHECK_DELEGATE_RETVAL(GIntRet, FTestDynamicDelegateRetInt4, 0, &UDynamicDelegateTestObject::FuncRetInt4,      (GInt16Arg, GStrArg, GPtrArg, GNameArg));
-		BIND_CALL_CHECK_DELEGATE_RETVAL(GIntRet, FTestDynamicDelegateRetInt4, 0, &UDynamicDelegateTestObject::ConstFuncRetInt4, (GInt16Arg, GStrArg, GPtrArg, GNameArg));
-		BIND_CALL_CHECK_DELEGATE_RETVAL(GStrRet, FTestDynamicDelegateRetStr4, 0, &UDynamicDelegateTestObject::FuncRetStr4,      (GInt16Arg, GStrArg, GPtrArg, GNameArg));
-		BIND_CALL_CHECK_DELEGATE_RETVAL(GStrRet, FTestDynamicDelegateRetStr4, 0, &UDynamicDelegateTestObject::ConstFuncRetStr4, (GInt16Arg, GStrArg, GPtrArg, GNameArg));
+		BIND_CALL_CHECK_DELEGATE(FTestDynamicMulticastDelegate4::FDelegate, 0, &UDynamicDelegateTestObject::Func4,      (GInt16Arg, GUnscopedEnumArg, GStrArg, GPtrArg));
+		BIND_CALL_CHECK_DELEGATE(FTestDynamicDelegate4,                     0, &UDynamicDelegateTestObject::Func4,      (GInt16Arg, GUnscopedEnumArg, GStrArg, GPtrArg));
+		BIND_CALL_CHECK_DELEGATE(FTestDynamicMulticastDelegate4,            1, &UDynamicDelegateTestObject::Func4,      (GInt16Arg, GUnscopedEnumArg, GStrArg, GPtrArg));
+		BIND_CALL_CHECK_DELEGATE(FTestDynamicMulticastDelegate4::FDelegate, 0, &UDynamicDelegateTestObject::ConstFunc4, (GInt16Arg, GUnscopedEnumArg, GStrArg, GPtrArg));
+		BIND_CALL_CHECK_DELEGATE(FTestDynamicDelegate4,                     0, &UDynamicDelegateTestObject::ConstFunc4, (GInt16Arg, GUnscopedEnumArg, GStrArg, GPtrArg));
+		BIND_CALL_CHECK_DELEGATE(FTestDynamicMulticastDelegate4,            1, &UDynamicDelegateTestObject::ConstFunc4, (GInt16Arg, GUnscopedEnumArg, GStrArg, GPtrArg));
+
+		BIND_CALL_CHECK_DELEGATE_RETVAL(GIntRet, FTestDynamicDelegateRetInt4, 0, &UDynamicDelegateTestObject::FuncRetInt4,      (GInt16Arg, GUnscopedEnumArg, GStrArg, GPtrArg));
+		BIND_CALL_CHECK_DELEGATE_RETVAL(GIntRet, FTestDynamicDelegateRetInt4, 0, &UDynamicDelegateTestObject::ConstFuncRetInt4, (GInt16Arg, GUnscopedEnumArg, GStrArg, GPtrArg));
+		BIND_CALL_CHECK_DELEGATE_RETVAL(GStrRet, FTestDynamicDelegateRetStr4, 0, &UDynamicDelegateTestObject::FuncRetStr4,      (GInt16Arg, GUnscopedEnumArg, GStrArg, GPtrArg));
+		BIND_CALL_CHECK_DELEGATE_RETVAL(GStrRet, FTestDynamicDelegateRetStr4, 0, &UDynamicDelegateTestObject::ConstFuncRetStr4, (GInt16Arg, GUnscopedEnumArg, GStrArg, GPtrArg));
 
 #if defined(UE_USE_DYNAMIC_DELEGATE_PAYLOADS) && UE_USE_DYNAMIC_DELEGATE_PAYLOADS
 		// 1 arg payload
-		BIND_CALL_CHECK_DELEGATE(FTestDynamicMulticastDelegate3::FDelegate, 0, &UDynamicDelegateTestObject::Func4,      (GInt16Arg, GStrArg, GPtrArg), GNameArg);
-		BIND_CALL_CHECK_DELEGATE(FTestDynamicDelegate3,                     0, &UDynamicDelegateTestObject::Func4,      (GInt16Arg, GStrArg, GPtrArg), GNameArg);
-		BIND_CALL_CHECK_DELEGATE(FTestDynamicMulticastDelegate3,            1, &UDynamicDelegateTestObject::Func4,      (GInt16Arg, GStrArg, GPtrArg), GNameArg);
-		BIND_CALL_CHECK_DELEGATE(FTestDynamicMulticastDelegate3::FDelegate, 0, &UDynamicDelegateTestObject::ConstFunc4, (GInt16Arg, GStrArg, GPtrArg), GNameArg);
-		BIND_CALL_CHECK_DELEGATE(FTestDynamicDelegate3,                     0, &UDynamicDelegateTestObject::ConstFunc4, (GInt16Arg, GStrArg, GPtrArg), GNameArg);
-		BIND_CALL_CHECK_DELEGATE(FTestDynamicMulticastDelegate3,            1, &UDynamicDelegateTestObject::ConstFunc4, (GInt16Arg, GStrArg, GPtrArg), GNameArg);
+		BIND_CALL_CHECK_DELEGATE(FTestDynamicMulticastDelegate3::FDelegate, 0, &UDynamicDelegateTestObject::Func4,      (GInt16Arg, GUnscopedEnumArg, GStrArg), GPtrArg);
+		BIND_CALL_CHECK_DELEGATE(FTestDynamicDelegate3,                     0, &UDynamicDelegateTestObject::Func4,      (GInt16Arg, GUnscopedEnumArg, GStrArg), GPtrArg);
+		BIND_CALL_CHECK_DELEGATE(FTestDynamicMulticastDelegate3,            1, &UDynamicDelegateTestObject::Func4,      (GInt16Arg, GUnscopedEnumArg, GStrArg), GPtrArg);
+		BIND_CALL_CHECK_DELEGATE(FTestDynamicMulticastDelegate3::FDelegate, 0, &UDynamicDelegateTestObject::ConstFunc4, (GInt16Arg, GUnscopedEnumArg, GStrArg), GPtrArg);
+		BIND_CALL_CHECK_DELEGATE(FTestDynamicDelegate3,                     0, &UDynamicDelegateTestObject::ConstFunc4, (GInt16Arg, GUnscopedEnumArg, GStrArg), GPtrArg);
+		BIND_CALL_CHECK_DELEGATE(FTestDynamicMulticastDelegate3,            1, &UDynamicDelegateTestObject::ConstFunc4, (GInt16Arg, GUnscopedEnumArg, GStrArg), GPtrArg);
 #if UE_USE_DYNAMIC_DELEGATE_PAYLOADS_RETVAL
-		BIND_CALL_CHECK_DELEGATE_RETVAL(GIntRet, FTestDynamicDelegateRetInt3, 0, &UDynamicDelegateTestObject::FuncRetInt4,      (GInt16Arg, GStrArg, GPtrArg), GNameArg);
-		BIND_CALL_CHECK_DELEGATE_RETVAL(GIntRet, FTestDynamicDelegateRetInt3, 0, &UDynamicDelegateTestObject::ConstFuncRetInt4, (GInt16Arg, GStrArg, GPtrArg), GNameArg);
-		BIND_CALL_CHECK_DELEGATE_RETVAL(GStrRet, FTestDynamicDelegateRetStr3, 0, &UDynamicDelegateTestObject::FuncRetStr4,      (GInt16Arg, GStrArg, GPtrArg), GNameArg);
-		BIND_CALL_CHECK_DELEGATE_RETVAL(GStrRet, FTestDynamicDelegateRetStr3, 0, &UDynamicDelegateTestObject::ConstFuncRetStr4, (GInt16Arg, GStrArg, GPtrArg), GNameArg);
+		BIND_CALL_CHECK_DELEGATE_RETVAL(GIntRet, FTestDynamicDelegateRetInt3, 0, &UDynamicDelegateTestObject::FuncRetInt4,      (GInt16Arg, GUnscopedEnumArg, GStrArg), GPtrArg);
+		BIND_CALL_CHECK_DELEGATE_RETVAL(GIntRet, FTestDynamicDelegateRetInt3, 0, &UDynamicDelegateTestObject::ConstFuncRetInt4, (GInt16Arg, GUnscopedEnumArg, GStrArg), GPtrArg);
+		BIND_CALL_CHECK_DELEGATE_RETVAL(GStrRet, FTestDynamicDelegateRetStr3, 0, &UDynamicDelegateTestObject::FuncRetStr4,      (GInt16Arg, GUnscopedEnumArg, GStrArg), GPtrArg);
+		BIND_CALL_CHECK_DELEGATE_RETVAL(GStrRet, FTestDynamicDelegateRetStr3, 0, &UDynamicDelegateTestObject::ConstFuncRetStr4, (GInt16Arg, GUnscopedEnumArg, GStrArg), GPtrArg);
 #endif
 
 		// 2 arg payload
-		BIND_CALL_CHECK_DELEGATE(FTestDynamicMulticastDelegate2::FDelegate, 0, &UDynamicDelegateTestObject::Func4,      (GInt16Arg, GStrArg), GPtrArg, GNameArg);
-		BIND_CALL_CHECK_DELEGATE(FTestDynamicDelegate2,                     0, &UDynamicDelegateTestObject::Func4,      (GInt16Arg, GStrArg), GPtrArg, GNameArg);
-		BIND_CALL_CHECK_DELEGATE(FTestDynamicMulticastDelegate2,            1, &UDynamicDelegateTestObject::Func4,      (GInt16Arg, GStrArg), GPtrArg, GNameArg);
-		BIND_CALL_CHECK_DELEGATE(FTestDynamicMulticastDelegate2::FDelegate, 0, &UDynamicDelegateTestObject::ConstFunc4, (GInt16Arg, GStrArg), GPtrArg, GNameArg);
-		BIND_CALL_CHECK_DELEGATE(FTestDynamicDelegate2,                     0, &UDynamicDelegateTestObject::ConstFunc4, (GInt16Arg, GStrArg), GPtrArg, GNameArg);
-		BIND_CALL_CHECK_DELEGATE(FTestDynamicMulticastDelegate2,            1, &UDynamicDelegateTestObject::ConstFunc4, (GInt16Arg, GStrArg), GPtrArg, GNameArg);
+		BIND_CALL_CHECK_DELEGATE(FTestDynamicMulticastDelegate2::FDelegate, 0, &UDynamicDelegateTestObject::Func4,      (GInt16Arg, GUnscopedEnumArg), GStrArg, GPtrArg);
+		BIND_CALL_CHECK_DELEGATE(FTestDynamicDelegate2,                     0, &UDynamicDelegateTestObject::Func4,      (GInt16Arg, GUnscopedEnumArg), GStrArg, GPtrArg);
+		BIND_CALL_CHECK_DELEGATE(FTestDynamicMulticastDelegate2,            1, &UDynamicDelegateTestObject::Func4,      (GInt16Arg, GUnscopedEnumArg), GStrArg, GPtrArg);
+		BIND_CALL_CHECK_DELEGATE(FTestDynamicMulticastDelegate2::FDelegate, 0, &UDynamicDelegateTestObject::ConstFunc4, (GInt16Arg, GUnscopedEnumArg), GStrArg, GPtrArg);
+		BIND_CALL_CHECK_DELEGATE(FTestDynamicDelegate2,                     0, &UDynamicDelegateTestObject::ConstFunc4, (GInt16Arg, GUnscopedEnumArg), GStrArg, GPtrArg);
+		BIND_CALL_CHECK_DELEGATE(FTestDynamicMulticastDelegate2,            1, &UDynamicDelegateTestObject::ConstFunc4, (GInt16Arg, GUnscopedEnumArg), GStrArg, GPtrArg);
 #if UE_USE_DYNAMIC_DELEGATE_PAYLOADS_RETVAL
-		BIND_CALL_CHECK_DELEGATE_RETVAL(GIntRet, FTestDynamicDelegateRetInt2, 0, &UDynamicDelegateTestObject::FuncRetInt4,      (GInt16Arg, GStrArg), GPtrArg, GNameArg);
-		BIND_CALL_CHECK_DELEGATE_RETVAL(GIntRet, FTestDynamicDelegateRetInt2, 0, &UDynamicDelegateTestObject::ConstFuncRetInt4, (GInt16Arg, GStrArg), GPtrArg, GNameArg);
-		BIND_CALL_CHECK_DELEGATE_RETVAL(GStrRet, FTestDynamicDelegateRetStr2, 0, &UDynamicDelegateTestObject::FuncRetStr4,      (GInt16Arg, GStrArg), GPtrArg, GNameArg);
-		BIND_CALL_CHECK_DELEGATE_RETVAL(GStrRet, FTestDynamicDelegateRetStr2, 0, &UDynamicDelegateTestObject::ConstFuncRetStr4, (GInt16Arg, GStrArg), GPtrArg, GNameArg);
+		BIND_CALL_CHECK_DELEGATE_RETVAL(GIntRet, FTestDynamicDelegateRetInt2, 0, &UDynamicDelegateTestObject::FuncRetInt4,      (GInt16Arg, GUnscopedEnumArg), GStrArg, GPtrArg);
+		BIND_CALL_CHECK_DELEGATE_RETVAL(GIntRet, FTestDynamicDelegateRetInt2, 0, &UDynamicDelegateTestObject::ConstFuncRetInt4, (GInt16Arg, GUnscopedEnumArg), GStrArg, GPtrArg);
+		BIND_CALL_CHECK_DELEGATE_RETVAL(GStrRet, FTestDynamicDelegateRetStr2, 0, &UDynamicDelegateTestObject::FuncRetStr4,      (GInt16Arg, GUnscopedEnumArg), GStrArg, GPtrArg);
+		BIND_CALL_CHECK_DELEGATE_RETVAL(GStrRet, FTestDynamicDelegateRetStr2, 0, &UDynamicDelegateTestObject::ConstFuncRetStr4, (GInt16Arg, GUnscopedEnumArg), GStrArg, GPtrArg);
 #endif
 
 		// 3 arg payload
-		BIND_CALL_CHECK_DELEGATE(FTestDynamicMulticastDelegate1::FDelegate, 0, &UDynamicDelegateTestObject::Func4,      (GInt16Arg), GStrArg, GPtrArg, GNameArg);
-		BIND_CALL_CHECK_DELEGATE(FTestDynamicDelegate1,                     0, &UDynamicDelegateTestObject::Func4,      (GInt16Arg), GStrArg, GPtrArg, GNameArg);
-		BIND_CALL_CHECK_DELEGATE(FTestDynamicMulticastDelegate1,            1, &UDynamicDelegateTestObject::Func4,      (GInt16Arg), GStrArg, GPtrArg, GNameArg);
-		BIND_CALL_CHECK_DELEGATE(FTestDynamicMulticastDelegate1::FDelegate, 0, &UDynamicDelegateTestObject::ConstFunc4, (GInt16Arg), GStrArg, GPtrArg, GNameArg);
-		BIND_CALL_CHECK_DELEGATE(FTestDynamicDelegate1,                     0, &UDynamicDelegateTestObject::ConstFunc4, (GInt16Arg), GStrArg, GPtrArg, GNameArg);
-		BIND_CALL_CHECK_DELEGATE(FTestDynamicMulticastDelegate1,            1, &UDynamicDelegateTestObject::ConstFunc4, (GInt16Arg), GStrArg, GPtrArg, GNameArg);
+		BIND_CALL_CHECK_DELEGATE(FTestDynamicMulticastDelegate1::FDelegate, 0, &UDynamicDelegateTestObject::Func4,      (GInt16Arg), GUnscopedEnumArg, GStrArg, GPtrArg);
+		BIND_CALL_CHECK_DELEGATE(FTestDynamicDelegate1,                     0, &UDynamicDelegateTestObject::Func4,      (GInt16Arg), GUnscopedEnumArg, GStrArg, GPtrArg);
+		BIND_CALL_CHECK_DELEGATE(FTestDynamicMulticastDelegate1,            1, &UDynamicDelegateTestObject::Func4,      (GInt16Arg), GUnscopedEnumArg, GStrArg, GPtrArg);
+		BIND_CALL_CHECK_DELEGATE(FTestDynamicMulticastDelegate1::FDelegate, 0, &UDynamicDelegateTestObject::ConstFunc4, (GInt16Arg), GUnscopedEnumArg, GStrArg, GPtrArg);
+		BIND_CALL_CHECK_DELEGATE(FTestDynamicDelegate1,                     0, &UDynamicDelegateTestObject::ConstFunc4, (GInt16Arg), GUnscopedEnumArg, GStrArg, GPtrArg);
+		BIND_CALL_CHECK_DELEGATE(FTestDynamicMulticastDelegate1,            1, &UDynamicDelegateTestObject::ConstFunc4, (GInt16Arg), GUnscopedEnumArg, GStrArg, GPtrArg);
 #if UE_USE_DYNAMIC_DELEGATE_PAYLOADS_RETVAL
-		BIND_CALL_CHECK_DELEGATE_RETVAL(GIntRet, FTestDynamicDelegateRetInt1, 0, &UDynamicDelegateTestObject::FuncRetInt4,      (GInt16Arg), GStrArg, GPtrArg, GNameArg);
-		BIND_CALL_CHECK_DELEGATE_RETVAL(GIntRet, FTestDynamicDelegateRetInt1, 0, &UDynamicDelegateTestObject::ConstFuncRetInt4, (GInt16Arg), GStrArg, GPtrArg, GNameArg);
-		BIND_CALL_CHECK_DELEGATE_RETVAL(GStrRet, FTestDynamicDelegateRetStr1, 0, &UDynamicDelegateTestObject::FuncRetStr4,      (GInt16Arg), GStrArg, GPtrArg, GNameArg);
-		BIND_CALL_CHECK_DELEGATE_RETVAL(GStrRet, FTestDynamicDelegateRetStr1, 0, &UDynamicDelegateTestObject::ConstFuncRetStr4, (GInt16Arg), GStrArg, GPtrArg, GNameArg);
+		BIND_CALL_CHECK_DELEGATE_RETVAL(GIntRet, FTestDynamicDelegateRetInt1, 0, &UDynamicDelegateTestObject::FuncRetInt4,      (GInt16Arg), GUnscopedEnumArg, GStrArg, GPtrArg);
+		BIND_CALL_CHECK_DELEGATE_RETVAL(GIntRet, FTestDynamicDelegateRetInt1, 0, &UDynamicDelegateTestObject::ConstFuncRetInt4, (GInt16Arg), GUnscopedEnumArg, GStrArg, GPtrArg);
+		BIND_CALL_CHECK_DELEGATE_RETVAL(GStrRet, FTestDynamicDelegateRetStr1, 0, &UDynamicDelegateTestObject::FuncRetStr4,      (GInt16Arg), GUnscopedEnumArg, GStrArg, GPtrArg);
+		BIND_CALL_CHECK_DELEGATE_RETVAL(GStrRet, FTestDynamicDelegateRetStr1, 0, &UDynamicDelegateTestObject::ConstFuncRetStr4, (GInt16Arg), GUnscopedEnumArg, GStrArg, GPtrArg);
 #endif
 
 		// 4 arg payload
-		BIND_CALL_CHECK_DELEGATE(FTestDynamicMulticastDelegate0::FDelegate, 0, &UDynamicDelegateTestObject::Func4,      (), GInt16Arg, GStrArg, GPtrArg, GNameArg);
-		BIND_CALL_CHECK_DELEGATE(FTestDynamicDelegate0,                     0, &UDynamicDelegateTestObject::Func4,      (), GInt16Arg, GStrArg, GPtrArg, GNameArg);
-		BIND_CALL_CHECK_DELEGATE(FTestDynamicMulticastDelegate0,            1, &UDynamicDelegateTestObject::Func4,      (), GInt16Arg, GStrArg, GPtrArg, GNameArg);
-		BIND_CALL_CHECK_DELEGATE(FTestDynamicMulticastDelegate0::FDelegate, 0, &UDynamicDelegateTestObject::ConstFunc4, (), GInt16Arg, GStrArg, GPtrArg, GNameArg);
-		BIND_CALL_CHECK_DELEGATE(FTestDynamicDelegate0,                     0, &UDynamicDelegateTestObject::ConstFunc4, (), GInt16Arg, GStrArg, GPtrArg, GNameArg);
-		BIND_CALL_CHECK_DELEGATE(FTestDynamicMulticastDelegate0,            1, &UDynamicDelegateTestObject::ConstFunc4, (), GInt16Arg, GStrArg, GPtrArg, GNameArg);
+		BIND_CALL_CHECK_DELEGATE(FTestDynamicMulticastDelegate0::FDelegate, 0, &UDynamicDelegateTestObject::Func4,      (), GInt16Arg, GUnscopedEnumArg, GStrArg, GPtrArg);
+		BIND_CALL_CHECK_DELEGATE(FTestDynamicDelegate0,                     0, &UDynamicDelegateTestObject::Func4,      (), GInt16Arg, GUnscopedEnumArg, GStrArg, GPtrArg);
+		BIND_CALL_CHECK_DELEGATE(FTestDynamicMulticastDelegate0,            1, &UDynamicDelegateTestObject::Func4,      (), GInt16Arg, GUnscopedEnumArg, GStrArg, GPtrArg);
+		BIND_CALL_CHECK_DELEGATE(FTestDynamicMulticastDelegate0::FDelegate, 0, &UDynamicDelegateTestObject::ConstFunc4, (), GInt16Arg, GUnscopedEnumArg, GStrArg, GPtrArg);
+		BIND_CALL_CHECK_DELEGATE(FTestDynamicDelegate0,                     0, &UDynamicDelegateTestObject::ConstFunc4, (), GInt16Arg, GUnscopedEnumArg, GStrArg, GPtrArg);
+		BIND_CALL_CHECK_DELEGATE(FTestDynamicMulticastDelegate0,            1, &UDynamicDelegateTestObject::ConstFunc4, (), GInt16Arg, GUnscopedEnumArg, GStrArg, GPtrArg);
 #if UE_USE_DYNAMIC_DELEGATE_PAYLOADS_RETVAL
-		BIND_CALL_CHECK_DELEGATE_RETVAL(GIntRet, FTestDynamicDelegateRetInt0, 0, &UDynamicDelegateTestObject::FuncRetInt4,      (), GInt16Arg, GStrArg, GPtrArg, GNameArg);
-		BIND_CALL_CHECK_DELEGATE_RETVAL(GIntRet, FTestDynamicDelegateRetInt0, 0, &UDynamicDelegateTestObject::ConstFuncRetInt4, (), GInt16Arg, GStrArg, GPtrArg, GNameArg);
-		BIND_CALL_CHECK_DELEGATE_RETVAL(GStrRet, FTestDynamicDelegateRetStr0, 0, &UDynamicDelegateTestObject::FuncRetStr4,      (), GInt16Arg, GStrArg, GPtrArg, GNameArg);
-		BIND_CALL_CHECK_DELEGATE_RETVAL(GStrRet, FTestDynamicDelegateRetStr0, 0, &UDynamicDelegateTestObject::ConstFuncRetStr4, (), GInt16Arg, GStrArg, GPtrArg, GNameArg);
+		BIND_CALL_CHECK_DELEGATE_RETVAL(GIntRet, FTestDynamicDelegateRetInt0, 0, &UDynamicDelegateTestObject::FuncRetInt4,      (), GInt16Arg, GUnscopedEnumArg, GStrArg, GPtrArg);
+		BIND_CALL_CHECK_DELEGATE_RETVAL(GIntRet, FTestDynamicDelegateRetInt0, 0, &UDynamicDelegateTestObject::ConstFuncRetInt4, (), GInt16Arg, GUnscopedEnumArg, GStrArg, GPtrArg);
+		BIND_CALL_CHECK_DELEGATE_RETVAL(GStrRet, FTestDynamicDelegateRetStr0, 0, &UDynamicDelegateTestObject::FuncRetStr4,      (), GInt16Arg, GUnscopedEnumArg, GStrArg, GPtrArg);
+		BIND_CALL_CHECK_DELEGATE_RETVAL(GStrRet, FTestDynamicDelegateRetStr0, 0, &UDynamicDelegateTestObject::ConstFuncRetStr4, (), GInt16Arg, GUnscopedEnumArg, GStrArg, GPtrArg);
 #endif
 #endif
 	}
@@ -757,87 +844,87 @@ void UDynamicDelegateTestObject::RunTests()
 	// 5 param function
 	{
 		// 0 arg payload
-		BIND_CALL_CHECK_DELEGATE(FTestDynamicMulticastDelegate5::FDelegate, 0, &UDynamicDelegateTestObject::Func5,      (GInt16Arg, GStrArg, GPtrArg, GNameArg, bGBoolArg));
-		BIND_CALL_CHECK_DELEGATE(FTestDynamicDelegate5,                     0, &UDynamicDelegateTestObject::Func5,      (GInt16Arg, GStrArg, GPtrArg, GNameArg, bGBoolArg));
-		BIND_CALL_CHECK_DELEGATE(FTestDynamicMulticastDelegate5,            1, &UDynamicDelegateTestObject::Func5,      (GInt16Arg, GStrArg, GPtrArg, GNameArg, bGBoolArg));
-		BIND_CALL_CHECK_DELEGATE(FTestDynamicMulticastDelegate5::FDelegate, 0, &UDynamicDelegateTestObject::ConstFunc5, (GInt16Arg, GStrArg, GPtrArg, GNameArg, bGBoolArg));
-		BIND_CALL_CHECK_DELEGATE(FTestDynamicDelegate5,                     0, &UDynamicDelegateTestObject::ConstFunc5, (GInt16Arg, GStrArg, GPtrArg, GNameArg, bGBoolArg));
-		BIND_CALL_CHECK_DELEGATE(FTestDynamicMulticastDelegate5,            1, &UDynamicDelegateTestObject::ConstFunc5, (GInt16Arg, GStrArg, GPtrArg, GNameArg, bGBoolArg));
-
-		BIND_CALL_CHECK_DELEGATE_RETVAL(GIntRet, FTestDynamicDelegateRetInt5, 0, &UDynamicDelegateTestObject::FuncRetInt5,      (GInt16Arg, GStrArg, GPtrArg, GNameArg, bGBoolArg));
-		BIND_CALL_CHECK_DELEGATE_RETVAL(GIntRet, FTestDynamicDelegateRetInt5, 0, &UDynamicDelegateTestObject::ConstFuncRetInt5, (GInt16Arg, GStrArg, GPtrArg, GNameArg, bGBoolArg));
-		BIND_CALL_CHECK_DELEGATE_RETVAL(GStrRet, FTestDynamicDelegateRetStr5, 0, &UDynamicDelegateTestObject::FuncRetStr5,      (GInt16Arg, GStrArg, GPtrArg, GNameArg, bGBoolArg));
-		BIND_CALL_CHECK_DELEGATE_RETVAL(GStrRet, FTestDynamicDelegateRetStr5, 0, &UDynamicDelegateTestObject::ConstFuncRetStr5, (GInt16Arg, GStrArg, GPtrArg, GNameArg, bGBoolArg));
+		BIND_CALL_CHECK_DELEGATE(FTestDynamicMulticastDelegate5::FDelegate, 0, &UDynamicDelegateTestObject::Func5,      (GInt16Arg, GUnscopedEnumArg, GStrArg, GPtrArg, GNameArg));
+		BIND_CALL_CHECK_DELEGATE(FTestDynamicDelegate5,                     0, &UDynamicDelegateTestObject::Func5,      (GInt16Arg, GUnscopedEnumArg, GStrArg, GPtrArg, GNameArg));
+		BIND_CALL_CHECK_DELEGATE(FTestDynamicMulticastDelegate5,            1, &UDynamicDelegateTestObject::Func5,      (GInt16Arg, GUnscopedEnumArg, GStrArg, GPtrArg, GNameArg));
+		BIND_CALL_CHECK_DELEGATE(FTestDynamicMulticastDelegate5::FDelegate, 0, &UDynamicDelegateTestObject::ConstFunc5, (GInt16Arg, GUnscopedEnumArg, GStrArg, GPtrArg, GNameArg));
+		BIND_CALL_CHECK_DELEGATE(FTestDynamicDelegate5,                     0, &UDynamicDelegateTestObject::ConstFunc5, (GInt16Arg, GUnscopedEnumArg, GStrArg, GPtrArg, GNameArg));
+		BIND_CALL_CHECK_DELEGATE(FTestDynamicMulticastDelegate5,            1, &UDynamicDelegateTestObject::ConstFunc5, (GInt16Arg, GUnscopedEnumArg, GStrArg, GPtrArg, GNameArg));
+
+		BIND_CALL_CHECK_DELEGATE_RETVAL(GIntRet, FTestDynamicDelegateRetInt5, 0, &UDynamicDelegateTestObject::FuncRetInt5,      (GInt16Arg, GUnscopedEnumArg, GStrArg, GPtrArg, GNameArg));
+		BIND_CALL_CHECK_DELEGATE_RETVAL(GIntRet, FTestDynamicDelegateRetInt5, 0, &UDynamicDelegateTestObject::ConstFuncRetInt5, (GInt16Arg, GUnscopedEnumArg, GStrArg, GPtrArg, GNameArg));
+		BIND_CALL_CHECK_DELEGATE_RETVAL(GStrRet, FTestDynamicDelegateRetStr5, 0, &UDynamicDelegateTestObject::FuncRetStr5,      (GInt16Arg, GUnscopedEnumArg, GStrArg, GPtrArg, GNameArg));
+		BIND_CALL_CHECK_DELEGATE_RETVAL(GStrRet, FTestDynamicDelegateRetStr5, 0, &UDynamicDelegateTestObject::ConstFuncRetStr5, (GInt16Arg, GUnscopedEnumArg, GStrArg, GPtrArg, GNameArg));
 
 #if defined(UE_USE_DYNAMIC_DELEGATE_PAYLOADS) && UE_USE_DYNAMIC_DELEGATE_PAYLOADS
 		// 1 arg payload
-		BIND_CALL_CHECK_DELEGATE(FTestDynamicMulticastDelegate4::FDelegate, 0, &UDynamicDelegateTestObject::Func5,      (GInt16Arg, GStrArg, GPtrArg, GNameArg), bGBoolArg);
-		BIND_CALL_CHECK_DELEGATE(FTestDynamicDelegate4,                     0, &UDynamicDelegateTestObject::Func5,      (GInt16Arg, GStrArg, GPtrArg, GNameArg), bGBoolArg);
-		BIND_CALL_CHECK_DELEGATE(FTestDynamicMulticastDelegate4,            1, &UDynamicDelegateTestObject::Func5,      (GInt16Arg, GStrArg, GPtrArg, GNameArg), bGBoolArg);
-		BIND_CALL_CHECK_DELEGATE(FTestDynamicMulticastDelegate4::FDelegate, 0, &UDynamicDelegateTestObject::ConstFunc5, (GInt16Arg, GStrArg, GPtrArg, GNameArg), bGBoolArg);
-		BIND_CALL_CHECK_DELEGATE(FTestDynamicDelegate4,                     0, &UDynamicDelegateTestObject::ConstFunc5, (GInt16Arg, GStrArg, GPtrArg, GNameArg), bGBoolArg);
-		BIND_CALL_CHECK_DELEGATE(FTestDynamicMulticastDelegate4,            1, &UDynamicDelegateTestObject::ConstFunc5, (GInt16Arg, GStrArg, GPtrArg, GNameArg), bGBoolArg);
+		BIND_CALL_CHECK_DELEGATE(FTestDynamicMulticastDelegate4::FDelegate, 0, &UDynamicDelegateTestObject::Func5,      (GInt16Arg, GUnscopedEnumArg, GStrArg, GPtrArg), GNameArg);
+		BIND_CALL_CHECK_DELEGATE(FTestDynamicDelegate4,                     0, &UDynamicDelegateTestObject::Func5,      (GInt16Arg, GUnscopedEnumArg, GStrArg, GPtrArg), GNameArg);
+		BIND_CALL_CHECK_DELEGATE(FTestDynamicMulticastDelegate4,            1, &UDynamicDelegateTestObject::Func5,      (GInt16Arg, GUnscopedEnumArg, GStrArg, GPtrArg), GNameArg);
+		BIND_CALL_CHECK_DELEGATE(FTestDynamicMulticastDelegate4::FDelegate, 0, &UDynamicDelegateTestObject::ConstFunc5, (GInt16Arg, GUnscopedEnumArg, GStrArg, GPtrArg), GNameArg);
+		BIND_CALL_CHECK_DELEGATE(FTestDynamicDelegate4,                     0, &UDynamicDelegateTestObject::ConstFunc5, (GInt16Arg, GUnscopedEnumArg, GStrArg, GPtrArg), GNameArg);
+		BIND_CALL_CHECK_DELEGATE(FTestDynamicMulticastDelegate4,            1, &UDynamicDelegateTestObject::ConstFunc5, (GInt16Arg, GUnscopedEnumArg, GStrArg, GPtrArg), GNameArg);
 #if UE_USE_DYNAMIC_DELEGATE_PAYLOADS_RETVAL
-		BIND_CALL_CHECK_DELEGATE_RETVAL(GIntRet, FTestDynamicDelegateRetInt4, 0, &UDynamicDelegateTestObject::FuncRetInt5,      (GInt16Arg, GStrArg, GPtrArg, GNameArg), bGBoolArg);
-		BIND_CALL_CHECK_DELEGATE_RETVAL(GIntRet, FTestDynamicDelegateRetInt4, 0, &UDynamicDelegateTestObject::ConstFuncRetInt5, (GInt16Arg, GStrArg, GPtrArg, GNameArg), bGBoolArg);
-		BIND_CALL_CHECK_DELEGATE_RETVAL(GStrRet, FTestDynamicDelegateRetStr4, 0, &UDynamicDelegateTestObject::FuncRetStr5,      (GInt16Arg, GStrArg, GPtrArg, GNameArg), bGBoolArg);
-		BIND_CALL_CHECK_DELEGATE_RETVAL(GStrRet, FTestDynamicDelegateRetStr4, 0, &UDynamicDelegateTestObject::ConstFuncRetStr5, (GInt16Arg, GStrArg, GPtrArg, GNameArg), bGBoolArg);
+		BIND_CALL_CHECK_DELEGATE_RETVAL(GIntRet, FTestDynamicDelegateRetInt4, 0, &UDynamicDelegateTestObject::FuncRetInt5,      (GInt16Arg, GUnscopedEnumArg, GStrArg, GPtrArg), GNameArg);
+		BIND_CALL_CHECK_DELEGATE_RETVAL(GIntRet, FTestDynamicDelegateRetInt4, 0, &UDynamicDelegateTestObject::ConstFuncRetInt5, (GInt16Arg, GUnscopedEnumArg, GStrArg, GPtrArg), GNameArg);
+		BIND_CALL_CHECK_DELEGATE_RETVAL(GStrRet, FTestDynamicDelegateRetStr4, 0, &UDynamicDelegateTestObject::FuncRetStr5,      (GInt16Arg, GUnscopedEnumArg, GStrArg, GPtrArg), GNameArg);
+		BIND_CALL_CHECK_DELEGATE_RETVAL(GStrRet, FTestDynamicDelegateRetStr4, 0, &UDynamicDelegateTestObject::ConstFuncRetStr5, (GInt16Arg, GUnscopedEnumArg, GStrArg, GPtrArg), GNameArg);
 #endif
 
 		// 2 arg payload
-		BIND_CALL_CHECK_DELEGATE(FTestDynamicMulticastDelegate3::FDelegate, 0, &UDynamicDelegateTestObject::Func5,      (GInt16Arg, GStrArg, GPtrArg), GNameArg, bGBoolArg);
-		BIND_CALL_CHECK_DELEGATE(FTestDynamicDelegate3,                     0, &UDynamicDelegateTestObject::Func5,      (GInt16Arg, GStrArg, GPtrArg), GNameArg, bGBoolArg);
-		BIND_CALL_CHECK_DELEGATE(FTestDynamicMulticastDelegate3,            1, &UDynamicDelegateTestObject::Func5,      (GInt16Arg, GStrArg, GPtrArg), GNameArg, bGBoolArg);
-		BIND_CALL_CHECK_DELEGATE(FTestDynamicMulticastDelegate3::FDelegate, 0, &UDynamicDelegateTestObject::ConstFunc5, (GInt16Arg, GStrArg, GPtrArg), GNameArg, bGBoolArg);
-		BIND_CALL_CHECK_DELEGATE(FTestDynamicDelegate3,                     0, &UDynamicDelegateTestObject::ConstFunc5, (GInt16Arg, GStrArg, GPtrArg), GNameArg, bGBoolArg);
-		BIND_CALL_CHECK_DELEGATE(FTestDynamicMulticastDelegate3,            1, &UDynamicDelegateTestObject::ConstFunc5, (GInt16Arg, GStrArg, GPtrArg), GNameArg, bGBoolArg);
+		BIND_CALL_CHECK_DELEGATE(FTestDynamicMulticastDelegate3::FDelegate, 0, &UDynamicDelegateTestObject::Func5,      (GInt16Arg, GUnscopedEnumArg, GStrArg), GPtrArg, GNameArg);
+		BIND_CALL_CHECK_DELEGATE(FTestDynamicDelegate3,                     0, &UDynamicDelegateTestObject::Func5,      (GInt16Arg, GUnscopedEnumArg, GStrArg), GPtrArg, GNameArg);
+		BIND_CALL_CHECK_DELEGATE(FTestDynamicMulticastDelegate3,            1, &UDynamicDelegateTestObject::Func5,      (GInt16Arg, GUnscopedEnumArg, GStrArg), GPtrArg, GNameArg);
+		BIND_CALL_CHECK_DELEGATE(FTestDynamicMulticastDelegate3::FDelegate, 0, &UDynamicDelegateTestObject::ConstFunc5, (GInt16Arg, GUnscopedEnumArg, GStrArg), GPtrArg, GNameArg);
+		BIND_CALL_CHECK_DELEGATE(FTestDynamicDelegate3,                     0, &UDynamicDelegateTestObject::ConstFunc5, (GInt16Arg, GUnscopedEnumArg, GStrArg), GPtrArg, GNameArg);
+		BIND_CALL_CHECK_DELEGATE(FTestDynamicMulticastDelegate3,            1, &UDynamicDelegateTestObject::ConstFunc5, (GInt16Arg, GUnscopedEnumArg, GStrArg), GPtrArg, GNameArg);
 #if UE_USE_DYNAMIC_DELEGATE_PAYLOADS_RETVAL
-		BIND_CALL_CHECK_DELEGATE_RETVAL(GIntRet, FTestDynamicDelegateRetInt3, 0, &UDynamicDelegateTestObject::FuncRetInt5,      (GInt16Arg, GStrArg, GPtrArg), GNameArg, bGBoolArg);
-		BIND_CALL_CHECK_DELEGATE_RETVAL(GIntRet, FTestDynamicDelegateRetInt3, 0, &UDynamicDelegateTestObject::ConstFuncRetInt5, (GInt16Arg, GStrArg, GPtrArg), GNameArg, bGBoolArg);
-		BIND_CALL_CHECK_DELEGATE_RETVAL(GStrRet, FTestDynamicDelegateRetStr3, 0, &UDynamicDelegateTestObject::FuncRetStr5,      (GInt16Arg, GStrArg, GPtrArg), GNameArg, bGBoolArg);
-		BIND_CALL_CHECK_DELEGATE_RETVAL(GStrRet, FTestDynamicDelegateRetStr3, 0, &UDynamicDelegateTestObject::ConstFuncRetStr5, (GInt16Arg, GStrArg, GPtrArg), GNameArg, bGBoolArg);
+		BIND_CALL_CHECK_DELEGATE_RETVAL(GIntRet, FTestDynamicDelegateRetInt3, 0, &UDynamicDelegateTestObject::FuncRetInt5,      (GInt16Arg, GUnscopedEnumArg, GStrArg), GPtrArg, GNameArg);
+		BIND_CALL_CHECK_DELEGATE_RETVAL(GIntRet, FTestDynamicDelegateRetInt3, 0, &UDynamicDelegateTestObject::ConstFuncRetInt5, (GInt16Arg, GUnscopedEnumArg, GStrArg), GPtrArg, GNameArg);
+		BIND_CALL_CHECK_DELEGATE_RETVAL(GStrRet, FTestDynamicDelegateRetStr3, 0, &UDynamicDelegateTestObject::FuncRetStr5,      (GInt16Arg, GUnscopedEnumArg, GStrArg), GPtrArg, GNameArg);
+		BIND_CALL_CHECK_DELEGATE_RETVAL(GStrRet, FTestDynamicDelegateRetStr3, 0, &UDynamicDelegateTestObject::ConstFuncRetStr5, (GInt16Arg, GUnscopedEnumArg, GStrArg), GPtrArg, GNameArg);
 #endif
 
 		// 3 arg payload
-		BIND_CALL_CHECK_DELEGATE(FTestDynamicMulticastDelegate2::FDelegate, 0, &UDynamicDelegateTestObject::Func5,      (GInt16Arg, GStrArg), GPtrArg, GNameArg, bGBoolArg);
-		BIND_CALL_CHECK_DELEGATE(FTestDynamicDelegate2,                     0, &UDynamicDelegateTestObject::Func5,      (GInt16Arg, GStrArg), GPtrArg, GNameArg, bGBoolArg);
-		BIND_CALL_CHECK_DELEGATE(FTestDynamicMulticastDelegate2,            1, &UDynamicDelegateTestObject::Func5,      (GInt16Arg, GStrArg), GPtrArg, GNameArg, bGBoolArg);
-		BIND_CALL_CHECK_DELEGATE(FTestDynamicMulticastDelegate2::FDelegate, 0, &UDynamicDelegateTestObject::ConstFunc5, (GInt16Arg, GStrArg), GPtrArg, GNameArg, bGBoolArg);
-		BIND_CALL_CHECK_DELEGATE(FTestDynamicDelegate2,                     0, &UDynamicDelegateTestObject::ConstFunc5, (GInt16Arg, GStrArg), GPtrArg, GNameArg, bGBoolArg);
-		BIND_CALL_CHECK_DELEGATE(FTestDynamicMulticastDelegate2,            1, &UDynamicDelegateTestObject::ConstFunc5, (GInt16Arg, GStrArg), GPtrArg, GNameArg, bGBoolArg);
+		BIND_CALL_CHECK_DELEGATE(FTestDynamicMulticastDelegate2::FDelegate, 0, &UDynamicDelegateTestObject::Func5,      (GInt16Arg, GUnscopedEnumArg), GStrArg, GPtrArg, GNameArg);
+		BIND_CALL_CHECK_DELEGATE(FTestDynamicDelegate2,                     0, &UDynamicDelegateTestObject::Func5,      (GInt16Arg, GUnscopedEnumArg), GStrArg, GPtrArg, GNameArg);
+		BIND_CALL_CHECK_DELEGATE(FTestDynamicMulticastDelegate2,            1, &UDynamicDelegateTestObject::Func5,      (GInt16Arg, GUnscopedEnumArg), GStrArg, GPtrArg, GNameArg);
+		BIND_CALL_CHECK_DELEGATE(FTestDynamicMulticastDelegate2::FDelegate, 0, &UDynamicDelegateTestObject::ConstFunc5, (GInt16Arg, GUnscopedEnumArg), GStrArg, GPtrArg, GNameArg);
+		BIND_CALL_CHECK_DELEGATE(FTestDynamicDelegate2,                     0, &UDynamicDelegateTestObject::ConstFunc5, (GInt16Arg, GUnscopedEnumArg), GStrArg, GPtrArg, GNameArg);
+		BIND_CALL_CHECK_DELEGATE(FTestDynamicMulticastDelegate2,            1, &UDynamicDelegateTestObject::ConstFunc5, (GInt16Arg, GUnscopedEnumArg), GStrArg, GPtrArg, GNameArg);
 #if UE_USE_DYNAMIC_DELEGATE_PAYLOADS_RETVAL
-		BIND_CALL_CHECK_DELEGATE_RETVAL(GIntRet, FTestDynamicDelegateRetInt2, 0, &UDynamicDelegateTestObject::FuncRetInt5,      (GInt16Arg, GStrArg), GPtrArg, GNameArg, bGBoolArg);
-		BIND_CALL_CHECK_DELEGATE_RETVAL(GIntRet, FTestDynamicDelegateRetInt2, 0, &UDynamicDelegateTestObject::ConstFuncRetInt5, (GInt16Arg, GStrArg), GPtrArg, GNameArg, bGBoolArg);
-		BIND_CALL_CHECK_DELEGATE_RETVAL(GStrRet, FTestDynamicDelegateRetStr2, 0, &UDynamicDelegateTestObject::FuncRetStr5,      (GInt16Arg, GStrArg), GPtrArg, GNameArg, bGBoolArg);
-		BIND_CALL_CHECK_DELEGATE_RETVAL(GStrRet, FTestDynamicDelegateRetStr2, 0, &UDynamicDelegateTestObject::ConstFuncRetStr5, (GInt16Arg, GStrArg), GPtrArg, GNameArg, bGBoolArg);
+		BIND_CALL_CHECK_DELEGATE_RETVAL(GIntRet, FTestDynamicDelegateRetInt2, 0, &UDynamicDelegateTestObject::FuncRetInt5,      (GInt16Arg, GUnscopedEnumArg), GStrArg, GPtrArg, GNameArg);
+		BIND_CALL_CHECK_DELEGATE_RETVAL(GIntRet, FTestDynamicDelegateRetInt2, 0, &UDynamicDelegateTestObject::ConstFuncRetInt5, (GInt16Arg, GUnscopedEnumArg), GStrArg, GPtrArg, GNameArg);
+		BIND_CALL_CHECK_DELEGATE_RETVAL(GStrRet, FTestDynamicDelegateRetStr2, 0, &UDynamicDelegateTestObject::FuncRetStr5,      (GInt16Arg, GUnscopedEnumArg), GStrArg, GPtrArg, GNameArg);
+		BIND_CALL_CHECK_DELEGATE_RETVAL(GStrRet, FTestDynamicDelegateRetStr2, 0, &UDynamicDelegateTestObject::ConstFuncRetStr5, (GInt16Arg, GUnscopedEnumArg), GStrArg, GPtrArg, GNameArg);
 #endif
 
 		// 4 arg payload
-		BIND_CALL_CHECK_DELEGATE(FTestDynamicMulticastDelegate1::FDelegate, 0, &UDynamicDelegateTestObject::Func5,      (GInt16Arg), GStrArg, GPtrArg, GNameArg, bGBoolArg);
-		BIND_CALL_CHECK_DELEGATE(FTestDynamicDelegate1,                     0, &UDynamicDelegateTestObject::Func5,      (GInt16Arg), GStrArg, GPtrArg, GNameArg, bGBoolArg);
-		BIND_CALL_CHECK_DELEGATE(FTestDynamicMulticastDelegate1,            1, &UDynamicDelegateTestObject::Func5,      (GInt16Arg), GStrArg, GPtrArg, GNameArg, bGBoolArg);
-		BIND_CALL_CHECK_DELEGATE(FTestDynamicMulticastDelegate1::FDelegate, 0, &UDynamicDelegateTestObject::ConstFunc5, (GInt16Arg), GStrArg, GPtrArg, GNameArg, bGBoolArg);
-		BIND_CALL_CHECK_DELEGATE(FTestDynamicDelegate1,                     0, &UDynamicDelegateTestObject::ConstFunc5, (GInt16Arg), GStrArg, GPtrArg, GNameArg, bGBoolArg);
-		BIND_CALL_CHECK_DELEGATE(FTestDynamicMulticastDelegate1,            1, &UDynamicDelegateTestObject::ConstFunc5, (GInt16Arg), GStrArg, GPtrArg, GNameArg, bGBoolArg);
+		BIND_CALL_CHECK_DELEGATE(FTestDynamicMulticastDelegate1::FDelegate, 0, &UDynamicDelegateTestObject::Func5,      (GInt16Arg), GUnscopedEnumArg, GStrArg, GPtrArg, GNameArg);
+		BIND_CALL_CHECK_DELEGATE(FTestDynamicDelegate1,                     0, &UDynamicDelegateTestObject::Func5,      (GInt16Arg), GUnscopedEnumArg, GStrArg, GPtrArg, GNameArg);
+		BIND_CALL_CHECK_DELEGATE(FTestDynamicMulticastDelegate1,            1, &UDynamicDelegateTestObject::Func5,      (GInt16Arg), GUnscopedEnumArg, GStrArg, GPtrArg, GNameArg);
+		BIND_CALL_CHECK_DELEGATE(FTestDynamicMulticastDelegate1::FDelegate, 0, &UDynamicDelegateTestObject::ConstFunc5, (GInt16Arg), GUnscopedEnumArg, GStrArg, GPtrArg, GNameArg);
+		BIND_CALL_CHECK_DELEGATE(FTestDynamicDelegate1,                     0, &UDynamicDelegateTestObject::ConstFunc5, (GInt16Arg), GUnscopedEnumArg, GStrArg, GPtrArg, GNameArg);
+		BIND_CALL_CHECK_DELEGATE(FTestDynamicMulticastDelegate1,            1, &UDynamicDelegateTestObject::ConstFunc5, (GInt16Arg), GUnscopedEnumArg, GStrArg, GPtrArg, GNameArg);
 #if UE_USE_DYNAMIC_DELEGATE_PAYLOADS_RETVAL
-		BIND_CALL_CHECK_DELEGATE_RETVAL(GIntRet, FTestDynamicDelegateRetInt1, 0, &UDynamicDelegateTestObject::FuncRetInt5,      (GInt16Arg), GStrArg, GPtrArg, GNameArg, bGBoolArg);
-		BIND_CALL_CHECK_DELEGATE_RETVAL(GIntRet, FTestDynamicDelegateRetInt1, 0, &UDynamicDelegateTestObject::ConstFuncRetInt5, (GInt16Arg), GStrArg, GPtrArg, GNameArg, bGBoolArg);
-		BIND_CALL_CHECK_DELEGATE_RETVAL(GStrRet, FTestDynamicDelegateRetStr1, 0, &UDynamicDelegateTestObject::FuncRetStr5,      (GInt16Arg), GStrArg, GPtrArg, GNameArg, bGBoolArg);
-		BIND_CALL_CHECK_DELEGATE_RETVAL(GStrRet, FTestDynamicDelegateRetStr1, 0, &UDynamicDelegateTestObject::ConstFuncRetStr5, (GInt16Arg), GStrArg, GPtrArg, GNameArg, bGBoolArg);
+		BIND_CALL_CHECK_DELEGATE_RETVAL(GIntRet, FTestDynamicDelegateRetInt1, 0, &UDynamicDelegateTestObject::FuncRetInt5,      (GInt16Arg), GUnscopedEnumArg, GStrArg, GPtrArg, GNameArg);
+		BIND_CALL_CHECK_DELEGATE_RETVAL(GIntRet, FTestDynamicDelegateRetInt1, 0, &UDynamicDelegateTestObject::ConstFuncRetInt5, (GInt16Arg), GUnscopedEnumArg, GStrArg, GPtrArg, GNameArg);
+		BIND_CALL_CHECK_DELEGATE_RETVAL(GStrRet, FTestDynamicDelegateRetStr1, 0, &UDynamicDelegateTestObject::FuncRetStr5,      (GInt16Arg), GUnscopedEnumArg, GStrArg, GPtrArg, GNameArg);
+		BIND_CALL_CHECK_DELEGATE_RETVAL(GStrRet, FTestDynamicDelegateRetStr1, 0, &UDynamicDelegateTestObject::ConstFuncRetStr5, (GInt16Arg), GUnscopedEnumArg, GStrArg, GPtrArg, GNameArg);
 #endif
 
 		// 5 arg payload
-		BIND_CALL_CHECK_DELEGATE(FTestDynamicMulticastDelegate0::FDelegate, 0, &UDynamicDelegateTestObject::Func5,      (), GInt16Arg, GStrArg, GPtrArg, GNameArg, bGBoolArg);
-		BIND_CALL_CHECK_DELEGATE(FTestDynamicDelegate0,                     0, &UDynamicDelegateTestObject::Func5,      (), GInt16Arg, GStrArg, GPtrArg, GNameArg, bGBoolArg);
-		BIND_CALL_CHECK_DELEGATE(FTestDynamicMulticastDelegate0,            1, &UDynamicDelegateTestObject::Func5,      (), GInt16Arg, GStrArg, GPtrArg, GNameArg, bGBoolArg);
-		BIND_CALL_CHECK_DELEGATE(FTestDynamicMulticastDelegate0::FDelegate, 0, &UDynamicDelegateTestObject::ConstFunc5, (), GInt16Arg, GStrArg, GPtrArg, GNameArg, bGBoolArg);
-		BIND_CALL_CHECK_DELEGATE(FTestDynamicDelegate0,                     0, &UDynamicDelegateTestObject::ConstFunc5, (), GInt16Arg, GStrArg, GPtrArg, GNameArg, bGBoolArg);
-		BIND_CALL_CHECK_DELEGATE(FTestDynamicMulticastDelegate0,            1, &UDynamicDelegateTestObject::ConstFunc5, (), GInt16Arg, GStrArg, GPtrArg, GNameArg, bGBoolArg);
+		BIND_CALL_CHECK_DELEGATE(FTestDynamicMulticastDelegate0::FDelegate, 0, &UDynamicDelegateTestObject::Func5,      (), GInt16Arg, GUnscopedEnumArg, GStrArg, GPtrArg, GNameArg);
+		BIND_CALL_CHECK_DELEGATE(FTestDynamicDelegate0,                     0, &UDynamicDelegateTestObject::Func5,      (), GInt16Arg, GUnscopedEnumArg, GStrArg, GPtrArg, GNameArg);
+		BIND_CALL_CHECK_DELEGATE(FTestDynamicMulticastDelegate0,            1, &UDynamicDelegateTestObject::Func5,      (), GInt16Arg, GUnscopedEnumArg, GStrArg, GPtrArg, GNameArg);
+		BIND_CALL_CHECK_DELEGATE(FTestDynamicMulticastDelegate0::FDelegate, 0, &UDynamicDelegateTestObject::ConstFunc5, (), GInt16Arg, GUnscopedEnumArg, GStrArg, GPtrArg, GNameArg);
+		BIND_CALL_CHECK_DELEGATE(FTestDynamicDelegate0,                     0, &UDynamicDelegateTestObject::ConstFunc5, (), GInt16Arg, GUnscopedEnumArg, GStrArg, GPtrArg, GNameArg);
+		BIND_CALL_CHECK_DELEGATE(FTestDynamicMulticastDelegate0,            1, &UDynamicDelegateTestObject::ConstFunc5, (), GInt16Arg, GUnscopedEnumArg, GStrArg, GPtrArg, GNameArg);
 #if UE_USE_DYNAMIC_DELEGATE_PAYLOADS_RETVAL
-		BIND_CALL_CHECK_DELEGATE_RETVAL(GIntRet, FTestDynamicDelegateRetInt0, 0, &UDynamicDelegateTestObject::FuncRetInt5,      (), GInt16Arg, GStrArg, GPtrArg, GNameArg, bGBoolArg);
-		BIND_CALL_CHECK_DELEGATE_RETVAL(GIntRet, FTestDynamicDelegateRetInt0, 0, &UDynamicDelegateTestObject::ConstFuncRetInt5, (), GInt16Arg, GStrArg, GPtrArg, GNameArg, bGBoolArg);
-		BIND_CALL_CHECK_DELEGATE_RETVAL(GStrRet, FTestDynamicDelegateRetStr0, 0, &UDynamicDelegateTestObject::FuncRetStr5,      (), GInt16Arg, GStrArg, GPtrArg, GNameArg, bGBoolArg);
-		BIND_CALL_CHECK_DELEGATE_RETVAL(GStrRet, FTestDynamicDelegateRetStr0, 0, &UDynamicDelegateTestObject::ConstFuncRetStr5, (), GInt16Arg, GStrArg, GPtrArg, GNameArg, bGBoolArg);
+		BIND_CALL_CHECK_DELEGATE_RETVAL(GIntRet, FTestDynamicDelegateRetInt0, 0, &UDynamicDelegateTestObject::FuncRetInt5,      (), GInt16Arg, GUnscopedEnumArg, GStrArg, GPtrArg, GNameArg);
+		BIND_CALL_CHECK_DELEGATE_RETVAL(GIntRet, FTestDynamicDelegateRetInt0, 0, &UDynamicDelegateTestObject::ConstFuncRetInt5, (), GInt16Arg, GUnscopedEnumArg, GStrArg, GPtrArg, GNameArg);
+		BIND_CALL_CHECK_DELEGATE_RETVAL(GStrRet, FTestDynamicDelegateRetStr0, 0, &UDynamicDelegateTestObject::FuncRetStr5,      (), GInt16Arg, GUnscopedEnumArg, GStrArg, GPtrArg, GNameArg);
+		BIND_CALL_CHECK_DELEGATE_RETVAL(GStrRet, FTestDynamicDelegateRetStr0, 0, &UDynamicDelegateTestObject::ConstFuncRetStr5, (), GInt16Arg, GUnscopedEnumArg, GStrArg, GPtrArg, GNameArg);
 #endif
 #endif
 	}
@@ -845,101 +932,217 @@ void UDynamicDelegateTestObject::RunTests()
 	// 6 param function
 	{
 		// 0 arg payload
-		BIND_CALL_CHECK_DELEGATE(FTestDynamicMulticastDelegate6::FDelegate, 0, &UDynamicDelegateTestObject::Func6,      (GInt16Arg, GStrArg, GPtrArg, GNameArg, bGBoolArg, GWeakObjPtrArg));
-		BIND_CALL_CHECK_DELEGATE(FTestDynamicDelegate6,                     0, &UDynamicDelegateTestObject::Func6,      (GInt16Arg, GStrArg, GPtrArg, GNameArg, bGBoolArg, GWeakObjPtrArg));
-		BIND_CALL_CHECK_DELEGATE(FTestDynamicMulticastDelegate6,            1, &UDynamicDelegateTestObject::Func6,      (GInt16Arg, GStrArg, GPtrArg, GNameArg, bGBoolArg, GWeakObjPtrArg));
-		BIND_CALL_CHECK_DELEGATE(FTestDynamicMulticastDelegate6::FDelegate, 0, &UDynamicDelegateTestObject::ConstFunc6, (GInt16Arg, GStrArg, GPtrArg, GNameArg, bGBoolArg, GWeakObjPtrArg));
-		BIND_CALL_CHECK_DELEGATE(FTestDynamicDelegate6,                     0, &UDynamicDelegateTestObject::ConstFunc6, (GInt16Arg, GStrArg, GPtrArg, GNameArg, bGBoolArg, GWeakObjPtrArg));
-		BIND_CALL_CHECK_DELEGATE(FTestDynamicMulticastDelegate6,            1, &UDynamicDelegateTestObject::ConstFunc6, (GInt16Arg, GStrArg, GPtrArg, GNameArg, bGBoolArg, GWeakObjPtrArg));
-
-		BIND_CALL_CHECK_DELEGATE_RETVAL(GIntRet, FTestDynamicDelegateRetInt6, 0, &UDynamicDelegateTestObject::FuncRetInt6,      (GInt16Arg, GStrArg, GPtrArg, GNameArg, bGBoolArg, GWeakObjPtrArg));
-		BIND_CALL_CHECK_DELEGATE_RETVAL(GIntRet, FTestDynamicDelegateRetInt6, 0, &UDynamicDelegateTestObject::ConstFuncRetInt6, (GInt16Arg, GStrArg, GPtrArg, GNameArg, bGBoolArg, GWeakObjPtrArg));
-		BIND_CALL_CHECK_DELEGATE_RETVAL(GStrRet, FTestDynamicDelegateRetStr6, 0, &UDynamicDelegateTestObject::FuncRetStr6,      (GInt16Arg, GStrArg, GPtrArg, GNameArg, bGBoolArg, GWeakObjPtrArg));
-		BIND_CALL_CHECK_DELEGATE_RETVAL(GStrRet, FTestDynamicDelegateRetStr6, 0, &UDynamicDelegateTestObject::ConstFuncRetStr6, (GInt16Arg, GStrArg, GPtrArg, GNameArg, bGBoolArg, GWeakObjPtrArg));
+		BIND_CALL_CHECK_DELEGATE(FTestDynamicMulticastDelegate6::FDelegate, 0, &UDynamicDelegateTestObject::Func6,      (GInt16Arg, GUnscopedEnumArg, GStrArg, GPtrArg, GNameArg, bGBoolArg));
+		BIND_CALL_CHECK_DELEGATE(FTestDynamicDelegate6,                     0, &UDynamicDelegateTestObject::Func6,      (GInt16Arg, GUnscopedEnumArg, GStrArg, GPtrArg, GNameArg, bGBoolArg));
+		BIND_CALL_CHECK_DELEGATE(FTestDynamicMulticastDelegate6,            1, &UDynamicDelegateTestObject::Func6,      (GInt16Arg, GUnscopedEnumArg, GStrArg, GPtrArg, GNameArg, bGBoolArg));
+		BIND_CALL_CHECK_DELEGATE(FTestDynamicMulticastDelegate6::FDelegate, 0, &UDynamicDelegateTestObject::ConstFunc6, (GInt16Arg, GUnscopedEnumArg, GStrArg, GPtrArg, GNameArg, bGBoolArg));
+		BIND_CALL_CHECK_DELEGATE(FTestDynamicDelegate6,                     0, &UDynamicDelegateTestObject::ConstFunc6, (GInt16Arg, GUnscopedEnumArg, GStrArg, GPtrArg, GNameArg, bGBoolArg));
+		BIND_CALL_CHECK_DELEGATE(FTestDynamicMulticastDelegate6,            1, &UDynamicDelegateTestObject::ConstFunc6, (GInt16Arg, GUnscopedEnumArg, GStrArg, GPtrArg, GNameArg, bGBoolArg));
+
+		BIND_CALL_CHECK_DELEGATE_RETVAL(GIntRet, FTestDynamicDelegateRetInt6, 0, &UDynamicDelegateTestObject::FuncRetInt6,      (GInt16Arg, GUnscopedEnumArg, GStrArg, GPtrArg, GNameArg, bGBoolArg));
+		BIND_CALL_CHECK_DELEGATE_RETVAL(GIntRet, FTestDynamicDelegateRetInt6, 0, &UDynamicDelegateTestObject::ConstFuncRetInt6, (GInt16Arg, GUnscopedEnumArg, GStrArg, GPtrArg, GNameArg, bGBoolArg));
+		BIND_CALL_CHECK_DELEGATE_RETVAL(GStrRet, FTestDynamicDelegateRetStr6, 0, &UDynamicDelegateTestObject::FuncRetStr6,      (GInt16Arg, GUnscopedEnumArg, GStrArg, GPtrArg, GNameArg, bGBoolArg));
+		BIND_CALL_CHECK_DELEGATE_RETVAL(GStrRet, FTestDynamicDelegateRetStr6, 0, &UDynamicDelegateTestObject::ConstFuncRetStr6, (GInt16Arg, GUnscopedEnumArg, GStrArg, GPtrArg, GNameArg, bGBoolArg));
 
 #if defined(UE_USE_DYNAMIC_DELEGATE_PAYLOADS) && UE_USE_DYNAMIC_DELEGATE_PAYLOADS
 		// 1 arg payload
-		BIND_CALL_CHECK_DELEGATE(FTestDynamicMulticastDelegate5::FDelegate, 0, &UDynamicDelegateTestObject::Func6,      (GInt16Arg, GStrArg, GPtrArg, GNameArg, bGBoolArg), GWeakObjPtrArg);
-		BIND_CALL_CHECK_DELEGATE(FTestDynamicDelegate5,                     0, &UDynamicDelegateTestObject::Func6,      (GInt16Arg, GStrArg, GPtrArg, GNameArg, bGBoolArg), GWeakObjPtrArg);
-		BIND_CALL_CHECK_DELEGATE(FTestDynamicMulticastDelegate5,            1, &UDynamicDelegateTestObject::Func6,      (GInt16Arg, GStrArg, GPtrArg, GNameArg, bGBoolArg), GWeakObjPtrArg);
-		BIND_CALL_CHECK_DELEGATE(FTestDynamicMulticastDelegate5::FDelegate, 0, &UDynamicDelegateTestObject::ConstFunc6, (GInt16Arg, GStrArg, GPtrArg, GNameArg, bGBoolArg), GWeakObjPtrArg);
-		BIND_CALL_CHECK_DELEGATE(FTestDynamicDelegate5,                     0, &UDynamicDelegateTestObject::ConstFunc6, (GInt16Arg, GStrArg, GPtrArg, GNameArg, bGBoolArg), GWeakObjPtrArg);
-		BIND_CALL_CHECK_DELEGATE(FTestDynamicMulticastDelegate5,            1, &UDynamicDelegateTestObject::ConstFunc6, (GInt16Arg, GStrArg, GPtrArg, GNameArg, bGBoolArg), GWeakObjPtrArg);
+		BIND_CALL_CHECK_DELEGATE(FTestDynamicMulticastDelegate5::FDelegate, 0, &UDynamicDelegateTestObject::Func6,      (GInt16Arg, GUnscopedEnumArg, GStrArg, GPtrArg, GNameArg), bGBoolArg);
+		BIND_CALL_CHECK_DELEGATE(FTestDynamicDelegate5,                     0, &UDynamicDelegateTestObject::Func6,      (GInt16Arg, GUnscopedEnumArg, GStrArg, GPtrArg, GNameArg), bGBoolArg);
+		BIND_CALL_CHECK_DELEGATE(FTestDynamicMulticastDelegate5,            1, &UDynamicDelegateTestObject::Func6,      (GInt16Arg, GUnscopedEnumArg, GStrArg, GPtrArg, GNameArg), bGBoolArg);
+		BIND_CALL_CHECK_DELEGATE(FTestDynamicMulticastDelegate5::FDelegate, 0, &UDynamicDelegateTestObject::ConstFunc6, (GInt16Arg, GUnscopedEnumArg, GStrArg, GPtrArg, GNameArg), bGBoolArg);
+		BIND_CALL_CHECK_DELEGATE(FTestDynamicDelegate5,                     0, &UDynamicDelegateTestObject::ConstFunc6, (GInt16Arg, GUnscopedEnumArg, GStrArg, GPtrArg, GNameArg), bGBoolArg);
+		BIND_CALL_CHECK_DELEGATE(FTestDynamicMulticastDelegate5,            1, &UDynamicDelegateTestObject::ConstFunc6, (GInt16Arg, GUnscopedEnumArg, GStrArg, GPtrArg, GNameArg), bGBoolArg);
 #if UE_USE_DYNAMIC_DELEGATE_PAYLOADS_RETVAL
-		BIND_CALL_CHECK_DELEGATE_RETVAL(GIntRet, FTestDynamicDelegateRetInt5, 0, &UDynamicDelegateTestObject::FuncRetInt6,      (GInt16Arg, GStrArg, GPtrArg, GNameArg, bGBoolArg), GWeakObjPtrArg);
-		BIND_CALL_CHECK_DELEGATE_RETVAL(GIntRet, FTestDynamicDelegateRetInt5, 0, &UDynamicDelegateTestObject::ConstFuncRetInt6, (GInt16Arg, GStrArg, GPtrArg, GNameArg, bGBoolArg), GWeakObjPtrArg);
-		BIND_CALL_CHECK_DELEGATE_RETVAL(GStrRet, FTestDynamicDelegateRetStr5, 0, &UDynamicDelegateTestObject::FuncRetStr6,      (GInt16Arg, GStrArg, GPtrArg, GNameArg, bGBoolArg), GWeakObjPtrArg);
-		BIND_CALL_CHECK_DELEGATE_RETVAL(GStrRet, FTestDynamicDelegateRetStr5, 0, &UDynamicDelegateTestObject::ConstFuncRetStr6, (GInt16Arg, GStrArg, GPtrArg, GNameArg, bGBoolArg), GWeakObjPtrArg);
+		BIND_CALL_CHECK_DELEGATE_RETVAL(GIntRet, FTestDynamicDelegateRetInt5, 0, &UDynamicDelegateTestObject::FuncRetInt6,      (GInt16Arg, GUnscopedEnumArg, GStrArg, GPtrArg, GNameArg), bGBoolArg);
+		BIND_CALL_CHECK_DELEGATE_RETVAL(GIntRet, FTestDynamicDelegateRetInt5, 0, &UDynamicDelegateTestObject::ConstFuncRetInt6, (GInt16Arg, GUnscopedEnumArg, GStrArg, GPtrArg, GNameArg), bGBoolArg);
+		BIND_CALL_CHECK_DELEGATE_RETVAL(GStrRet, FTestDynamicDelegateRetStr5, 0, &UDynamicDelegateTestObject::FuncRetStr6,      (GInt16Arg, GUnscopedEnumArg, GStrArg, GPtrArg, GNameArg), bGBoolArg);
+		BIND_CALL_CHECK_DELEGATE_RETVAL(GStrRet, FTestDynamicDelegateRetStr5, 0, &UDynamicDelegateTestObject::ConstFuncRetStr6, (GInt16Arg, GUnscopedEnumArg, GStrArg, GPtrArg, GNameArg), bGBoolArg);
 #endif
 
 		// 2 arg payload
-		BIND_CALL_CHECK_DELEGATE(FTestDynamicMulticastDelegate4::FDelegate, 0, &UDynamicDelegateTestObject::Func6,      (GInt16Arg, GStrArg, GPtrArg, GNameArg), bGBoolArg, GWeakObjPtrArg);
-		BIND_CALL_CHECK_DELEGATE(FTestDynamicDelegate4,                     0, &UDynamicDelegateTestObject::Func6,      (GInt16Arg, GStrArg, GPtrArg, GNameArg), bGBoolArg, GWeakObjPtrArg);
-		BIND_CALL_CHECK_DELEGATE(FTestDynamicMulticastDelegate4,            1, &UDynamicDelegateTestObject::Func6,      (GInt16Arg, GStrArg, GPtrArg, GNameArg), bGBoolArg, GWeakObjPtrArg);
-		BIND_CALL_CHECK_DELEGATE(FTestDynamicMulticastDelegate4::FDelegate, 0, &UDynamicDelegateTestObject::ConstFunc6, (GInt16Arg, GStrArg, GPtrArg, GNameArg), bGBoolArg, GWeakObjPtrArg);
-		BIND_CALL_CHECK_DELEGATE(FTestDynamicDelegate4,                     0, &UDynamicDelegateTestObject::ConstFunc6, (GInt16Arg, GStrArg, GPtrArg, GNameArg), bGBoolArg, GWeakObjPtrArg);
-		BIND_CALL_CHECK_DELEGATE(FTestDynamicMulticastDelegate4,            1, &UDynamicDelegateTestObject::ConstFunc6, (GInt16Arg, GStrArg, GPtrArg, GNameArg), bGBoolArg, GWeakObjPtrArg);
+		BIND_CALL_CHECK_DELEGATE(FTestDynamicMulticastDelegate4::FDelegate, 0, &UDynamicDelegateTestObject::Func6,      (GInt16Arg, GUnscopedEnumArg, GStrArg, GPtrArg), GNameArg, bGBoolArg);
+		BIND_CALL_CHECK_DELEGATE(FTestDynamicDelegate4,                     0, &UDynamicDelegateTestObject::Func6,      (GInt16Arg, GUnscopedEnumArg, GStrArg, GPtrArg), GNameArg, bGBoolArg);
+		BIND_CALL_CHECK_DELEGATE(FTestDynamicMulticastDelegate4,            1, &UDynamicDelegateTestObject::Func6,      (GInt16Arg, GUnscopedEnumArg, GStrArg, GPtrArg), GNameArg, bGBoolArg);
+		BIND_CALL_CHECK_DELEGATE(FTestDynamicMulticastDelegate4::FDelegate, 0, &UDynamicDelegateTestObject::ConstFunc6, (GInt16Arg, GUnscopedEnumArg, GStrArg, GPtrArg), GNameArg, bGBoolArg);
+		BIND_CALL_CHECK_DELEGATE(FTestDynamicDelegate4,                     0, &UDynamicDelegateTestObject::ConstFunc6, (GInt16Arg, GUnscopedEnumArg, GStrArg, GPtrArg), GNameArg, bGBoolArg);
+		BIND_CALL_CHECK_DELEGATE(FTestDynamicMulticastDelegate4,            1, &UDynamicDelegateTestObject::ConstFunc6, (GInt16Arg, GUnscopedEnumArg, GStrArg, GPtrArg), GNameArg, bGBoolArg);
 #if UE_USE_DYNAMIC_DELEGATE_PAYLOADS_RETVAL
-		BIND_CALL_CHECK_DELEGATE_RETVAL(GIntRet, FTestDynamicDelegateRetInt4, 0, &UDynamicDelegateTestObject::FuncRetInt6,      (GInt16Arg, GStrArg, GPtrArg, GNameArg), bGBoolArg, GWeakObjPtrArg);
-		BIND_CALL_CHECK_DELEGATE_RETVAL(GIntRet, FTestDynamicDelegateRetInt4, 0, &UDynamicDelegateTestObject::ConstFuncRetInt6, (GInt16Arg, GStrArg, GPtrArg, GNameArg), bGBoolArg, GWeakObjPtrArg);
-		BIND_CALL_CHECK_DELEGATE_RETVAL(GStrRet, FTestDynamicDelegateRetStr4, 0, &UDynamicDelegateTestObject::FuncRetStr6,      (GInt16Arg, GStrArg, GPtrArg, GNameArg), bGBoolArg, GWeakObjPtrArg);
-		BIND_CALL_CHECK_DELEGATE_RETVAL(GStrRet, FTestDynamicDelegateRetStr4, 0, &UDynamicDelegateTestObject::ConstFuncRetStr6, (GInt16Arg, GStrArg, GPtrArg, GNameArg), bGBoolArg, GWeakObjPtrArg);
+		BIND_CALL_CHECK_DELEGATE_RETVAL(GIntRet, FTestDynamicDelegateRetInt4, 0, &UDynamicDelegateTestObject::FuncRetInt6,      (GInt16Arg, GUnscopedEnumArg, GStrArg, GPtrArg), GNameArg, bGBoolArg);
+		BIND_CALL_CHECK_DELEGATE_RETVAL(GIntRet, FTestDynamicDelegateRetInt4, 0, &UDynamicDelegateTestObject::ConstFuncRetInt6, (GInt16Arg, GUnscopedEnumArg, GStrArg, GPtrArg), GNameArg, bGBoolArg);
+		BIND_CALL_CHECK_DELEGATE_RETVAL(GStrRet, FTestDynamicDelegateRetStr4, 0, &UDynamicDelegateTestObject::FuncRetStr6,      (GInt16Arg, GUnscopedEnumArg, GStrArg, GPtrArg), GNameArg, bGBoolArg);
+		BIND_CALL_CHECK_DELEGATE_RETVAL(GStrRet, FTestDynamicDelegateRetStr4, 0, &UDynamicDelegateTestObject::ConstFuncRetStr6, (GInt16Arg, GUnscopedEnumArg, GStrArg, GPtrArg), GNameArg, bGBoolArg);
 #endif
 
 		// 3 arg payload
-		BIND_CALL_CHECK_DELEGATE(FTestDynamicMulticastDelegate3::FDelegate, 0, &UDynamicDelegateTestObject::Func6,      (GInt16Arg, GStrArg, GPtrArg), GNameArg, bGBoolArg, GWeakObjPtrArg);
-		BIND_CALL_CHECK_DELEGATE(FTestDynamicDelegate3,                     0, &UDynamicDelegateTestObject::Func6,      (GInt16Arg, GStrArg, GPtrArg), GNameArg, bGBoolArg, GWeakObjPtrArg);
-		BIND_CALL_CHECK_DELEGATE(FTestDynamicMulticastDelegate3,            1, &UDynamicDelegateTestObject::Func6,      (GInt16Arg, GStrArg, GPtrArg), GNameArg, bGBoolArg, GWeakObjPtrArg);
-		BIND_CALL_CHECK_DELEGATE(FTestDynamicMulticastDelegate3::FDelegate, 0, &UDynamicDelegateTestObject::ConstFunc6, (GInt16Arg, GStrArg, GPtrArg), GNameArg, bGBoolArg, GWeakObjPtrArg);
-		BIND_CALL_CHECK_DELEGATE(FTestDynamicDelegate3,                     0, &UDynamicDelegateTestObject::ConstFunc6, (GInt16Arg, GStrArg, GPtrArg), GNameArg, bGBoolArg, GWeakObjPtrArg);
-		BIND_CALL_CHECK_DELEGATE(FTestDynamicMulticastDelegate3,            1, &UDynamicDelegateTestObject::ConstFunc6, (GInt16Arg, GStrArg, GPtrArg), GNameArg, bGBoolArg, GWeakObjPtrArg);
+		BIND_CALL_CHECK_DELEGATE(FTestDynamicMulticastDelegate3::FDelegate, 0, &UDynamicDelegateTestObject::Func6,      (GInt16Arg, GUnscopedEnumArg, GStrArg), GPtrArg, GNameArg, bGBoolArg);
+		BIND_CALL_CHECK_DELEGATE(FTestDynamicDelegate3,                     0, &UDynamicDelegateTestObject::Func6,      (GInt16Arg, GUnscopedEnumArg, GStrArg), GPtrArg, GNameArg, bGBoolArg);
+		BIND_CALL_CHECK_DELEGATE(FTestDynamicMulticastDelegate3,            1, &UDynamicDelegateTestObject::Func6,      (GInt16Arg, GUnscopedEnumArg, GStrArg), GPtrArg, GNameArg, bGBoolArg);
+		BIND_CALL_CHECK_DELEGATE(FTestDynamicMulticastDelegate3::FDelegate, 0, &UDynamicDelegateTestObject::ConstFunc6, (GInt16Arg, GUnscopedEnumArg, GStrArg), GPtrArg, GNameArg, bGBoolArg);
+		BIND_CALL_CHECK_DELEGATE(FTestDynamicDelegate3,                     0, &UDynamicDelegateTestObject::ConstFunc6, (GInt16Arg, GUnscopedEnumArg, GStrArg), GPtrArg, GNameArg, bGBoolArg);
+		BIND_CALL_CHECK_DELEGATE(FTestDynamicMulticastDelegate3,            1, &UDynamicDelegateTestObject::ConstFunc6, (GInt16Arg, GUnscopedEnumArg, GStrArg), GPtrArg, GNameArg, bGBoolArg);
 #if UE_USE_DYNAMIC_DELEGATE_PAYLOADS_RETVAL
-		BIND_CALL_CHECK_DELEGATE_RETVAL(GIntRet, FTestDynamicDelegateRetInt3, 0, &UDynamicDelegateTestObject::FuncRetInt6,      (GInt16Arg, GStrArg, GPtrArg), GNameArg, bGBoolArg, GWeakObjPtrArg);
-		BIND_CALL_CHECK_DELEGATE_RETVAL(GIntRet, FTestDynamicDelegateRetInt3, 0, &UDynamicDelegateTestObject::ConstFuncRetInt6, (GInt16Arg, GStrArg, GPtrArg), GNameArg, bGBoolArg, GWeakObjPtrArg);
-		BIND_CALL_CHECK_DELEGATE_RETVAL(GStrRet, FTestDynamicDelegateRetStr3, 0, &UDynamicDelegateTestObject::FuncRetStr6,      (GInt16Arg, GStrArg, GPtrArg), GNameArg, bGBoolArg, GWeakObjPtrArg);
-		BIND_CALL_CHECK_DELEGATE_RETVAL(GStrRet, FTestDynamicDelegateRetStr3, 0, &UDynamicDelegateTestObject::ConstFuncRetStr6, (GInt16Arg, GStrArg, GPtrArg), GNameArg, bGBoolArg, GWeakObjPtrArg);
+		BIND_CALL_CHECK_DELEGATE_RETVAL(GIntRet, FTestDynamicDelegateRetInt3, 0, &UDynamicDelegateTestObject::FuncRetInt6,      (GInt16Arg, GUnscopedEnumArg, GStrArg), GPtrArg, GNameArg, bGBoolArg);
+		BIND_CALL_CHECK_DELEGATE_RETVAL(GIntRet, FTestDynamicDelegateRetInt3, 0, &UDynamicDelegateTestObject::ConstFuncRetInt6, (GInt16Arg, GUnscopedEnumArg, GStrArg), GPtrArg, GNameArg, bGBoolArg);
+		BIND_CALL_CHECK_DELEGATE_RETVAL(GStrRet, FTestDynamicDelegateRetStr3, 0, &UDynamicDelegateTestObject::FuncRetStr6,      (GInt16Arg, GUnscopedEnumArg, GStrArg), GPtrArg, GNameArg, bGBoolArg);
+		BIND_CALL_CHECK_DELEGATE_RETVAL(GStrRet, FTestDynamicDelegateRetStr3, 0, &UDynamicDelegateTestObject::ConstFuncRetStr6, (GInt16Arg, GUnscopedEnumArg, GStrArg), GPtrArg, GNameArg, bGBoolArg);
 #endif
 
 		// 4 arg payload
-		BIND_CALL_CHECK_DELEGATE(FTestDynamicMulticastDelegate2::FDelegate, 0, &UDynamicDelegateTestObject::Func6,      (GInt16Arg, GStrArg), GPtrArg, GNameArg, bGBoolArg, GWeakObjPtrArg);
-		BIND_CALL_CHECK_DELEGATE(FTestDynamicDelegate2,                     0, &UDynamicDelegateTestObject::Func6,      (GInt16Arg, GStrArg), GPtrArg, GNameArg, bGBoolArg, GWeakObjPtrArg);
-		BIND_CALL_CHECK_DELEGATE(FTestDynamicMulticastDelegate2,            1, &UDynamicDelegateTestObject::Func6,      (GInt16Arg, GStrArg), GPtrArg, GNameArg, bGBoolArg, GWeakObjPtrArg);
-		BIND_CALL_CHECK_DELEGATE(FTestDynamicMulticastDelegate2::FDelegate, 0, &UDynamicDelegateTestObject::ConstFunc6, (GInt16Arg, GStrArg), GPtrArg, GNameArg, bGBoolArg, GWeakObjPtrArg);
-		BIND_CALL_CHECK_DELEGATE(FTestDynamicDelegate2,                     0, &UDynamicDelegateTestObject::ConstFunc6, (GInt16Arg, GStrArg), GPtrArg, GNameArg, bGBoolArg, GWeakObjPtrArg);
-		BIND_CALL_CHECK_DELEGATE(FTestDynamicMulticastDelegate2,            1, &UDynamicDelegateTestObject::ConstFunc6, (GInt16Arg, GStrArg), GPtrArg, GNameArg, bGBoolArg, GWeakObjPtrArg);
+		BIND_CALL_CHECK_DELEGATE(FTestDynamicMulticastDelegate2::FDelegate, 0, &UDynamicDelegateTestObject::Func6,      (GInt16Arg, GUnscopedEnumArg), GStrArg, GPtrArg, GNameArg, bGBoolArg);
+		BIND_CALL_CHECK_DELEGATE(FTestDynamicDelegate2,                     0, &UDynamicDelegateTestObject::Func6,      (GInt16Arg, GUnscopedEnumArg), GStrArg, GPtrArg, GNameArg, bGBoolArg);
+		BIND_CALL_CHECK_DELEGATE(FTestDynamicMulticastDelegate2,            1, &UDynamicDelegateTestObject::Func6,      (GInt16Arg, GUnscopedEnumArg), GStrArg, GPtrArg, GNameArg, bGBoolArg);
+		BIND_CALL_CHECK_DELEGATE(FTestDynamicMulticastDelegate2::FDelegate, 0, &UDynamicDelegateTestObject::ConstFunc6, (GInt16Arg, GUnscopedEnumArg), GStrArg, GPtrArg, GNameArg, bGBoolArg);
+		BIND_CALL_CHECK_DELEGATE(FTestDynamicDelegate2,                     0, &UDynamicDelegateTestObject::ConstFunc6, (GInt16Arg, GUnscopedEnumArg), GStrArg, GPtrArg, GNameArg, bGBoolArg);
+		BIND_CALL_CHECK_DELEGATE(FTestDynamicMulticastDelegate2,            1, &UDynamicDelegateTestObject::ConstFunc6, (GInt16Arg, GUnscopedEnumArg), GStrArg, GPtrArg, GNameArg, bGBoolArg);
 #if UE_USE_DYNAMIC_DELEGATE_PAYLOADS_RETVAL
-		BIND_CALL_CHECK_DELEGATE_RETVAL(GIntRet, FTestDynamicDelegateRetInt2, 0, &UDynamicDelegateTestObject::FuncRetInt6,      (GInt16Arg, GStrArg), GPtrArg, GNameArg, bGBoolArg, GWeakObjPtrArg);
-		BIND_CALL_CHECK_DELEGATE_RETVAL(GIntRet, FTestDynamicDelegateRetInt2, 0, &UDynamicDelegateTestObject::ConstFuncRetInt6, (GInt16Arg, GStrArg), GPtrArg, GNameArg, bGBoolArg, GWeakObjPtrArg);
-		BIND_CALL_CHECK_DELEGATE_RETVAL(GStrRet, FTestDynamicDelegateRetStr2, 0, &UDynamicDelegateTestObject::FuncRetStr6,      (GInt16Arg, GStrArg), GPtrArg, GNameArg, bGBoolArg, GWeakObjPtrArg);
-		BIND_CALL_CHECK_DELEGATE_RETVAL(GStrRet, FTestDynamicDelegateRetStr2, 0, &UDynamicDelegateTestObject::ConstFuncRetStr6, (GInt16Arg, GStrArg), GPtrArg, GNameArg, bGBoolArg, GWeakObjPtrArg);
+		BIND_CALL_CHECK_DELEGATE_RETVAL(GIntRet, FTestDynamicDelegateRetInt2, 0, &UDynamicDelegateTestObject::FuncRetInt6,      (GInt16Arg, GUnscopedEnumArg), GStrArg, GPtrArg, GNameArg, bGBoolArg);
+		BIND_CALL_CHECK_DELEGATE_RETVAL(GIntRet, FTestDynamicDelegateRetInt2, 0, &UDynamicDelegateTestObject::ConstFuncRetInt6, (GInt16Arg, GUnscopedEnumArg), GStrArg, GPtrArg, GNameArg, bGBoolArg);
+		BIND_CALL_CHECK_DELEGATE_RETVAL(GStrRet, FTestDynamicDelegateRetStr2, 0, &UDynamicDelegateTestObject::FuncRetStr6,      (GInt16Arg, GUnscopedEnumArg), GStrArg, GPtrArg, GNameArg, bGBoolArg);
+		BIND_CALL_CHECK_DELEGATE_RETVAL(GStrRet, FTestDynamicDelegateRetStr2, 0, &UDynamicDelegateTestObject::ConstFuncRetStr6, (GInt16Arg, GUnscopedEnumArg), GStrArg, GPtrArg, GNameArg, bGBoolArg);
 #endif
 
 		// 5 arg payload
-		BIND_CALL_CHECK_DELEGATE(FTestDynamicMulticastDelegate1::FDelegate, 0, &UDynamicDelegateTestObject::Func6,      (GInt16Arg), GStrArg, GPtrArg, GNameArg, bGBoolArg, GWeakObjPtrArg);
-		BIND_CALL_CHECK_DELEGATE(FTestDynamicDelegate1,                     0, &UDynamicDelegateTestObject::Func6,      (GInt16Arg), GStrArg, GPtrArg, GNameArg, bGBoolArg, GWeakObjPtrArg);
-		BIND_CALL_CHECK_DELEGATE(FTestDynamicMulticastDelegate1,            1, &UDynamicDelegateTestObject::Func6,      (GInt16Arg), GStrArg, GPtrArg, GNameArg, bGBoolArg, GWeakObjPtrArg);
-		BIND_CALL_CHECK_DELEGATE(FTestDynamicMulticastDelegate1::FDelegate, 0, &UDynamicDelegateTestObject::ConstFunc6, (GInt16Arg), GStrArg, GPtrArg, GNameArg, bGBoolArg, GWeakObjPtrArg);
-		BIND_CALL_CHECK_DELEGATE(FTestDynamicDelegate1,                     0, &UDynamicDelegateTestObject::ConstFunc6, (GInt16Arg), GStrArg, GPtrArg, GNameArg, bGBoolArg, GWeakObjPtrArg);
-		BIND_CALL_CHECK_DELEGATE(FTestDynamicMulticastDelegate1,            1, &UDynamicDelegateTestObject::ConstFunc6, (GInt16Arg), GStrArg, GPtrArg, GNameArg, bGBoolArg, GWeakObjPtrArg);
+		BIND_CALL_CHECK_DELEGATE(FTestDynamicMulticastDelegate1::FDelegate, 0, &UDynamicDelegateTestObject::Func6,      (GInt16Arg), GUnscopedEnumArg, GStrArg, GPtrArg, GNameArg, bGBoolArg);
+		BIND_CALL_CHECK_DELEGATE(FTestDynamicDelegate1,                     0, &UDynamicDelegateTestObject::Func6,      (GInt16Arg), GUnscopedEnumArg, GStrArg, GPtrArg, GNameArg, bGBoolArg);
+		BIND_CALL_CHECK_DELEGATE(FTestDynamicMulticastDelegate1,            1, &UDynamicDelegateTestObject::Func6,      (GInt16Arg), GUnscopedEnumArg, GStrArg, GPtrArg, GNameArg, bGBoolArg);
+		BIND_CALL_CHECK_DELEGATE(FTestDynamicMulticastDelegate1::FDelegate, 0, &UDynamicDelegateTestObject::ConstFunc6, (GInt16Arg), GUnscopedEnumArg, GStrArg, GPtrArg, GNameArg, bGBoolArg);
+		BIND_CALL_CHECK_DELEGATE(FTestDynamicDelegate1,                     0, &UDynamicDelegateTestObject::ConstFunc6, (GInt16Arg), GUnscopedEnumArg, GStrArg, GPtrArg, GNameArg, bGBoolArg);
+		BIND_CALL_CHECK_DELEGATE(FTestDynamicMulticastDelegate1,            1, &UDynamicDelegateTestObject::ConstFunc6, (GInt16Arg), GUnscopedEnumArg, GStrArg, GPtrArg, GNameArg, bGBoolArg);
 #if UE_USE_DYNAMIC_DELEGATE_PAYLOADS_RETVAL
-		BIND_CALL_CHECK_DELEGATE_RETVAL(GIntRet, FTestDynamicDelegateRetInt1, 0, &UDynamicDelegateTestObject::FuncRetInt6,      (GInt16Arg), GStrArg, GPtrArg, GNameArg, bGBoolArg, GWeakObjPtrArg);
-		BIND_CALL_CHECK_DELEGATE_RETVAL(GIntRet, FTestDynamicDelegateRetInt1, 0, &UDynamicDelegateTestObject::ConstFuncRetInt6, (GInt16Arg), GStrArg, GPtrArg, GNameArg, bGBoolArg, GWeakObjPtrArg);
-		BIND_CALL_CHECK_DELEGATE_RETVAL(GStrRet, FTestDynamicDelegateRetStr1, 0, &UDynamicDelegateTestObject::FuncRetStr6,      (GInt16Arg), GStrArg, GPtrArg, GNameArg, bGBoolArg, GWeakObjPtrArg);
-		BIND_CALL_CHECK_DELEGATE_RETVAL(GStrRet, FTestDynamicDelegateRetStr1, 0, &UDynamicDelegateTestObject::ConstFuncRetStr6, (GInt16Arg), GStrArg, GPtrArg, GNameArg, bGBoolArg, GWeakObjPtrArg);
+		BIND_CALL_CHECK_DELEGATE_RETVAL(GIntRet, FTestDynamicDelegateRetInt1, 0, &UDynamicDelegateTestObject::FuncRetInt6,      (GInt16Arg), GUnscopedEnumArg, GStrArg, GPtrArg, GNameArg, bGBoolArg);
+		BIND_CALL_CHECK_DELEGATE_RETVAL(GIntRet, FTestDynamicDelegateRetInt1, 0, &UDynamicDelegateTestObject::ConstFuncRetInt6, (GInt16Arg), GUnscopedEnumArg, GStrArg, GPtrArg, GNameArg, bGBoolArg);
+		BIND_CALL_CHECK_DELEGATE_RETVAL(GStrRet, FTestDynamicDelegateRetStr1, 0, &UDynamicDelegateTestObject::FuncRetStr6,      (GInt16Arg), GUnscopedEnumArg, GStrArg, GPtrArg, GNameArg, bGBoolArg);
+		BIND_CALL_CHECK_DELEGATE_RETVAL(GStrRet, FTestDynamicDelegateRetStr1, 0, &UDynamicDelegateTestObject::ConstFuncRetStr6, (GInt16Arg), GUnscopedEnumArg, GStrArg, GPtrArg, GNameArg, bGBoolArg);
 #endif
 
 		// 6 arg payload
-		BIND_CALL_CHECK_DELEGATE(FTestDynamicMulticastDelegate0::FDelegate, 0, &UDynamicDelegateTestObject::Func6,      (), GInt16Arg, GStrArg, GPtrArg, GNameArg, bGBoolArg, GWeakObjPtrArg);
-		BIND_CALL_CHECK_DELEGATE(FTestDynamicDelegate0,                     0, &UDynamicDelegateTestObject::Func6,      (), GInt16Arg, GStrArg, GPtrArg, GNameArg, bGBoolArg, GWeakObjPtrArg);
-		BIND_CALL_CHECK_DELEGATE(FTestDynamicMulticastDelegate0,            1, &UDynamicDelegateTestObject::Func6,      (), GInt16Arg, GStrArg, GPtrArg, GNameArg, bGBoolArg, GWeakObjPtrArg);
-		BIND_CALL_CHECK_DELEGATE(FTestDynamicMulticastDelegate0::FDelegate, 0, &UDynamicDelegateTestObject::ConstFunc6, (), GInt16Arg, GStrArg, GPtrArg, GNameArg, bGBoolArg, GWeakObjPtrArg);
-		BIND_CALL_CHECK_DELEGATE(FTestDynamicDelegate0,                     0, &UDynamicDelegateTestObject::ConstFunc6, (), GInt16Arg, GStrArg, GPtrArg, GNameArg, bGBoolArg, GWeakObjPtrArg);
-		BIND_CALL_CHECK_DELEGATE(FTestDynamicMulticastDelegate0,            1, &UDynamicDelegateTestObject::ConstFunc6, (), GInt16Arg, GStrArg, GPtrArg, GNameArg, bGBoolArg, GWeakObjPtrArg);
+		BIND_CALL_CHECK_DELEGATE(FTestDynamicMulticastDelegate0::FDelegate, 0, &UDynamicDelegateTestObject::Func6,      (), GInt16Arg, GUnscopedEnumArg, GStrArg, GPtrArg, GNameArg, bGBoolArg);
+		BIND_CALL_CHECK_DELEGATE(FTestDynamicDelegate0,                     0, &UDynamicDelegateTestObject::Func6,      (), GInt16Arg, GUnscopedEnumArg, GStrArg, GPtrArg, GNameArg, bGBoolArg);
+		BIND_CALL_CHECK_DELEGATE(FTestDynamicMulticastDelegate0,            1, &UDynamicDelegateTestObject::Func6,      (), GInt16Arg, GUnscopedEnumArg, GStrArg, GPtrArg, GNameArg, bGBoolArg);
+		BIND_CALL_CHECK_DELEGATE(FTestDynamicMulticastDelegate0::FDelegate, 0, &UDynamicDelegateTestObject::ConstFunc6, (), GInt16Arg, GUnscopedEnumArg, GStrArg, GPtrArg, GNameArg, bGBoolArg);
+		BIND_CALL_CHECK_DELEGATE(FTestDynamicDelegate0,                     0, &UDynamicDelegateTestObject::ConstFunc6, (), GInt16Arg, GUnscopedEnumArg, GStrArg, GPtrArg, GNameArg, bGBoolArg);
+		BIND_CALL_CHECK_DELEGATE(FTestDynamicMulticastDelegate0,            1, &UDynamicDelegateTestObject::ConstFunc6, (), GInt16Arg, GUnscopedEnumArg, GStrArg, GPtrArg, GNameArg, bGBoolArg);
+#if UE_USE_DYNAMIC_DELEGATE_PAYLOADS_RETVAL
+		BIND_CALL_CHECK_DELEGATE_RETVAL(GIntRet, FTestDynamicDelegateRetInt0, 0, &UDynamicDelegateTestObject::FuncRetInt6,      (), GInt16Arg, GUnscopedEnumArg, GStrArg, GPtrArg, GNameArg, bGBoolArg);
+		BIND_CALL_CHECK_DELEGATE_RETVAL(GIntRet, FTestDynamicDelegateRetInt0, 0, &UDynamicDelegateTestObject::ConstFuncRetInt6, (), GInt16Arg, GUnscopedEnumArg, GStrArg, GPtrArg, GNameArg, bGBoolArg);
+		BIND_CALL_CHECK_DELEGATE_RETVAL(GStrRet, FTestDynamicDelegateRetStr0, 0, &UDynamicDelegateTestObject::FuncRetStr6,      (), GInt16Arg, GUnscopedEnumArg, GStrArg, GPtrArg, GNameArg, bGBoolArg);
+		BIND_CALL_CHECK_DELEGATE_RETVAL(GStrRet, FTestDynamicDelegateRetStr0, 0, &UDynamicDelegateTestObject::ConstFuncRetStr6, (), GInt16Arg, GUnscopedEnumArg, GStrArg, GPtrArg, GNameArg, bGBoolArg);
+#endif
+#endif
+	}
+
+	// 7 param function
+	{
+		// 0 arg payload
+		BIND_CALL_CHECK_DELEGATE(FTestDynamicMulticastDelegate7::FDelegate, 0, &UDynamicDelegateTestObject::Func7,      (GInt16Arg, GUnscopedEnumArg, GStrArg, GPtrArg, GNameArg, bGBoolArg, GWeakObjPtrArg));
+		BIND_CALL_CHECK_DELEGATE(FTestDynamicDelegate7,                     0, &UDynamicDelegateTestObject::Func7,      (GInt16Arg, GUnscopedEnumArg, GStrArg, GPtrArg, GNameArg, bGBoolArg, GWeakObjPtrArg));
+		BIND_CALL_CHECK_DELEGATE(FTestDynamicMulticastDelegate7,            1, &UDynamicDelegateTestObject::Func7,      (GInt16Arg, GUnscopedEnumArg, GStrArg, GPtrArg, GNameArg, bGBoolArg, GWeakObjPtrArg));
+		BIND_CALL_CHECK_DELEGATE(FTestDynamicMulticastDelegate7::FDelegate, 0, &UDynamicDelegateTestObject::ConstFunc7, (GInt16Arg, GUnscopedEnumArg, GStrArg, GPtrArg, GNameArg, bGBoolArg, GWeakObjPtrArg));
+		BIND_CALL_CHECK_DELEGATE(FTestDynamicDelegate7,                     0, &UDynamicDelegateTestObject::ConstFunc7, (GInt16Arg, GUnscopedEnumArg, GStrArg, GPtrArg, GNameArg, bGBoolArg, GWeakObjPtrArg));
+		BIND_CALL_CHECK_DELEGATE(FTestDynamicMulticastDelegate7,            1, &UDynamicDelegateTestObject::ConstFunc7, (GInt16Arg, GUnscopedEnumArg, GStrArg, GPtrArg, GNameArg, bGBoolArg, GWeakObjPtrArg));
+
+		BIND_CALL_CHECK_DELEGATE_RETVAL(GIntRet, FTestDynamicDelegateRetInt7, 0, &UDynamicDelegateTestObject::FuncRetInt7,      (GInt16Arg, GUnscopedEnumArg, GStrArg, GPtrArg, GNameArg, bGBoolArg, GWeakObjPtrArg));
+		BIND_CALL_CHECK_DELEGATE_RETVAL(GIntRet, FTestDynamicDelegateRetInt7, 0, &UDynamicDelegateTestObject::ConstFuncRetInt7, (GInt16Arg, GUnscopedEnumArg, GStrArg, GPtrArg, GNameArg, bGBoolArg, GWeakObjPtrArg));
+		BIND_CALL_CHECK_DELEGATE_RETVAL(GStrRet, FTestDynamicDelegateRetStr7, 0, &UDynamicDelegateTestObject::FuncRetStr7,      (GInt16Arg, GUnscopedEnumArg, GStrArg, GPtrArg, GNameArg, bGBoolArg, GWeakObjPtrArg));
+		BIND_CALL_CHECK_DELEGATE_RETVAL(GStrRet, FTestDynamicDelegateRetStr7, 0, &UDynamicDelegateTestObject::ConstFuncRetStr7, (GInt16Arg, GUnscopedEnumArg, GStrArg, GPtrArg, GNameArg, bGBoolArg, GWeakObjPtrArg));
+
+#if defined(UE_USE_DYNAMIC_DELEGATE_PAYLOADS) && UE_USE_DYNAMIC_DELEGATE_PAYLOADS
+		// 1 arg payload
+		BIND_CALL_CHECK_DELEGATE(FTestDynamicMulticastDelegate6::FDelegate, 0, &UDynamicDelegateTestObject::Func7,      (GInt16Arg, GUnscopedEnumArg, GStrArg, GPtrArg, GNameArg, bGBoolArg), GWeakObjPtrArg);
+		BIND_CALL_CHECK_DELEGATE(FTestDynamicDelegate6,                     0, &UDynamicDelegateTestObject::Func7,      (GInt16Arg, GUnscopedEnumArg, GStrArg, GPtrArg, GNameArg, bGBoolArg), GWeakObjPtrArg);
+		BIND_CALL_CHECK_DELEGATE(FTestDynamicMulticastDelegate6,            1, &UDynamicDelegateTestObject::Func7,      (GInt16Arg, GUnscopedEnumArg, GStrArg, GPtrArg, GNameArg, bGBoolArg), GWeakObjPtrArg);
+		BIND_CALL_CHECK_DELEGATE(FTestDynamicMulticastDelegate6::FDelegate, 0, &UDynamicDelegateTestObject::ConstFunc7, (GInt16Arg, GUnscopedEnumArg, GStrArg, GPtrArg, GNameArg, bGBoolArg), GWeakObjPtrArg);
+		BIND_CALL_CHECK_DELEGATE(FTestDynamicDelegate6,                     0, &UDynamicDelegateTestObject::ConstFunc7, (GInt16Arg, GUnscopedEnumArg, GStrArg, GPtrArg, GNameArg, bGBoolArg), GWeakObjPtrArg);
+		BIND_CALL_CHECK_DELEGATE(FTestDynamicMulticastDelegate6,            1, &UDynamicDelegateTestObject::ConstFunc7, (GInt16Arg, GUnscopedEnumArg, GStrArg, GPtrArg, GNameArg, bGBoolArg), GWeakObjPtrArg);
+#if UE_USE_DYNAMIC_DELEGATE_PAYLOADS_RETVAL
+		BIND_CALL_CHECK_DELEGATE_RETVAL(GIntRet, FTestDynamicDelegateRetInt6, 0, &UDynamicDelegateTestObject::FuncRetInt7,      (GInt16Arg, GUnscopedEnumArg, GStrArg, GPtrArg, GNameArg, bGBoolArg), GWeakObjPtrArg);
+		BIND_CALL_CHECK_DELEGATE_RETVAL(GIntRet, FTestDynamicDelegateRetInt6, 0, &UDynamicDelegateTestObject::ConstFuncRetInt7, (GInt16Arg, GUnscopedEnumArg, GStrArg, GPtrArg, GNameArg, bGBoolArg), GWeakObjPtrArg);
+		BIND_CALL_CHECK_DELEGATE_RETVAL(GStrRet, FTestDynamicDelegateRetStr6, 0, &UDynamicDelegateTestObject::FuncRetStr7,      (GInt16Arg, GUnscopedEnumArg, GStrArg, GPtrArg, GNameArg, bGBoolArg), GWeakObjPtrArg);
+		BIND_CALL_CHECK_DELEGATE_RETVAL(GStrRet, FTestDynamicDelegateRetStr6, 0, &UDynamicDelegateTestObject::ConstFuncRetStr7, (GInt16Arg, GUnscopedEnumArg, GStrArg, GPtrArg, GNameArg, bGBoolArg), GWeakObjPtrArg);
+#endif
+
+		// 2 arg payload
+		BIND_CALL_CHECK_DELEGATE(FTestDynamicMulticastDelegate5::FDelegate, 0, &UDynamicDelegateTestObject::Func7,      (GInt16Arg, GUnscopedEnumArg, GStrArg, GPtrArg, GNameArg), bGBoolArg, GWeakObjPtrArg);
+		BIND_CALL_CHECK_DELEGATE(FTestDynamicDelegate5,                     0, &UDynamicDelegateTestObject::Func7,      (GInt16Arg, GUnscopedEnumArg, GStrArg, GPtrArg, GNameArg), bGBoolArg, GWeakObjPtrArg);
+		BIND_CALL_CHECK_DELEGATE(FTestDynamicMulticastDelegate5,            1, &UDynamicDelegateTestObject::Func7,      (GInt16Arg, GUnscopedEnumArg, GStrArg, GPtrArg, GNameArg), bGBoolArg, GWeakObjPtrArg);
+		BIND_CALL_CHECK_DELEGATE(FTestDynamicMulticastDelegate5::FDelegate, 0, &UDynamicDelegateTestObject::ConstFunc7, (GInt16Arg, GUnscopedEnumArg, GStrArg, GPtrArg, GNameArg), bGBoolArg, GWeakObjPtrArg);
+		BIND_CALL_CHECK_DELEGATE(FTestDynamicDelegate5,                     0, &UDynamicDelegateTestObject::ConstFunc7, (GInt16Arg, GUnscopedEnumArg, GStrArg, GPtrArg, GNameArg), bGBoolArg, GWeakObjPtrArg);
+		BIND_CALL_CHECK_DELEGATE(FTestDynamicMulticastDelegate5,            1, &UDynamicDelegateTestObject::ConstFunc7, (GInt16Arg, GUnscopedEnumArg, GStrArg, GPtrArg, GNameArg), bGBoolArg, GWeakObjPtrArg);
+#if UE_USE_DYNAMIC_DELEGATE_PAYLOADS_RETVAL
+		BIND_CALL_CHECK_DELEGATE_RETVAL(GIntRet, FTestDynamicDelegateRetInt5, 0, &UDynamicDelegateTestObject::FuncRetInt7,      (GInt16Arg, GUnscopedEnumArg, GStrArg, GPtrArg, GNameArg), bGBoolArg, GWeakObjPtrArg);
+		BIND_CALL_CHECK_DELEGATE_RETVAL(GIntRet, FTestDynamicDelegateRetInt5, 0, &UDynamicDelegateTestObject::ConstFuncRetInt7, (GInt16Arg, GUnscopedEnumArg, GStrArg, GPtrArg, GNameArg), bGBoolArg, GWeakObjPtrArg);
+		BIND_CALL_CHECK_DELEGATE_RETVAL(GStrRet, FTestDynamicDelegateRetStr5, 0, &UDynamicDelegateTestObject::FuncRetStr7,      (GInt16Arg, GUnscopedEnumArg, GStrArg, GPtrArg, GNameArg), bGBoolArg, GWeakObjPtrArg);
+		BIND_CALL_CHECK_DELEGATE_RETVAL(GStrRet, FTestDynamicDelegateRetStr5, 0, &UDynamicDelegateTestObject::ConstFuncRetStr7, (GInt16Arg, GUnscopedEnumArg, GStrArg, GPtrArg, GNameArg), bGBoolArg, GWeakObjPtrArg);
+#endif
+
+		// 3 arg payload
+		BIND_CALL_CHECK_DELEGATE(FTestDynamicMulticastDelegate4::FDelegate, 0, &UDynamicDelegateTestObject::Func7,      (GInt16Arg, GUnscopedEnumArg, GStrArg, GPtrArg), GNameArg, bGBoolArg, GWeakObjPtrArg);
+		BIND_CALL_CHECK_DELEGATE(FTestDynamicDelegate4,                     0, &UDynamicDelegateTestObject::Func7,      (GInt16Arg, GUnscopedEnumArg, GStrArg, GPtrArg), GNameArg, bGBoolArg, GWeakObjPtrArg);
+		BIND_CALL_CHECK_DELEGATE(FTestDynamicMulticastDelegate4,            1, &UDynamicDelegateTestObject::Func7,      (GInt16Arg, GUnscopedEnumArg, GStrArg, GPtrArg), GNameArg, bGBoolArg, GWeakObjPtrArg);
+		BIND_CALL_CHECK_DELEGATE(FTestDynamicMulticastDelegate4::FDelegate, 0, &UDynamicDelegateTestObject::ConstFunc7, (GInt16Arg, GUnscopedEnumArg, GStrArg, GPtrArg), GNameArg, bGBoolArg, GWeakObjPtrArg);
+		BIND_CALL_CHECK_DELEGATE(FTestDynamicDelegate4,                     0, &UDynamicDelegateTestObject::ConstFunc7, (GInt16Arg, GUnscopedEnumArg, GStrArg, GPtrArg), GNameArg, bGBoolArg, GWeakObjPtrArg);
+		BIND_CALL_CHECK_DELEGATE(FTestDynamicMulticastDelegate4,            1, &UDynamicDelegateTestObject::ConstFunc7, (GInt16Arg, GUnscopedEnumArg, GStrArg, GPtrArg), GNameArg, bGBoolArg, GWeakObjPtrArg);
+#if UE_USE_DYNAMIC_DELEGATE_PAYLOADS_RETVAL
+		BIND_CALL_CHECK_DELEGATE_RETVAL(GIntRet, FTestDynamicDelegateRetInt4, 0, &UDynamicDelegateTestObject::FuncRetInt7,      (GInt16Arg, GUnscopedEnumArg, GStrArg, GPtrArg), GNameArg, bGBoolArg, GWeakObjPtrArg);
+		BIND_CALL_CHECK_DELEGATE_RETVAL(GIntRet, FTestDynamicDelegateRetInt4, 0, &UDynamicDelegateTestObject::ConstFuncRetInt7, (GInt16Arg, GUnscopedEnumArg, GStrArg, GPtrArg), GNameArg, bGBoolArg, GWeakObjPtrArg);
+		BIND_CALL_CHECK_DELEGATE_RETVAL(GStrRet, FTestDynamicDelegateRetStr4, 0, &UDynamicDelegateTestObject::FuncRetStr7,      (GInt16Arg, GUnscopedEnumArg, GStrArg, GPtrArg), GNameArg, bGBoolArg, GWeakObjPtrArg);
+		BIND_CALL_CHECK_DELEGATE_RETVAL(GStrRet, FTestDynamicDelegateRetStr4, 0, &UDynamicDelegateTestObject::ConstFuncRetStr7, (GInt16Arg, GUnscopedEnumArg, GStrArg, GPtrArg), GNameArg, bGBoolArg, GWeakObjPtrArg);
+#endif
+
+		// 4 arg payload
+		BIND_CALL_CHECK_DELEGATE(FTestDynamicMulticastDelegate3::FDelegate, 0, &UDynamicDelegateTestObject::Func7,      (GInt16Arg, GUnscopedEnumArg, GStrArg), GPtrArg, GNameArg, bGBoolArg, GWeakObjPtrArg);
+		BIND_CALL_CHECK_DELEGATE(FTestDynamicDelegate3,                     0, &UDynamicDelegateTestObject::Func7,      (GInt16Arg, GUnscopedEnumArg, GStrArg), GPtrArg, GNameArg, bGBoolArg, GWeakObjPtrArg);
+		BIND_CALL_CHECK_DELEGATE(FTestDynamicMulticastDelegate3,            1, &UDynamicDelegateTestObject::Func7,      (GInt16Arg, GUnscopedEnumArg, GStrArg), GPtrArg, GNameArg, bGBoolArg, GWeakObjPtrArg);
+		BIND_CALL_CHECK_DELEGATE(FTestDynamicMulticastDelegate3::FDelegate, 0, &UDynamicDelegateTestObject::ConstFunc7, (GInt16Arg, GUnscopedEnumArg, GStrArg), GPtrArg, GNameArg, bGBoolArg, GWeakObjPtrArg);
+		BIND_CALL_CHECK_DELEGATE(FTestDynamicDelegate3,                     0, &UDynamicDelegateTestObject::ConstFunc7, (GInt16Arg, GUnscopedEnumArg, GStrArg), GPtrArg, GNameArg, bGBoolArg, GWeakObjPtrArg);
+		BIND_CALL_CHECK_DELEGATE(FTestDynamicMulticastDelegate3,            1, &UDynamicDelegateTestObject::ConstFunc7, (GInt16Arg, GUnscopedEnumArg, GStrArg), GPtrArg, GNameArg, bGBoolArg, GWeakObjPtrArg);
+#if UE_USE_DYNAMIC_DELEGATE_PAYLOADS_RETVAL
+		BIND_CALL_CHECK_DELEGATE_RETVAL(GIntRet, FTestDynamicDelegateRetInt3, 0, &UDynamicDelegateTestObject::FuncRetInt7,      (GInt16Arg, GUnscopedEnumArg, GStrArg), GPtrArg, GNameArg, bGBoolArg, GWeakObjPtrArg);
+		BIND_CALL_CHECK_DELEGATE_RETVAL(GIntRet, FTestDynamicDelegateRetInt3, 0, &UDynamicDelegateTestObject::ConstFuncRetInt7, (GInt16Arg, GUnscopedEnumArg, GStrArg), GPtrArg, GNameArg, bGBoolArg, GWeakObjPtrArg);
+		BIND_CALL_CHECK_DELEGATE_RETVAL(GStrRet, FTestDynamicDelegateRetStr3, 0, &UDynamicDelegateTestObject::FuncRetStr7,      (GInt16Arg, GUnscopedEnumArg, GStrArg), GPtrArg, GNameArg, bGBoolArg, GWeakObjPtrArg);
+		BIND_CALL_CHECK_DELEGATE_RETVAL(GStrRet, FTestDynamicDelegateRetStr3, 0, &UDynamicDelegateTestObject::ConstFuncRetStr7, (GInt16Arg, GUnscopedEnumArg, GStrArg), GPtrArg, GNameArg, bGBoolArg, GWeakObjPtrArg);
+#endif
+
+		// 5 arg payload
+		BIND_CALL_CHECK_DELEGATE(FTestDynamicMulticastDelegate2::FDelegate, 0, &UDynamicDelegateTestObject::Func7,      (GInt16Arg, GUnscopedEnumArg), GStrArg, GPtrArg, GNameArg, bGBoolArg, GWeakObjPtrArg);
+		BIND_CALL_CHECK_DELEGATE(FTestDynamicDelegate2,                     0, &UDynamicDelegateTestObject::Func7,      (GInt16Arg, GUnscopedEnumArg), GStrArg, GPtrArg, GNameArg, bGBoolArg, GWeakObjPtrArg);
+		BIND_CALL_CHECK_DELEGATE(FTestDynamicMulticastDelegate2,            1, &UDynamicDelegateTestObject::Func7,      (GInt16Arg, GUnscopedEnumArg), GStrArg, GPtrArg, GNameArg, bGBoolArg, GWeakObjPtrArg);
+		BIND_CALL_CHECK_DELEGATE(FTestDynamicMulticastDelegate2::FDelegate, 0, &UDynamicDelegateTestObject::ConstFunc7, (GInt16Arg, GUnscopedEnumArg), GStrArg, GPtrArg, GNameArg, bGBoolArg, GWeakObjPtrArg);
+		BIND_CALL_CHECK_DELEGATE(FTestDynamicDelegate2,                     0, &UDynamicDelegateTestObject::ConstFunc7, (GInt16Arg, GUnscopedEnumArg), GStrArg, GPtrArg, GNameArg, bGBoolArg, GWeakObjPtrArg);
+		BIND_CALL_CHECK_DELEGATE(FTestDynamicMulticastDelegate2,            1, &UDynamicDelegateTestObject::ConstFunc7, (GInt16Arg, GUnscopedEnumArg), GStrArg, GPtrArg, GNameArg, bGBoolArg, GWeakObjPtrArg);
+#if UE_USE_DYNAMIC_DELEGATE_PAYLOADS_RETVAL
+		BIND_CALL_CHECK_DELEGATE_RETVAL(GIntRet, FTestDynamicDelegateRetInt2, 0, &UDynamicDelegateTestObject::FuncRetInt7,      (GInt16Arg, GUnscopedEnumArg), GStrArg, GPtrArg, GNameArg, bGBoolArg, GWeakObjPtrArg);
+		BIND_CALL_CHECK_DELEGATE_RETVAL(GIntRet, FTestDynamicDelegateRetInt2, 0, &UDynamicDelegateTestObject::ConstFuncRetInt7, (GInt16Arg, GUnscopedEnumArg), GStrArg, GPtrArg, GNameArg, bGBoolArg, GWeakObjPtrArg);
+		BIND_CALL_CHECK_DELEGATE_RETVAL(GStrRet, FTestDynamicDelegateRetStr2, 0, &UDynamicDelegateTestObject::FuncRetStr7,      (GInt16Arg, GUnscopedEnumArg), GStrArg, GPtrArg, GNameArg, bGBoolArg, GWeakObjPtrArg);
+		BIND_CALL_CHECK_DELEGATE_RETVAL(GStrRet, FTestDynamicDelegateRetStr2, 0, &UDynamicDelegateTestObject::ConstFuncRetStr7, (GInt16Arg, GUnscopedEnumArg), GStrArg, GPtrArg, GNameArg, bGBoolArg, GWeakObjPtrArg);
+#endif
+
+		// 6 arg payload
+		BIND_CALL_CHECK_DELEGATE(FTestDynamicMulticastDelegate1::FDelegate, 0, &UDynamicDelegateTestObject::Func7,      (GInt16Arg), GUnscopedEnumArg, GStrArg, GPtrArg, GNameArg, bGBoolArg, GWeakObjPtrArg);
+		BIND_CALL_CHECK_DELEGATE(FTestDynamicDelegate1,                     0, &UDynamicDelegateTestObject::Func7,      (GInt16Arg), GUnscopedEnumArg, GStrArg, GPtrArg, GNameArg, bGBoolArg, GWeakObjPtrArg);
+		BIND_CALL_CHECK_DELEGATE(FTestDynamicMulticastDelegate1,            1, &UDynamicDelegateTestObject::Func7,      (GInt16Arg), GUnscopedEnumArg, GStrArg, GPtrArg, GNameArg, bGBoolArg, GWeakObjPtrArg);
+		BIND_CALL_CHECK_DELEGATE(FTestDynamicMulticastDelegate1::FDelegate, 0, &UDynamicDelegateTestObject::ConstFunc7, (GInt16Arg), GUnscopedEnumArg, GStrArg, GPtrArg, GNameArg, bGBoolArg, GWeakObjPtrArg);
+		BIND_CALL_CHECK_DELEGATE(FTestDynamicDelegate1,                     0, &UDynamicDelegateTestObject::ConstFunc7, (GInt16Arg), GUnscopedEnumArg, GStrArg, GPtrArg, GNameArg, bGBoolArg, GWeakObjPtrArg);
+		BIND_CALL_CHECK_DELEGATE(FTestDynamicMulticastDelegate1,            1, &UDynamicDelegateTestObject::ConstFunc7, (GInt16Arg), GUnscopedEnumArg, GStrArg, GPtrArg, GNameArg, bGBoolArg, GWeakObjPtrArg);
+#if UE_USE_DYNAMIC_DELEGATE_PAYLOADS_RETVAL
+		BIND_CALL_CHECK_DELEGATE_RETVAL(GIntRet, FTestDynamicDelegateRetInt1, 0, &UDynamicDelegateTestObject::FuncRetInt7,      (GInt16Arg), GUnscopedEnumArg, GStrArg, GPtrArg, GNameArg, bGBoolArg, GWeakObjPtrArg);
+		BIND_CALL_CHECK_DELEGATE_RETVAL(GIntRet, FTestDynamicDelegateRetInt1, 0, &UDynamicDelegateTestObject::ConstFuncRetInt7, (GInt16Arg), GUnscopedEnumArg, GStrArg, GPtrArg, GNameArg, bGBoolArg, GWeakObjPtrArg);
+		BIND_CALL_CHECK_DELEGATE_RETVAL(GStrRet, FTestDynamicDelegateRetStr1, 0, &UDynamicDelegateTestObject::FuncRetStr7,      (GInt16Arg), GUnscopedEnumArg, GStrArg, GPtrArg, GNameArg, bGBoolArg, GWeakObjPtrArg);
+		BIND_CALL_CHECK_DELEGATE_RETVAL(GStrRet, FTestDynamicDelegateRetStr1, 0, &UDynamicDelegateTestObject::ConstFuncRetStr7, (GInt16Arg), GUnscopedEnumArg, GStrArg, GPtrArg, GNameArg, bGBoolArg, GWeakObjPtrArg);
+#endif
+
+		// 7 arg payload
+		BIND_CALL_CHECK_DELEGATE(FTestDynamicMulticastDelegate0::FDelegate, 0, &UDynamicDelegateTestObject::Func7,      (), GInt16Arg, GUnscopedEnumArg, GStrArg, GPtrArg, GNameArg, bGBoolArg, GWeakObjPtrArg);
+		BIND_CALL_CHECK_DELEGATE(FTestDynamicDelegate0,                     0, &UDynamicDelegateTestObject::Func7,      (), GInt16Arg, GUnscopedEnumArg, GStrArg, GPtrArg, GNameArg, bGBoolArg, GWeakObjPtrArg);
+		BIND_CALL_CHECK_DELEGATE(FTestDynamicMulticastDelegate0,            1, &UDynamicDelegateTestObject::Func7,      (), GInt16Arg, GUnscopedEnumArg, GStrArg, GPtrArg, GNameArg, bGBoolArg, GWeakObjPtrArg);
+		BIND_CALL_CHECK_DELEGATE(FTestDynamicMulticastDelegate0::FDelegate, 0, &UDynamicDelegateTestObject::ConstFunc7, (), GInt16Arg, GUnscopedEnumArg, GStrArg, GPtrArg, GNameArg, bGBoolArg, GWeakObjPtrArg);
+		BIND_CALL_CHECK_DELEGATE(FTestDynamicDelegate0,                     0, &UDynamicDelegateTestObject::ConstFunc7, (), GInt16Arg, GUnscopedEnumArg, GStrArg, GPtrArg, GNameArg, bGBoolArg, GWeakObjPtrArg);
+		BIND_CALL_CHECK_DELEGATE(FTestDynamicMulticastDelegate0,            1, &UDynamicDelegateTestObject::ConstFunc7, (), GInt16Arg, GUnscopedEnumArg, GStrArg, GPtrArg, GNameArg, bGBoolArg, GWeakObjPtrArg);
 #if UE_USE_DYNAMIC_DELEGATE_PAYLOADS_RETVAL
-		BIND_CALL_CHECK_DELEGATE_RETVAL(GIntRet, FTestDynamicDelegateRetInt0, 0, &UDynamicDelegateTestObject::FuncRetInt6,      (), GInt16Arg, GStrArg, GPtrArg, GNameArg, bGBoolArg, GWeakObjPtrArg);
-		BIND_CALL_CHECK_DELEGATE_RETVAL(GIntRet, FTestDynamicDelegateRetInt0, 0, &UDynamicDelegateTestObject::ConstFuncRetInt6, (), GInt16Arg, GStrArg, GPtrArg, GNameArg, bGBoolArg, GWeakObjPtrArg);
-		BIND_CALL_CHECK_DELEGATE_RETVAL(GStrRet, FTestDynamicDelegateRetStr0, 0, &UDynamicDelegateTestObject::FuncRetStr6,      (), GInt16Arg, GStrArg, GPtrArg, GNameArg, bGBoolArg, GWeakObjPtrArg);
-		BIND_CALL_CHECK_DELEGATE_RETVAL(GStrRet, FTestDynamicDelegateRetStr0, 0, &UDynamicDelegateTestObject::ConstFuncRetStr6, (), GInt16Arg, GStrArg, GPtrArg, GNameArg, bGBoolArg, GWeakObjPtrArg);
+		BIND_CALL_CHECK_DELEGATE_RETVAL(GIntRet, FTestDynamicDelegateRetInt0, 0, &UDynamicDelegateTestObject::FuncRetInt7,      (), GInt16Arg, GUnscopedEnumArg, GStrArg, GPtrArg, GNameArg, bGBoolArg, GWeakObjPtrArg);
+		BIND_CALL_CHECK_DELEGATE_RETVAL(GIntRet, FTestDynamicDelegateRetInt0, 0, &UDynamicDelegateTestObject::ConstFuncRetInt7, (), GInt16Arg, GUnscopedEnumArg, GStrArg, GPtrArg, GNameArg, bGBoolArg, GWeakObjPtrArg);
+		BIND_CALL_CHECK_DELEGATE_RETVAL(GStrRet, FTestDynamicDelegateRetStr0, 0, &UDynamicDelegateTestObject::FuncRetStr7,      (), GInt16Arg, GUnscopedEnumArg, GStrArg, GPtrArg, GNameArg, bGBoolArg, GWeakObjPtrArg);
+		BIND_CALL_CHECK_DELEGATE_RETVAL(GStrRet, FTestDynamicDelegateRetStr0, 0, &UDynamicDelegateTestObject::ConstFuncRetStr7, (), GInt16Arg, GUnscopedEnumArg, GStrArg, GPtrArg, GNameArg, bGBoolArg, GWeakObjPtrArg);
 #endif
 #endif
 	}
@@ -956,49 +1159,52 @@ void UDynamicDelegateTestObject::RunTests()
 
 	// 2 out param function
 	{
-		BIND_CALL_CHECK_DELEGATE_OUT(FTestDynamicMulticastDelegateOut2::FDelegate, 0, &UDynamicDelegateTestObject::FuncOut2,      (GInt16Arg, GStrArg), (GInt16Out, GStrOut));
-		BIND_CALL_CHECK_DELEGATE_OUT(FTestDynamicDelegateOut2,                     0, &UDynamicDelegateTestObject::FuncOut2,      (GInt16Arg, GStrArg), (GInt16Out, GStrOut));
-		BIND_CALL_CHECK_DELEGATE_OUT(FTestDynamicMulticastDelegateOut2,            1, &UDynamicDelegateTestObject::FuncOut2,      (GInt16Arg, GStrArg), (GInt16Out, GStrOut));
-		BIND_CALL_CHECK_DELEGATE_OUT(FTestDynamicMulticastDelegateOut2::FDelegate, 0, &UDynamicDelegateTestObject::ConstFuncOut2, (GInt16Arg, GStrArg), (GInt16Out, GStrOut));
-		BIND_CALL_CHECK_DELEGATE_OUT(FTestDynamicDelegateOut2,                     0, &UDynamicDelegateTestObject::ConstFuncOut2, (GInt16Arg, GStrArg), (GInt16Out, GStrOut));
-		BIND_CALL_CHECK_DELEGATE_OUT(FTestDynamicMulticastDelegateOut2,            1, &UDynamicDelegateTestObject::ConstFuncOut2, (GInt16Arg, GStrArg), (GInt16Out, GStrOut));
+		BIND_CALL_CHECK_DELEGATE_OUT(FTestDynamicMulticastDelegateOut2::FDelegate, 0, &UDynamicDelegateTestObject::FuncOut2,      (GInt16Arg, GUnscopedEnumArg), (GInt16Out, GUnscopedEnumOut));
+		BIND_CALL_CHECK_DELEGATE_OUT(FTestDynamicDelegateOut2,                     0, &UDynamicDelegateTestObject::FuncOut2,      (GInt16Arg, GUnscopedEnumArg), (GInt16Out, GUnscopedEnumOut));
+		BIND_CALL_CHECK_DELEGATE_OUT(FTestDynamicMulticastDelegateOut2,            1, &UDynamicDelegateTestObject::FuncOut2,      (GInt16Arg, GUnscopedEnumArg), (GInt16Out, GUnscopedEnumOut));
+		BIND_CALL_CHECK_DELEGATE_OUT(FTestDynamicMulticastDelegateOut2::FDelegate, 0, &UDynamicDelegateTestObject::ConstFuncOut2, (GInt16Arg, GUnscopedEnumArg), (GInt16Out, GUnscopedEnumOut));
+		BIND_CALL_CHECK_DELEGATE_OUT(FTestDynamicDelegateOut2,                     0, &UDynamicDelegateTestObject::ConstFuncOut2, (GInt16Arg, GUnscopedEnumArg), (GInt16Out, GUnscopedEnumOut));
+		BIND_CALL_CHECK_DELEGATE_OUT(FTestDynamicMulticastDelegateOut2,            1, &UDynamicDelegateTestObject::ConstFuncOut2, (GInt16Arg, GUnscopedEnumArg), (GInt16Out, GUnscopedEnumOut));
 	}
 
 	// 3 out param function
 	{
-		BIND_CALL_CHECK_DELEGATE_OUT(FTestDynamicMulticastDelegateOut3::FDelegate, 0, &UDynamicDelegateTestObject::FuncOut3,      (GInt16Arg, GStrArg, GNameArg), (GInt16Out, GStrOut, GNameOut));
-		BIND_CALL_CHECK_DELEGATE_OUT(FTestDynamicDelegateOut3,                     0, &UDynamicDelegateTestObject::FuncOut3,      (GInt16Arg, GStrArg, GNameArg), (GInt16Out, GStrOut, GNameOut));
-		BIND_CALL_CHECK_DELEGATE_OUT(FTestDynamicMulticastDelegateOut3,            1, &UDynamicDelegateTestObject::FuncOut3,      (GInt16Arg, GStrArg, GNameArg), (GInt16Out, GStrOut, GNameOut));
-		BIND_CALL_CHECK_DELEGATE_OUT(FTestDynamicMulticastDelegateOut3::FDelegate, 0, &UDynamicDelegateTestObject::ConstFuncOut3, (GInt16Arg, GStrArg, GNameArg), (GInt16Out, GStrOut, GNameOut));
-		BIND_CALL_CHECK_DELEGATE_OUT(FTestDynamicDelegateOut3,                     0, &UDynamicDelegateTestObject::ConstFuncOut3, (GInt16Arg, GStrArg, GNameArg), (GInt16Out, GStrOut, GNameOut));
-		BIND_CALL_CHECK_DELEGATE_OUT(FTestDynamicMulticastDelegateOut3,            1, &UDynamicDelegateTestObject::ConstFuncOut3, (GInt16Arg, GStrArg, GNameArg), (GInt16Out, GStrOut, GNameOut));
+		BIND_CALL_CHECK_DELEGATE_OUT(FTestDynamicMulticastDelegateOut3::FDelegate, 0, &UDynamicDelegateTestObject::FuncOut3,      (GInt16Arg, GUnscopedEnumArg, GStrArg), (GInt16Out, GUnscopedEnumOut, GStrOut));
+		BIND_CALL_CHECK_DELEGATE_OUT(FTestDynamicDelegateOut3,                     0, &UDynamicDelegateTestObject::FuncOut3,      (GInt16Arg, GUnscopedEnumArg, GStrArg), (GInt16Out, GUnscopedEnumOut, GStrOut));
+		BIND_CALL_CHECK_DELEGATE_OUT(FTestDynamicMulticastDelegateOut3,            1, &UDynamicDelegateTestObject::FuncOut3,      (GInt16Arg, GUnscopedEnumArg, GStrArg), (GInt16Out, GUnscopedEnumOut, GStrOut));
+		BIND_CALL_CHECK_DELEGATE_OUT(FTestDynamicMulticastDelegateOut3::FDelegate, 0, &UDynamicDelegateTestObject::ConstFuncOut3, (GInt16Arg, GUnscopedEnumArg, GStrArg), (GInt16Out, GUnscopedEnumOut, GStrOut));
+		BIND_CALL_CHECK_DELEGATE_OUT(FTestDynamicDelegateOut3,                     0, &UDynamicDelegateTestObject::ConstFuncOut3, (GInt16Arg, GUnscopedEnumArg, GStrArg), (GInt16Out, GUnscopedEnumOut, GStrOut));
+		BIND_CALL_CHECK_DELEGATE_OUT(FTestDynamicMulticastDelegateOut3,            1, &UDynamicDelegateTestObject::ConstFuncOut3, (GInt16Arg, GUnscopedEnumArg, GStrArg), (GInt16Out, GUnscopedEnumOut, GStrOut));
 	}
 
 	// 4 out param function
 	{
-		BIND_CALL_CHECK_DELEGATE_OUT(FTestDynamicMulticastDelegateOut4::FDelegate, 0, &UDynamicDelegateTestObject::FuncOut4,      (GInt16Arg, GStrArg, GNameArg, bGBoolArg), (GInt16Out, GStrOut, GNameOut, bGBoolOut));
-		BIND_CALL_CHECK_DELEGATE_OUT(FTestDynamicDelegateOut4,                     0, &UDynamicDelegateTestObject::FuncOut4,      (GInt16Arg, GStrArg, GNameArg, bGBoolArg), (GInt16Out, GStrOut, GNameOut, bGBoolOut));
-		BIND_CALL_CHECK_DELEGATE_OUT(FTestDynamicMulticastDelegateOut4,            1, &UDynamicDelegateTestObject::FuncOut4,      (GInt16Arg, GStrArg, GNameArg, bGBoolArg), (GInt16Out, GStrOut, GNameOut, bGBoolOut));
-		BIND_CALL_CHECK_DELEGATE_OUT(FTestDynamicMulticastDelegateOut4::FDelegate, 0, &UDynamicDelegateTestObject::ConstFuncOut4, (GInt16Arg, GStrArg, GNameArg, bGBoolArg), (GInt16Out, GStrOut, GNameOut, bGBoolOut));
-		BIND_CALL_CHECK_DELEGATE_OUT(FTestDynamicDelegateOut4,                     0, &UDynamicDelegateTestObject::ConstFuncOut4, (GInt16Arg, GStrArg, GNameArg, bGBoolArg), (GInt16Out, GStrOut, GNameOut, bGBoolOut));
-		BIND_CALL_CHECK_DELEGATE_OUT(FTestDynamicMulticastDelegateOut4,            1, &UDynamicDelegateTestObject::ConstFuncOut4, (GInt16Arg, GStrArg, GNameArg, bGBoolArg), (GInt16Out, GStrOut, GNameOut, bGBoolOut));
+		BIND_CALL_CHECK_DELEGATE_OUT(FTestDynamicMulticastDelegateOut4::FDelegate, 0, &UDynamicDelegateTestObject::FuncOut4,      (GInt16Arg, GUnscopedEnumArg, GStrArg, GNameArg), (GInt16Out, GUnscopedEnumOut, GStrOut, GNameOut));
+		BIND_CALL_CHECK_DELEGATE_OUT(FTestDynamicDelegateOut4,                     0, &UDynamicDelegateTestObject::FuncOut4,      (GInt16Arg, GUnscopedEnumArg, GStrArg, GNameArg), (GInt16Out, GUnscopedEnumOut, GStrOut, GNameOut));
+		BIND_CALL_CHECK_DELEGATE_OUT(FTestDynamicMulticastDelegateOut4,            1, &UDynamicDelegateTestObject::FuncOut4,      (GInt16Arg, GUnscopedEnumArg, GStrArg, GNameArg), (GInt16Out, GUnscopedEnumOut, GStrOut, GNameOut));
+		BIND_CALL_CHECK_DELEGATE_OUT(FTestDynamicMulticastDelegateOut4::FDelegate, 0, &UDynamicDelegateTestObject::ConstFuncOut4, (GInt16Arg, GUnscopedEnumArg, GStrArg, GNameArg), (GInt16Out, GUnscopedEnumOut, GStrOut, GNameOut));
+		BIND_CALL_CHECK_DELEGATE_OUT(FTestDynamicDelegateOut4,                     0, &UDynamicDelegateTestObject::ConstFuncOut4, (GInt16Arg, GUnscopedEnumArg, GStrArg, GNameArg), (GInt16Out, GUnscopedEnumOut, GStrOut, GNameOut));
+		BIND_CALL_CHECK_DELEGATE_OUT(FTestDynamicMulticastDelegateOut4,            1, &UDynamicDelegateTestObject::ConstFuncOut4, (GInt16Arg, GUnscopedEnumArg, GStrArg, GNameArg), (GInt16Out, GUnscopedEnumOut, GStrOut, GNameOut));
 	}
 
 	// 5 out param function
 	{
-		BIND_CALL_CHECK_DELEGATE_OUT(FTestDynamicMulticastDelegateOut5::FDelegate, 0, &UDynamicDelegateTestObject::FuncOut5,      (GInt16Arg, GStrArg, GNameArg, bGBoolArg, GWeakObjPtrArg), (GInt16Out, GStrOut, GNameOut, bGBoolOut, GWeakObjPtrOut));
-		BIND_CALL_CHECK_DELEGATE_OUT(FTestDynamicDelegateOut5,                     0, &UDynamicDelegateTestObject::FuncOut5,      (GInt16Arg, GStrArg, GNameArg, bGBoolArg, GWeakObjPtrArg), (GInt16Out, GStrOut, GNameOut, bGBoolOut, GWeakObjPtrOut));
-		BIND_CALL_CHECK_DELEGATE_OUT(FTestDynamicMulticastDelegateOut5,            1, &UDynamicDelegateTestObject::FuncOut5,      (GInt16Arg, GStrArg, GNameArg, bGBoolArg, GWeakObjPtrArg), (GInt16Out, GStrOut, GNameOut, bGBoolOut, GWeakObjPtrOut));
-		BIND_CALL_CHECK_DELEGATE_OUT(FTestDynamicMulticastDelegateOut5::FDelegate, 0, &UDynamicDelegateTestObject::ConstFuncOut5, (GInt16Arg, GStrArg, GNameArg, bGBoolArg, GWeakObjPtrArg), (GInt16Out, GStrOut, GNameOut, bGBoolOut, GWeakObjPtrOut));
-		BIND_CALL_CHECK_DELEGATE_OUT(FTestDynamicDelegateOut5,                     0, &UDynamicDelegateTestObject::ConstFuncOut5, (GInt16Arg, GStrArg, GNameArg, bGBoolArg, GWeakObjPtrArg), (GInt16Out, GStrOut, GNameOut, bGBoolOut, GWeakObjPtrOut));
-		BIND_CALL_CHECK_DELEGATE_OUT(FTestDynamicMulticastDelegateOut5,            1, &UDynamicDelegateTestObject::ConstFuncOut5, (GInt16Arg, GStrArg, GNameArg, bGBoolArg, GWeakObjPtrArg), (GInt16Out, GStrOut, GNameOut, bGBoolOut, GWeakObjPtrOut));
+		BIND_CALL_CHECK_DELEGATE_OUT(FTestDynamicMulticastDelegateOut5::FDelegate, 0, &UDynamicDelegateTestObject::FuncOut5,      (GInt16Arg, GUnscopedEnumArg, GStrArg, GNameArg, bGBoolArg), (GInt16Out, GUnscopedEnumOut, GStrOut, GNameOut, bGBoolOut));
+		BIND_CALL_CHECK_DELEGATE_OUT(FTestDynamicDelegateOut5,                     0, &UDynamicDelegateTestObject::FuncOut5,      (GInt16Arg, GUnscopedEnumArg, GStrArg, GNameArg, bGBoolArg), (GInt16Out, GUnscopedEnumOut, GStrOut, GNameOut, bGBoolOut));
+		BIND_CALL_CHECK_DELEGATE_OUT(FTestDynamicMulticastDelegateOut5,            1, &UDynamicDelegateTestObject::FuncOut5,      (GInt16Arg, GUnscopedEnumArg, GStrArg, GNameArg, bGBoolArg), (GInt16Out, GUnscopedEnumOut, GStrOut, GNameOut, bGBoolOut));
+		BIND_CALL_CHECK_DELEGATE_OUT(FTestDynamicMulticastDelegateOut5::FDelegate, 0, &UDynamicDelegateTestObject::ConstFuncOut5, (GInt16Arg, GUnscopedEnumArg, GStrArg, GNameArg, bGBoolArg), (GInt16Out, GUnscopedEnumOut, GStrOut, GNameOut, bGBoolOut));
+		BIND_CALL_CHECK_DELEGATE_OUT(FTestDynamicDelegateOut5,                     0, &UDynamicDelegateTestObject::ConstFuncOut5, (GInt16Arg, GUnscopedEnumArg, GStrArg, GNameArg, bGBoolArg), (GInt16Out, GUnscopedEnumOut, GStrOut, GNameOut, bGBoolOut));
+		BIND_CALL_CHECK_DELEGATE_OUT(FTestDynamicMulticastDelegateOut5,            1, &UDynamicDelegateTestObject::ConstFuncOut5, (GInt16Arg, GUnscopedEnumArg, GStrArg, GNameArg, bGBoolArg), (GInt16Out, GUnscopedEnumOut, GStrOut, GNameOut, bGBoolOut));
 	}
 
-	// Unscoped enum
+	// 6 out param function
 	{
-		FTestUnscopedEnumDelegate Test;
-		Test.BindDynamic(this, &UDynamicDelegateTestObject::TestUnscopedEnum);
-		Test.Execute(5, ETestUnscopedEnum::E2, 10);
+		BIND_CALL_CHECK_DELEGATE_OUT(FTestDynamicMulticastDelegateOut6::FDelegate, 0, &UDynamicDelegateTestObject::FuncOut6,      (GInt16Arg, GUnscopedEnumArg, GStrArg, GNameArg, bGBoolArg, GWeakObjPtrArg), (GInt16Out, GUnscopedEnumOut, GStrOut, GNameOut, bGBoolOut, GWeakObjPtrOut));
+		BIND_CALL_CHECK_DELEGATE_OUT(FTestDynamicDelegateOut6,                     0, &UDynamicDelegateTestObject::FuncOut6,      (GInt16Arg, GUnscopedEnumArg, GStrArg, GNameArg, bGBoolArg, GWeakObjPtrArg), (GInt16Out, GUnscopedEnumOut, GStrOut, GNameOut, bGBoolOut, GWeakObjPtrOut));
+		BIND_CALL_CHECK_DELEGATE_OUT(FTestDynamicMulticastDelegateOut6,            1, &UDynamicDelegateTestObject::FuncOut6,      (GInt16Arg, GUnscopedEnumArg, GStrArg, GNameArg, bGBoolArg, GWeakObjPtrArg), (GInt16Out, GUnscopedEnumOut, GStrOut, GNameOut, bGBoolOut, GWeakObjPtrOut));
+		BIND_CALL_CHECK_DELEGATE_OUT(FTestDynamicMulticastDelegateOut6::FDelegate, 0, &UDynamicDelegateTestObject::ConstFuncOut6, (GInt16Arg, GUnscopedEnumArg, GStrArg, GNameArg, bGBoolArg, GWeakObjPtrArg), (GInt16Out, GUnscopedEnumOut, GStrOut, GNameOut, bGBoolOut, GWeakObjPtrOut));
+		BIND_CALL_CHECK_DELEGATE_OUT(FTestDynamicDelegateOut6,                     0, &UDynamicDelegateTestObject::ConstFuncOut6, (GInt16Arg, GUnscopedEnumArg, GStrArg, GNameArg, bGBoolArg, GWeakObjPtrArg), (GInt16Out, GUnscopedEnumOut, GStrOut, GNameOut, bGBoolOut, GWeakObjPtrOut));
+		BIND_CALL_CHECK_DELEGATE_OUT(FTestDynamicMulticastDelegateOut6,            1, &UDynamicDelegateTestObject::ConstFuncOut6, (GInt16Arg, GUnscopedEnumArg, GStrArg, GNameArg, bGBoolArg, GWeakObjPtrArg), (GInt16Out, GUnscopedEnumOut, GStrOut, GNameOut, bGBoolOut, GWeakObjPtrOut));
 	}
 }
 
diff --git a/Engine/Source/Runtime/CoreUObject/Public/UObject/DynamicDelegateTestObject.h b/Engine/Source/Runtime/CoreUObject/Public/UObject/DynamicDelegateTestObject.h
--- a/Engine/Source/Runtime/CoreUObject/Public/UObject/DynamicDelegateTestObject.h
+++ b/Engine/Source/Runtime/CoreUObject/Public/UObject/DynamicDelegateTestObject.h
@@ -9,62 +9,66 @@
 
 #include "DynamicDelegateTestObject.generated.h"
 
+UENUM()
+namespace ETestUnscopedEnum
+{
+	enum Type : int
+	{
+		E1,
+		E2
+	};
+}
+
 DECLARE_DYNAMIC_DELEGATE            (FTestDynamicDelegate0);
 DECLARE_DYNAMIC_DELEGATE_OneParam   (FTestDynamicDelegate1, int16, Int);
-DECLARE_DYNAMIC_DELEGATE_TwoParams  (FTestDynamicDelegate2, int16, Int, FString, Str);
-DECLARE_DYNAMIC_DELEGATE_ThreeParams(FTestDynamicDelegate3, int16, Int, FString, Str, const UDynamicDelegateTestObject*, Ptr);
-DECLARE_DYNAMIC_DELEGATE_FourParams (FTestDynamicDelegate4, int16, Int, FString, Str, const UDynamicDelegateTestObject*, Ptr, FName, Name);
-DECLARE_DYNAMIC_DELEGATE_FiveParams (FTestDynamicDelegate5, int16, Int, FString, Str, const UDynamicDelegateTestObject*, Ptr, FName, Name, bool, bBool);
-DECLARE_DYNAMIC_DELEGATE_SixParams  (FTestDynamicDelegate6, int16, Int, FString, Str, const UDynamicDelegateTestObject*, Ptr, FName, Name, bool, bBool, TWeakObjectPtr<UDynamicDelegateTestObject>, WeakPtr);
+DECLARE_DYNAMIC_DELEGATE_TwoParams  (FTestDynamicDelegate2, int16, Int, ETestUnscopedEnum::Type, UnscopedEnum);
+DECLARE_DYNAMIC_DELEGATE_ThreeParams(FTestDynamicDelegate3, int16, Int, ETestUnscopedEnum::Type, UnscopedEnum, FString, Str);
+DECLARE_DYNAMIC_DELEGATE_FourParams (FTestDynamicDelegate4, int16, Int, ETestUnscopedEnum::Type, UnscopedEnum, FString, Str, const UDynamicDelegateTestObject*, Ptr);
+DECLARE_DYNAMIC_DELEGATE_FiveParams (FTestDynamicDelegate5, int16, Int, ETestUnscopedEnum::Type, UnscopedEnum, FString, Str, const UDynamicDelegateTestObject*, Ptr, FName, Name);
+DECLARE_DYNAMIC_DELEGATE_SixParams  (FTestDynamicDelegate6, int16, Int, ETestUnscopedEnum::Type, UnscopedEnum, FString, Str, const UDynamicDelegateTestObject*, Ptr, FName, Name, bool, bBool);
+DECLARE_DYNAMIC_DELEGATE_SevenParams(FTestDynamicDelegate7, int16, Int, ETestUnscopedEnum::Type, UnscopedEnum, FString, Str, const UDynamicDelegateTestObject*, Ptr, FName, Name, bool, bBool, TWeakObjectPtr<UDynamicDelegateTestObject>, WeakPtr);
 
 DECLARE_DYNAMIC_MULTICAST_DELEGATE            (FTestDynamicMulticastDelegate0);
 DECLARE_DYNAMIC_MULTICAST_DELEGATE_OneParam   (FTestDynamicMulticastDelegate1, int16, Int);
-DECLARE_DYNAMIC_MULTICAST_DELEGATE_TwoParams  (FTestDynamicMulticastDelegate2, int16, Int, FString, Str);
-DECLARE_DYNAMIC_MULTICAST_DELEGATE_ThreeParams(FTestDynamicMulticastDelegate3, int16, Int, FString, Str, const UDynamicDelegateTestObject*, Ptr);
-DECLARE_DYNAMIC_MULTICAST_DELEGATE_FourParams (FTestDynamicMulticastDelegate4, int16, Int, FString, Str, const UDynamicDelegateTestObject*, Ptr, FName, Name);
-DECLARE_DYNAMIC_MULTICAST_DELEGATE_FiveParams (FTestDynamicMulticastDelegate5, int16, Int, FString, Str, const UDynamicDelegateTestObject*, Ptr, FName, Name, bool, bBool);
-DECLARE_DYNAMIC_MULTICAST_DELEGATE_SixParams  (FTestDynamicMulticastDelegate6, int16, Int, FString, Str, const UDynamicDelegateTestObject*, Ptr, FName, Name, bool, bBool, TWeakObjectPtr<UDynamicDelegateTestObject>, WeakPtr);
+DECLARE_DYNAMIC_MULTICAST_DELEGATE_TwoParams  (FTestDynamicMulticastDelegate2, int16, Int, ETestUnscopedEnum::Type, UnscopedEnum);
+DECLARE_DYNAMIC_MULTICAST_DELEGATE_ThreeParams(FTestDynamicMulticastDelegate3, int16, Int, ETestUnscopedEnum::Type, UnscopedEnum, FString, Str);
+DECLARE_DYNAMIC_MULTICAST_DELEGATE_FourParams (FTestDynamicMulticastDelegate4, int16, Int, ETestUnscopedEnum::Type, UnscopedEnum, FString, Str, const UDynamicDelegateTestObject*, Ptr);
+DECLARE_DYNAMIC_MULTICAST_DELEGATE_FiveParams (FTestDynamicMulticastDelegate5, int16, Int, ETestUnscopedEnum::Type, UnscopedEnum, FString, Str, const UDynamicDelegateTestObject*, Ptr, FName, Name);
+DECLARE_DYNAMIC_MULTICAST_DELEGATE_SixParams  (FTestDynamicMulticastDelegate6, int16, Int, ETestUnscopedEnum::Type, UnscopedEnum, FString, Str, const UDynamicDelegateTestObject*, Ptr, FName, Name, bool, bBool);
+DECLARE_DYNAMIC_MULTICAST_DELEGATE_SevenParams(FTestDynamicMulticastDelegate7, int16, Int, ETestUnscopedEnum::Type, UnscopedEnum, FString, Str, const UDynamicDelegateTestObject*, Ptr, FName, Name, bool, bBool, TWeakObjectPtr<UDynamicDelegateTestObject>, WeakPtr);
 
 DECLARE_DYNAMIC_DELEGATE_RetVal            (int, FTestDynamicDelegateRetInt0);
 DECLARE_DYNAMIC_DELEGATE_RetVal_OneParam   (int, FTestDynamicDelegateRetInt1, int16, Int);
-DECLARE_DYNAMIC_DELEGATE_RetVal_TwoParams  (int, FTestDynamicDelegateRetInt2, int16, Int, FString, Str);
-DECLARE_DYNAMIC_DELEGATE_RetVal_ThreeParams(int, FTestDynamicDelegateRetInt3, int16, Int, FString, Str, const UDynamicDelegateTestObject*, Ptr);
-DECLARE_DYNAMIC_DELEGATE_RetVal_FourParams (int, FTestDynamicDelegateRetInt4, int16, Int, FString, Str, const UDynamicDelegateTestObject*, Ptr, FName, Name);
-DECLARE_DYNAMIC_DELEGATE_RetVal_FiveParams (int, FTestDynamicDelegateRetInt5, int16, Int, FString, Str, const UDynamicDelegateTestObject*, Ptr, FName, Name, bool, bBool);
-DECLARE_DYNAMIC_DELEGATE_RetVal_SixParams  (int, FTestDynamicDelegateRetInt6, int16, Int, FString, Str, const UDynamicDelegateTestObject*, Ptr, FName, Name, bool, bBool, TWeakObjectPtr<UDynamicDelegateTestObject>, WeakPtr);
+DECLARE_DYNAMIC_DELEGATE_RetVal_TwoParams  (int, FTestDynamicDelegateRetInt2, int16, Int, ETestUnscopedEnum::Type, UnscopedEnum);
+DECLARE_DYNAMIC_DELEGATE_RetVal_ThreeParams(int, FTestDynamicDelegateRetInt3, int16, Int, ETestUnscopedEnum::Type, UnscopedEnum, FString, Str);
+DECLARE_DYNAMIC_DELEGATE_RetVal_FourParams (int, FTestDynamicDelegateRetInt4, int16, Int, ETestUnscopedEnum::Type, UnscopedEnum, FString, Str, const UDynamicDelegateTestObject*, Ptr);
+DECLARE_DYNAMIC_DELEGATE_RetVal_FiveParams (int, FTestDynamicDelegateRetInt5, int16, Int, ETestUnscopedEnum::Type, UnscopedEnum, FString, Str, const UDynamicDelegateTestObject*, Ptr, FName, Name);
+DECLARE_DYNAMIC_DELEGATE_RetVal_SixParams  (int, FTestDynamicDelegateRetInt6, int16, Int, ETestUnscopedEnum::Type, UnscopedEnum, FString, Str, const UDynamicDelegateTestObject*, Ptr, FName, Name, bool, bBool);
+DECLARE_DYNAMIC_DELEGATE_RetVal_SevenParams(int, FTestDynamicDelegateRetInt7, int16, Int, ETestUnscopedEnum::Type, UnscopedEnum, FString, Str, const UDynamicDelegateTestObject*, Ptr, FName, Name, bool, bBool, TWeakObjectPtr<UDynamicDelegateTestObject>, WeakPtr);
 
 DECLARE_DYNAMIC_DELEGATE_RetVal            (FString, FTestDynamicDelegateRetStr0);
 DECLARE_DYNAMIC_DELEGATE_RetVal_OneParam   (FString, FTestDynamicDelegateRetStr1, int16, Int);
-DECLARE_DYNAMIC_DELEGATE_RetVal_TwoParams  (FString, FTestDynamicDelegateRetStr2, int16, Int, FString, Str);
-DECLARE_DYNAMIC_DELEGATE_RetVal_ThreeParams(FString, FTestDynamicDelegateRetStr3, int16, Int, FString, Str, const UDynamicDelegateTestObject*, Ptr);
-DECLARE_DYNAMIC_DELEGATE_RetVal_FourParams (FString, FTestDynamicDelegateRetStr4, int16, Int, FString, Str, const UDynamicDelegateTestObject*, Ptr, FName, Name);
-DECLARE_DYNAMIC_DELEGATE_RetVal_FiveParams (FString, FTestDynamicDelegateRetStr5, int16, Int, FString, Str, const UDynamicDelegateTestObject*, Ptr, FName, Name, bool, bBool);
-DECLARE_DYNAMIC_DELEGATE_RetVal_SixParams  (FString, FTestDynamicDelegateRetStr6, int16, Int, FString, Str, const UDynamicDelegateTestObject*, Ptr, FName, Name, bool, bBool, TWeakObjectPtr<UDynamicDelegateTestObject>, WeakPtr);
+DECLARE_DYNAMIC_DELEGATE_RetVal_TwoParams  (FString, FTestDynamicDelegateRetStr2, int16, Int, ETestUnscopedEnum::Type, UnscopedEnum);
+DECLARE_DYNAMIC_DELEGATE_RetVal_ThreeParams(FString, FTestDynamicDelegateRetStr3, int16, Int, ETestUnscopedEnum::Type, UnscopedEnum, FString, Str);
+DECLARE_DYNAMIC_DELEGATE_RetVal_FourParams (FString, FTestDynamicDelegateRetStr4, int16, Int, ETestUnscopedEnum::Type, UnscopedEnum, FString, Str, const UDynamicDelegateTestObject*, Ptr);
+DECLARE_DYNAMIC_DELEGATE_RetVal_FiveParams (FString, FTestDynamicDelegateRetStr5, int16, Int, ETestUnscopedEnum::Type, UnscopedEnum, FString, Str, const UDynamicDelegateTestObject*, Ptr, FName, Name);
+DECLARE_DYNAMIC_DELEGATE_RetVal_SixParams  (FString, FTestDynamicDelegateRetStr6, int16, Int, ETestUnscopedEnum::Type, UnscopedEnum, FString, Str, const UDynamicDelegateTestObject*, Ptr, FName, Name, bool, bBool);
+DECLARE_DYNAMIC_DELEGATE_RetVal_SevenParams(FString, FTestDynamicDelegateRetStr7, int16, Int, ETestUnscopedEnum::Type, UnscopedEnum, FString, Str, const UDynamicDelegateTestObject*, Ptr, FName, Name, bool, bBool, TWeakObjectPtr<UDynamicDelegateTestObject>, WeakPtr);
 
 DECLARE_DYNAMIC_DELEGATE_OneParam   (FTestDynamicDelegateOut1, int16&, Int);
-DECLARE_DYNAMIC_DELEGATE_TwoParams  (FTestDynamicDelegateOut2, int16&, Int, FString&, Str);
-DECLARE_DYNAMIC_DELEGATE_ThreeParams(FTestDynamicDelegateOut3, int16&, Int, FString&, Str, FName&, Name);
-DECLARE_DYNAMIC_DELEGATE_FourParams (FTestDynamicDelegateOut4, int16&, Int, FString&, Str, FName&, Name, bool&, bBool);
-DECLARE_DYNAMIC_DELEGATE_FiveParams (FTestDynamicDelegateOut5, int16&, Int, FString&, Str, FName&, Name, bool&, bBool, TWeakObjectPtr<UDynamicDelegateTestObject>&, WeakPtr);
+DECLARE_DYNAMIC_DELEGATE_TwoParams  (FTestDynamicDelegateOut2, int16&, Int, TEnumAsByte<ETestUnscopedEnum::Type>&, UnscopedEnum);
+DECLARE_DYNAMIC_DELEGATE_ThreeParams(FTestDynamicDelegateOut3, int16&, Int, TEnumAsByte<ETestUnscopedEnum::Type>&, UnscopedEnum, FString&, Str);
+DECLARE_DYNAMIC_DELEGATE_FourParams (FTestDynamicDelegateOut4, int16&, Int, TEnumAsByte<ETestUnscopedEnum::Type>&, UnscopedEnum, FString&, Str, FName&, Name);
+DECLARE_DYNAMIC_DELEGATE_FiveParams (FTestDynamicDelegateOut5, int16&, Int, TEnumAsByte<ETestUnscopedEnum::Type>&, UnscopedEnum, FString&, Str, FName&, Name, bool&, bBool);
+DECLARE_DYNAMIC_DELEGATE_SixParams  (FTestDynamicDelegateOut6, int16&, Int, TEnumAsByte<ETestUnscopedEnum::Type>&, UnscopedEnum, FString&, Str, FName&, Name, bool&, bBool, TWeakObjectPtr<UDynamicDelegateTestObject>&, WeakPtr);
 
 // Pointer out params aren't currently supported
 DECLARE_DYNAMIC_MULTICAST_DELEGATE_OneParam   (FTestDynamicMulticastDelegateOut1, int16&, Int);
-DECLARE_DYNAMIC_MULTICAST_DELEGATE_TwoParams  (FTestDynamicMulticastDelegateOut2, int16&, Int, FString&, Str);
-DECLARE_DYNAMIC_MULTICAST_DELEGATE_ThreeParams(FTestDynamicMulticastDelegateOut3, int16&, Int, FString&, Str, FName&, Name);
-DECLARE_DYNAMIC_MULTICAST_DELEGATE_FourParams (FTestDynamicMulticastDelegateOut4, int16&, Int, FString&, Str, FName&, Name, bool&, bBool);
-DECLARE_DYNAMIC_MULTICAST_DELEGATE_FiveParams (FTestDynamicMulticastDelegateOut5, int16&, Int, FString&, Str, FName&, Name, bool&, bBool, TWeakObjectPtr<UDynamicDelegateTestObject>&, WeakPtr);
-
-UENUM()
-namespace ETestUnscopedEnum
-{
-	enum Type : int
-	{
-		E1,
-		E2
-	};
-}
-
-DECLARE_DYNAMIC_DELEGATE_ThreeParams(FTestUnscopedEnumDelegate, int16, Int1, ETestUnscopedEnum::Type, E, int16, Int2);
+DECLARE_DYNAMIC_MULTICAST_DELEGATE_TwoParams  (FTestDynamicMulticastDelegateOut2, int16&, Int, TEnumAsByte<ETestUnscopedEnum::Type>&, UnscopedEnum);
+DECLARE_DYNAMIC_MULTICAST_DELEGATE_ThreeParams(FTestDynamicMulticastDelegateOut3, int16&, Int, TEnumAsByte<ETestUnscopedEnum::Type>&, UnscopedEnum, FString&, Str);
+DECLARE_DYNAMIC_MULTICAST_DELEGATE_FourParams (FTestDynamicMulticastDelegateOut4, int16&, Int, TEnumAsByte<ETestUnscopedEnum::Type>&, UnscopedEnum, FString&, Str, FName&, Name);
+DECLARE_DYNAMIC_MULTICAST_DELEGATE_FiveParams (FTestDynamicMulticastDelegateOut5, int16&, Int, TEnumAsByte<ETestUnscopedEnum::Type>&, UnscopedEnum, FString&, Str, FName&, Name, bool&, bBool);
+DECLARE_DYNAMIC_MULTICAST_DELEGATE_SixParams  (FTestDynamicMulticastDelegateOut6, int16&, Int, TEnumAsByte<ETestUnscopedEnum::Type>&, UnscopedEnum, FString&, Str, FName&, Name, bool&, bBool, TWeakObjectPtr<UDynamicDelegateTestObject>&, WeakPtr);
 
 UCLASS()
 class UDynamicDelegateTestObject : public UObject
@@ -76,67 +80,73 @@ public:
 
 	UFUNCTION() void Func0();
 	UFUNCTION() void Func1(int16 Int);
-	UFUNCTION() void Func2(int16 Int, FString Str);
-	UFUNCTION() void Func3(int16 Int, FString Str, const UDynamicDelegateTestObject* Ptr);
-	UFUNCTION() void Func4(int16 Int, FString Str, const UDynamicDelegateTestObject* Ptr, FName Name);
-	UFUNCTION() void Func5(int16 Int, FString Str, const UDynamicDelegateTestObject* Ptr, FName Name, bool bBool);
-	UFUNCTION() void Func6(int16 Int, FString Str, const UDynamicDelegateTestObject* Ptr, FName Name, bool bBool, TWeakObjectPtr<UDynamicDelegateTestObject> WeakPtr);
+	UFUNCTION() void Func2(int16 Int, ETestUnscopedEnum::Type UnscopedEnum);
+	UFUNCTION() void Func3(int16 Int, ETestUnscopedEnum::Type UnscopedEnum, FString Str);
+	UFUNCTION() void Func4(int16 Int, ETestUnscopedEnum::Type UnscopedEnum, FString Str, const UDynamicDelegateTestObject* Ptr);
+	UFUNCTION() void Func5(int16 Int, ETestUnscopedEnum::Type UnscopedEnum, FString Str, const UDynamicDelegateTestObject* Ptr, FName Name);
+	UFUNCTION() void Func6(int16 Int, ETestUnscopedEnum::Type UnscopedEnum, FString Str, const UDynamicDelegateTestObject* Ptr, FName Name, bool bBool);
+	UFUNCTION() void Func7(int16 Int, ETestUnscopedEnum::Type UnscopedEnum, FString Str, const UDynamicDelegateTestObject* Ptr, FName Name, bool bBool, TWeakObjectPtr<UDynamicDelegateTestObject> WeakPtr);
 
 	UFUNCTION() void ConstFunc0() const;
 	UFUNCTION() void ConstFunc1(int16 Int) const;
-	UFUNCTION() void ConstFunc2(int16 Int, FString Str) const;
-	UFUNCTION() void ConstFunc3(int16 Int, FString Str, const UDynamicDelegateTestObject* Ptr) const;
-	UFUNCTION() void ConstFunc4(int16 Int, FString Str, const UDynamicDelegateTestObject* Ptr, FName Name) const;
-	UFUNCTION() void ConstFunc5(int16 Int, FString Str, const UDynamicDelegateTestObject* Ptr, FName Name, bool bBool) const;
-	UFUNCTION() void ConstFunc6(int16 Int, FString Str, const UDynamicDelegateTestObject* Ptr, FName Name, bool bBool, TWeakObjectPtr<UDynamicDelegateTestObject> WeakPtr) const;
+	UFUNCTION() void ConstFunc2(int16 Int, ETestUnscopedEnum::Type UnscopedEnum) const;
+	UFUNCTION() void ConstFunc3(int16 Int, ETestUnscopedEnum::Type UnscopedEnum, FString Str) const;
+	UFUNCTION() void ConstFunc4(int16 Int, ETestUnscopedEnum::Type UnscopedEnum, FString Str, const UDynamicDelegateTestObject* Ptr) const;
+	UFUNCTION() void ConstFunc5(int16 Int, ETestUnscopedEnum::Type UnscopedEnum, FString Str, const UDynamicDelegateTestObject* Ptr, FName Name) const;
+	UFUNCTION() void ConstFunc6(int16 Int, ETestUnscopedEnum::Type UnscopedEnum, FString Str, const UDynamicDelegateTestObject* Ptr, FName Name, bool bBool) const;
+	UFUNCTION() void ConstFunc7(int16 Int, ETestUnscopedEnum::Type UnscopedEnum, FString Str, const UDynamicDelegateTestObject* Ptr, FName Name, bool bBool, TWeakObjectPtr<UDynamicDelegateTestObject> WeakPtr) const;
 
 	UFUNCTION() int FuncRetInt0();
 	UFUNCTION() int FuncRetInt1(int16 Int);
-	UFUNCTION() int FuncRetInt2(int16 Int, FString Str);
-	UFUNCTION() int FuncRetInt3(int16 Int, FString Str, const UDynamicDelegateTestObject* Ptr);
-	UFUNCTION() int FuncRetInt4(int16 Int, FString Str, const UDynamicDelegateTestObject* Ptr, FName Name);
-	UFUNCTION() int FuncRetInt5(int16 Int, FString Str, const UDynamicDelegateTestObject* Ptr, FName Name, bool bBool);
-	UFUNCTION() int FuncRetInt6(int16 Int, FString Str, const UDynamicDelegateTestObject* Ptr, FName Name, bool bBool, TWeakObjectPtr<UDynamicDelegateTestObject> WeakPtr);
+	UFUNCTION() int FuncRetInt2(int16 Int, ETestUnscopedEnum::Type UnscopedEnum);
+	UFUNCTION() int FuncRetInt3(int16 Int, ETestUnscopedEnum::Type UnscopedEnum, FString Str);
+	UFUNCTION() int FuncRetInt4(int16 Int, ETestUnscopedEnum::Type UnscopedEnum, FString Str, const UDynamicDelegateTestObject* Ptr);
+	UFUNCTION() int FuncRetInt5(int16 Int, ETestUnscopedEnum::Type UnscopedEnum, FString Str, const UDynamicDelegateTestObject* Ptr, FName Name);
+	UFUNCTION() int FuncRetInt6(int16 Int, ETestUnscopedEnum::Type UnscopedEnum, FString Str, const UDynamicDelegateTestObject* Ptr, FName Name, bool bBool);
+	UFUNCTION() int FuncRetInt7(int16 Int, ETestUnscopedEnum::Type UnscopedEnum, FString Str, const UDynamicDelegateTestObject* Ptr, FName Name, bool bBool, TWeakObjectPtr<UDynamicDelegateTestObject> WeakPtr);
 
 	UFUNCTION() FString FuncRetStr0();
 	UFUNCTION() FString FuncRetStr1(int16 Int);
-	UFUNCTION() FString FuncRetStr2(int16 Int, FString Str);
-	UFUNCTION() FString FuncRetStr3(int16 Int, FString Str, const UDynamicDelegateTestObject* Ptr);
-	UFUNCTION() FString FuncRetStr4(int16 Int, FString Str, const UDynamicDelegateTestObject* Ptr, FName Name);
-	UFUNCTION() FString FuncRetStr5(int16 Int, FString Str, const UDynamicDelegateTestObject* Ptr, FName Name, bool bBool);
-	UFUNCTION() FString FuncRetStr6(int16 Int, FString Str, const UDynamicDelegateTestObject* Ptr, FName Name, bool bBool, TWeakObjectPtr<UDynamicDelegateTestObject> WeakPtr);
+	UFUNCTION() FString FuncRetStr2(int16 Int, ETestUnscopedEnum::Type UnscopedEnum);
+	UFUNCTION() FString FuncRetStr3(int16 Int, ETestUnscopedEnum::Type UnscopedEnum, FString Str);
+	UFUNCTION() FString FuncRetStr4(int16 Int, ETestUnscopedEnum::Type UnscopedEnum, FString Str, const UDynamicDelegateTestObject* Ptr);
+	UFUNCTION() FString FuncRetStr5(int16 Int, ETestUnscopedEnum::Type UnscopedEnum, FString Str, const UDynamicDelegateTestObject* Ptr, FName Name);
+	UFUNCTION() FString FuncRetStr6(int16 Int, ETestUnscopedEnum::Type UnscopedEnum, FString Str, const UDynamicDelegateTestObject* Ptr, FName Name, bool bBool);
+	UFUNCTION() FString FuncRetStr7(int16 Int, ETestUnscopedEnum::Type UnscopedEnum, FString Str, const UDynamicDelegateTestObject* Ptr, FName Name, bool bBool, TWeakObjectPtr<UDynamicDelegateTestObject> WeakPtr);
 
 	UFUNCTION() int ConstFuncRetInt0() const;
 	UFUNCTION() int ConstFuncRetInt1(int16 Int) const;
-	UFUNCTION() int ConstFuncRetInt2(int16 Int, FString Str) const;
-	UFUNCTION() int ConstFuncRetInt3(int16 Int, FString Str, const UDynamicDelegateTestObject* Ptr) const;
-	UFUNCTION() int ConstFuncRetInt4(int16 Int, FString Str, const UDynamicDelegateTestObject* Ptr, FName Name) const;
-	UFUNCTION() int ConstFuncRetInt5(int16 Int, FString Str, const UDynamicDelegateTestObject* Ptr, FName Name, bool bBool) const;
-	UFUNCTION() int ConstFuncRetInt6(int16 Int, FString Str, const UDynamicDelegateTestObject* Ptr, FName Name, bool bBool, TWeakObjectPtr<UDynamicDelegateTestObject> WeakPtr) const;
+	UFUNCTION() int ConstFuncRetInt2(int16 Int, ETestUnscopedEnum::Type UnscopedEnum) const;
+	UFUNCTION() int ConstFuncRetInt3(int16 Int, ETestUnscopedEnum::Type UnscopedEnum, FString Str) const;
+	UFUNCTION() int ConstFuncRetInt4(int16 Int, ETestUnscopedEnum::Type UnscopedEnum, FString Str, const UDynamicDelegateTestObject* Ptr) const;
+	UFUNCTION() int ConstFuncRetInt5(int16 Int, ETestUnscopedEnum::Type UnscopedEnum, FString Str, const UDynamicDelegateTestObject* Ptr, FName Name) const;
+	UFUNCTION() int ConstFuncRetInt6(int16 Int, ETestUnscopedEnum::Type UnscopedEnum, FString Str, const UDynamicDelegateTestObject* Ptr, FName Name, bool bBool) const;
+	UFUNCTION() int ConstFuncRetInt7(int16 Int, ETestUnscopedEnum::Type UnscopedEnum, FString Str, const UDynamicDelegateTestObject* Ptr, FName Name, bool bBool, TWeakObjectPtr<UDynamicDelegateTestObject> WeakPtr) const;
 
 	UFUNCTION() FString ConstFuncRetStr0() const;
 	UFUNCTION() FString ConstFuncRetStr1(int16 Int) const;
-	UFUNCTION() FString ConstFuncRetStr2(int16 Int, FString Str) const;
-	UFUNCTION() FString ConstFuncRetStr3(int16 Int, FString Str, const UDynamicDelegateTestObject* Ptr) const;
-	UFUNCTION() FString ConstFuncRetStr4(int16 Int, FString Str, const UDynamicDelegateTestObject* Ptr, FName Name) const;
-	UFUNCTION() FString ConstFuncRetStr5(int16 Int, FString Str, const UDynamicDelegateTestObject* Ptr, FName Name, bool bBool) const;
-	UFUNCTION() FString ConstFuncRetStr6(int16 Int, FString Str, const UDynamicDelegateTestObject* Ptr, FName Name, bool bBool, TWeakObjectPtr<UDynamicDelegateTestObject> WeakPtr) const;
+	UFUNCTION() FString ConstFuncRetStr2(int16 Int, ETestUnscopedEnum::Type UnscopedEnum) const;
+	UFUNCTION() FString ConstFuncRetStr3(int16 Int, ETestUnscopedEnum::Type UnscopedEnum, FString Str) const;
+	UFUNCTION() FString ConstFuncRetStr4(int16 Int, ETestUnscopedEnum::Type UnscopedEnum, FString Str, const UDynamicDelegateTestObject* Ptr) const;
+	UFUNCTION() FString ConstFuncRetStr5(int16 Int, ETestUnscopedEnum::Type UnscopedEnum, FString Str, const UDynamicDelegateTestObject* Ptr, FName Name) const;
+	UFUNCTION() FString ConstFuncRetStr6(int16 Int, ETestUnscopedEnum::Type UnscopedEnum, FString Str, const UDynamicDelegateTestObject* Ptr, FName Name, bool bBool) const;
+	UFUNCTION() FString ConstFuncRetStr7(int16 Int, ETestUnscopedEnum::Type UnscopedEnum, FString Str, const UDynamicDelegateTestObject* Ptr, FName Name, bool bBool, TWeakObjectPtr<UDynamicDelegateTestObject> WeakPtr) const;
 
 	// Pointer out params aren't currently supported
 	UFUNCTION() void FuncOut1(int16& Int);
-	UFUNCTION() void FuncOut2(int16& Int, FString& Str);
-	UFUNCTION() void FuncOut3(int16& Int, FString& Str, FName& Name);
-	UFUNCTION() void FuncOut4(int16& Int, FString& Str, FName& Name, bool& bBool);
-	UFUNCTION() void FuncOut5(int16& Int, FString& Str, FName& Name, bool& bBool, TWeakObjectPtr<UDynamicDelegateTestObject>& WeakPtr);
+	UFUNCTION() void FuncOut2(int16& Int, TEnumAsByte<ETestUnscopedEnum::Type>& UnscopedEnum);
+	UFUNCTION() void FuncOut3(int16& Int, TEnumAsByte<ETestUnscopedEnum::Type>& UnscopedEnum, FString& Str);
+	UFUNCTION() void FuncOut4(int16& Int, TEnumAsByte<ETestUnscopedEnum::Type>& UnscopedEnum, FString& Str, FName& Name);
+	UFUNCTION() void FuncOut5(int16& Int, TEnumAsByte<ETestUnscopedEnum::Type>& UnscopedEnum, FString& Str, FName& Name, bool& bBool);
+	UFUNCTION() void FuncOut6(int16& Int, TEnumAsByte<ETestUnscopedEnum::Type>& UnscopedEnum, FString& Str, FName& Name, bool& bBool, TWeakObjectPtr<UDynamicDelegateTestObject>& WeakPtr);
 
 	// Pointer out params aren't currently supported
 	UFUNCTION() void ConstFuncOut1(int16& Int) const;
-	UFUNCTION() void ConstFuncOut2(int16& Int, FString& Str) const;
-	UFUNCTION() void ConstFuncOut3(int16& Int, FString& Str, FName& Name) const;
-	UFUNCTION() void ConstFuncOut4(int16& Int, FString& Str, FName& Name, bool& bBool) const;
-	UFUNCTION() void ConstFuncOut5(int16& Int, FString& Str, FName& Name, bool& bBool, TWeakObjectPtr<UDynamicDelegateTestObject>& WeakPtr) const;
-
-	UFUNCTION() void TestUnscopedEnum(int16 Int1, ETestUnscopedEnum::Type E, int16 Int2);
+	UFUNCTION() void ConstFuncOut2(int16& Int, TEnumAsByte<ETestUnscopedEnum::Type>& UnscopedEnum) const;
+	UFUNCTION() void ConstFuncOut3(int16& Int, TEnumAsByte<ETestUnscopedEnum::Type>& UnscopedEnum, FString& Str) const;
+	UFUNCTION() void ConstFuncOut4(int16& Int, TEnumAsByte<ETestUnscopedEnum::Type>& UnscopedEnum, FString& Str, FName& Name) const;
+	UFUNCTION() void ConstFuncOut5(int16& Int, TEnumAsByte<ETestUnscopedEnum::Type>& UnscopedEnum, FString& Str, FName& Name, bool& bBool) const;
+	UFUNCTION() void ConstFuncOut6(int16& Int, TEnumAsByte<ETestUnscopedEnum::Type>& UnscopedEnum, FString& Str, FName& Name, bool& bBool, TWeakObjectPtr<UDynamicDelegateTestObject>& WeakPtr) const;
 
 	void RunTests();
 };
