[Firefight] Fixing bad merge

diff --git a/Engine/Source/Runtime/CoreUObject/Private/Cooker/CookDependency.cpp b/Engine/Source/Runtime/CoreUObject/Private/Cooker/CookDependency.cpp
--- a/Engine/Source/Runtime/CoreUObject/Private/Cooker/CookDependency.cpp
+++ b/Engine/Source/Runtime/CoreUObject/Private/Cooker/CookDependency.cpp
@@ -225,7 +225,6 @@ FCookDependency FCookDependency::AssetRegistryQuery(FARFilter Filter)
 	Filter.SortForSaving();
 	if (!Filter.bIncludeOnlyOnDiskAssets)
 	{
-			TEXT("For the recorded dependency we are forcing bIncludeOnlyOnDiskAssets=true. Callstack of the AssetRegistryQuery construction: ")
 		UE_LOGF(LogCore, Warning,
 			"An AssetRegistryQuery dependency has an ARFilter with bIncludeOnlyOnDiskAssets=false. " "This will create cookindeterminism since the packages in can vary between cooks. Caller should set bIncludeOnlyOnDiskAssets=true. " "For the recorded dependency we are forcing bIncludeOnlyOnDiskAssets=true. Callstack of the AssetRegistryQuery construction: "
 		);
