Added support for remaining dynamic delegate payload types.

#Synced-CL 52474350

diff --git a/Engine/Source/Runtime/Core/Public/UObject/ScriptDelegates.h b/Engine/Source/Runtime/Core/Public/UObject/ScriptDelegates.h
--- a/Engine/Source/Runtime/Core/Public/UObject/ScriptDelegates.h
+++ b/Engine/Source/Runtime/Core/Public/UObject/ScriptDelegates.h
@@ -43,8 +43,10 @@
 
 #if UE_USE_DYNAMIC_DELEGATE_PAYLOADS
 	#include "Concepts/Arithmetic.h"
+	#include "Concepts/UEnum.h"
 	#include "Concepts/UnscopedEnum.h"
 	#include "Concepts/UObject.h"
+	#include "Concepts/UScriptStruct.h"
 	#include "Templates/Tuple.h"
 	#include "Traits/FunctionArity.h"
 	#include "UObject/DelegatePayloadSchema.h"
@@ -74,10 +76,10 @@ namespace UE::Core::Private
 	// Can't use TWeakObjectPtr directly here, so take a guess it being one by checking for one of its unique members
 	template <typename T>
 	concept CWeakPointerish =
-		requires(const T& Ptr)
-		{
-			Ptr.HasSameIndexAndSerialNumber(Ptr);
-		};
+		requires(const T & Ptr)
+	{
+		Ptr.HasSameIndexAndSerialNumber(Ptr);
+	};
 
 	// Wrapper for TWeakObjectPtr that implicitly converts to and from its raw pointer, for use in payload captures.
 	template <typename T, typename WeakPtrType>
@@ -88,7 +90,7 @@ namespace UE::Core::Private
 		{
 		}
 
-		UE_REWRITE operator T*() const
+		UE_REWRITE operator T* () const
 		{
 			return this->WeakObjPtr.Get();
 		}
@@ -121,19 +123,60 @@ namespace UE::Core::Private
 	// - Numeric types.
 	// - void (only for return types, not parameters, but C++ will stop you having a void parameter before getting here)
 	// - Strings and FNames.
-	template <typename WeakPtrType, UE::CUObject T>      struct TDynamicDelegateParamPayloadCaptureType<WeakPtrType, T*>          { using Type = TImplicitWeakObjectPtr<T, WeakPtrType>; };
-	template <typename WeakPtrType, UE::CArithmetic T>   struct TDynamicDelegateParamPayloadCaptureType<WeakPtrType, T>           { using Type = T; };
-	template <typename WeakPtrType, CWeakPointerish T>   struct TDynamicDelegateParamPayloadCaptureType<WeakPtrType, T>           { using Type = T; };
-	template <typename WeakPtrType>                      struct TDynamicDelegateParamPayloadCaptureType<WeakPtrType, void>        { using Type = void; };
-	template <typename WeakPtrType>                      struct TDynamicDelegateParamPayloadCaptureType<WeakPtrType, FAnsiString> { using Type = FAnsiString; };
-	template <typename WeakPtrType>                      struct TDynamicDelegateParamPayloadCaptureType<WeakPtrType, FString>     { using Type = FString; };
-	template <typename WeakPtrType>                      struct TDynamicDelegateParamPayloadCaptureType<WeakPtrType, FUtf8String> { using Type = FUtf8String; };
-	template <typename WeakPtrType>                      struct TDynamicDelegateParamPayloadCaptureType<WeakPtrType, FName>       { using Type = FName; };
+	template <typename WeakPtrType, UE::CUObject T>       struct TDynamicDelegateParamPayloadCaptureType<WeakPtrType, T*>          { using Type = TImplicitWeakObjectPtr<T, WeakPtrType>; };
+	template <typename WeakPtrType, UE::CArithmetic T>    struct TDynamicDelegateParamPayloadCaptureType<WeakPtrType, T>           { using Type = T; };
+	template <typename WeakPtrType, CWeakPointerish T>    struct TDynamicDelegateParamPayloadCaptureType<WeakPtrType, T>           { using Type = T; };
+	template <typename WeakPtrType>                       struct TDynamicDelegateParamPayloadCaptureType<WeakPtrType, void>        { using Type = void; };
+	template <typename WeakPtrType>                       struct TDynamicDelegateParamPayloadCaptureType<WeakPtrType, FAnsiString> { using Type = FAnsiString; };
+	template <typename WeakPtrType>                       struct TDynamicDelegateParamPayloadCaptureType<WeakPtrType, FString>     { using Type = FString; };
+	template <typename WeakPtrType>                       struct TDynamicDelegateParamPayloadCaptureType<WeakPtrType, FUtf8String> { using Type = FUtf8String; };
+	template <typename WeakPtrType>                       struct TDynamicDelegateParamPayloadCaptureType<WeakPtrType, FName>       { using Type = FName; };
+	template <typename WeakPtrType>                       struct TDynamicDelegateParamPayloadCaptureType<WeakPtrType, FText>       { using Type = FText; };
+	template <typename WeakPtrType, UE::CUScriptStruct T> struct TDynamicDelegateParamPayloadCaptureType<WeakPtrType, T>           { using Type = T; };
+
+	template <typename WeakPtrType, UE::CUEnum T>
+	struct TDynamicDelegateParamPayloadCaptureType<WeakPtrType, T>
+	{
 #if UE_DYNAMIC_DELEGATE_UNSCOPED_ENUM_HACK
-	template <typename WeakPtrType, UE::CUnscopedEnum T> struct TDynamicDelegateParamPayloadCaptureType<WeakPtrType, T>           { using Type = TEnumAsByte<T>; };
+		using Type = std::conditional_t<UE::CUnscopedEnum<T>, TEnumAsByte<T>, T>;
 #else
-	template <typename WeakPtrType, UE::CUnscopedEnum T> struct TDynamicDelegateParamPayloadCaptureType<WeakPtrType, T>           { using Type = T; };
+		using Type = T;
 #endif
+	};
+
+	template <typename WeakPtrType, typename T>
+	struct TDynamicDelegateParamPayloadCaptureType<WeakPtrType, TOptional<T>>
+	{
+		using Type = TOptional<T>;
+
+		// If you get a compile error here, it's because you're trying to use an optional of an unsupported payload type
+		using Validation = typename TDynamicDelegateParamPayloadCaptureType<WeakPtrType, T>::Type;
+	};
+	template <typename WeakPtrType, typename T>
+	struct TDynamicDelegateParamPayloadCaptureType<WeakPtrType, TArray<T>>
+	{
+		using Type = TArray<T>;
+
+		// If you get a compile error here, it's because you're trying to use an array of an unsupported payload type
+		using Validation = typename TDynamicDelegateParamPayloadCaptureType<WeakPtrType, T>::Type;
+	};
+	template <typename WeakPtrType, typename T>
+	struct TDynamicDelegateParamPayloadCaptureType<WeakPtrType, TSet<T>>
+	{
+		using Type = TSet<T>;
+
+		// If you get a compile error here, it's because you're trying to use a set of an unsupported payload type
+		using Validation = typename TDynamicDelegateParamPayloadCaptureType<WeakPtrType, T>::Type;
+	};
+	template <typename WeakPtrType, typename KeyType, typename ValueType>
+	struct TDynamicDelegateParamPayloadCaptureType<WeakPtrType, TMap<KeyType, ValueType>>
+	{
+		using Type = TMap<KeyType, ValueType>;
+
+		// If you get a compile error here, it's because you're trying to use a map with keys or values of an unsupported payload type
+		using KeyValidation   = typename TDynamicDelegateParamPayloadCaptureType<WeakPtrType, KeyType>::Type;
+		using ValueValidation = typename TDynamicDelegateParamPayloadCaptureType<WeakPtrType, ValueType>::Type;
+	};
 
 	template <typename WeakPtrType, typename T>
 	using TDynamicDelegateParamPayloadCaptureType_T = typename TDynamicDelegateParamPayloadCaptureType<WeakPtrType, T>::Type;
diff --git a/Engine/Source/Runtime/CoreUObject/Public/UObject/DynamicDelegateTestObject.cpp b/Engine/Source/Runtime/CoreUObject/Public/UObject/DynamicDelegateTestObject.cpp
--- a/Engine/Source/Runtime/CoreUObject/Public/UObject/DynamicDelegateTestObject.cpp
+++ b/Engine/Source/Runtime/CoreUObject/Public/UObject/DynamicDelegateTestObject.cpp
@@ -36,6 +36,39 @@ TWeakObjectPtr<UDynamicDelegateTestObject> GWeakObjPtrOut = nullptr;
 int32   GIntRet = 99;
 FString GStrRet = "World";
 
+ETestScopedEnum GScopedEnumArg = ETestScopedEnum::Value2;
+ETestScopedEnum GScopedEnumOut = ETestScopedEnum::Value1;
+
+float GFloatArg = 3.14f;
+float GFloatOut = 9.81f;
+
+double GDoubleArg = 2.718281828;
+double GDoubleOut = 1.41421356;
+
+uint8 GByteArg = 42;
+uint8 GByteOut = 255;
+
+int64 GInt64Arg = 123456789012345LL;
+int64 GInt64Out = 987654321098765LL;
+
+FText GTextArg = FText::FromString(TEXT("TestText"));
+FText GTextOut = FText::FromString(TEXT("OutputText"));
+
+FTestPayloadStruct GStructArg = { 1, 2, 3 };
+FTestPayloadStruct GStructOut = { 10, 20, 30 };
+
+FVector GVectorArg = FVector(1.0, 2.0, 3.0);
+FVector GVectorOut = FVector(4.0, 5.0, 6.0);
+
+TArray<int32> GIntArrayArg = { 10, 20, 30 };
+TArray<int32> GIntArrayOut = { 40, 50, 60 };
+
+TSet<int32> GIntSetArg = { 1, 2, 3 };
+TSet<int32> GIntSetOut = { 4, 5, 6 };
+
+TOptional<int32> GOptionalIntArg = 42;
+TOptional<int32> GOptionalIntOut = 99;
+
 UDynamicDelegateTestObject::UDynamicDelegateTestObject()
 {
 }
@@ -526,6 +559,220 @@ void UDynamicDelegateTestObject::ConstFuncOut6(int16& Int, TEnumAsByte<ETestUnsc
 	WeakObjPtr = GWeakObjPtrOut;
 }
 
+void UDynamicDelegateTestObject::FuncFloat(float Float)
+{
+	CHECK(GCalledFunctionName.IsEmpty());
+	GCalledFunctionName = "FuncFloat";
+	CHECK(Float == GFloatArg);
+}
+
+void UDynamicDelegateTestObject::FuncDouble(double Double)
+{
+	CHECK(GCalledFunctionName.IsEmpty());
+	GCalledFunctionName = "FuncDouble";
+	CHECK(Double == GDoubleArg);
+}
+
+void UDynamicDelegateTestObject::FuncByte(uint8 Byte)
+{
+	CHECK(GCalledFunctionName.IsEmpty());
+	GCalledFunctionName = "FuncByte";
+	CHECK(Byte == GByteArg);
+}
+
+void UDynamicDelegateTestObject::FuncInt64(int64 Int64)
+{
+	CHECK(GCalledFunctionName.IsEmpty());
+	GCalledFunctionName = "FuncInt64";
+	CHECK(Int64 == GInt64Arg);
+}
+
+void UDynamicDelegateTestObject::FuncText(FText Text)
+{
+	CHECK(GCalledFunctionName.IsEmpty());
+	GCalledFunctionName = "FuncText";
+	CHECK(Text.EqualTo(GTextArg));
+}
+
+void UDynamicDelegateTestObject::FuncStruct(FTestPayloadStruct Struct)
+{
+	CHECK(GCalledFunctionName.IsEmpty());
+	GCalledFunctionName = "FuncStruct";
+	CHECK(Struct == GStructArg);
+}
+
+void UDynamicDelegateTestObject::FuncVector(FVector Vector)
+{
+	CHECK(GCalledFunctionName.IsEmpty());
+	GCalledFunctionName = "FuncVector";
+	CHECK(Vector == GVectorArg);
+}
+
+void UDynamicDelegateTestObject::FuncIntArray(TArray<int32> IntArray)
+{
+	CHECK(GCalledFunctionName.IsEmpty());
+	GCalledFunctionName = "FuncIntArray";
+	CHECK(IntArray == GIntArrayArg);
+}
+
+void UDynamicDelegateTestObject::FuncScopedEnum(ETestScopedEnum Enum)
+{
+	CHECK(GCalledFunctionName.IsEmpty());
+	GCalledFunctionName = "FuncScopedEnum";
+	CHECK(Enum == GScopedEnumArg);
+}
+
+void UDynamicDelegateTestObject::FuncIntSet(TSet<int32> IntSet)
+{
+	CHECK(GCalledFunctionName.IsEmpty());
+	GCalledFunctionName = "FuncIntSet";
+	CHECK(IntSet.Num() == GIntSetArg.Num());
+	for (int32 Val : GIntSetArg)
+	{
+		CHECK(IntSet.Contains(Val));
+	}
+}
+
+void UDynamicDelegateTestObject::FuncOptionalInt(TOptional<int32> OptionalInt)
+{
+	CHECK(GCalledFunctionName.IsEmpty());
+	GCalledFunctionName = "FuncOptionalInt";
+	CHECK(OptionalInt.IsSet() == GOptionalIntArg.IsSet());
+	if (OptionalInt.IsSet())
+	{
+		CHECK(OptionalInt.GetValue() == GOptionalIntArg.GetValue());
+	}
+}
+
+void UDynamicDelegateTestObject::FuncOutFloat(float& Float)
+{
+	CHECK(GCalledFunctionName.IsEmpty());
+	GCalledFunctionName = "FuncOutFloat";
+	CHECK(Float == GFloatArg);
+	Float = GFloatOut;
+}
+
+void UDynamicDelegateTestObject::FuncOutDouble(double& Double)
+{
+	CHECK(GCalledFunctionName.IsEmpty());
+	GCalledFunctionName = "FuncOutDouble";
+	CHECK(Double == GDoubleArg);
+	Double = GDoubleOut;
+}
+
+void UDynamicDelegateTestObject::FuncOutByte(uint8& Byte)
+{
+	CHECK(GCalledFunctionName.IsEmpty());
+	GCalledFunctionName = "FuncOutByte";
+	CHECK(Byte == GByteArg);
+	Byte = GByteOut;
+}
+
+void UDynamicDelegateTestObject::FuncOutInt64(int64& Int64)
+{
+	CHECK(GCalledFunctionName.IsEmpty());
+	GCalledFunctionName = "FuncOutInt64";
+	CHECK(Int64 == GInt64Arg);
+	Int64 = GInt64Out;
+}
+
+void UDynamicDelegateTestObject::FuncOutText(FText& Text)
+{
+	CHECK(GCalledFunctionName.IsEmpty());
+	GCalledFunctionName = "FuncOutText";
+	CHECK(Text.EqualTo(GTextArg));
+	Text = GTextOut;
+}
+
+void UDynamicDelegateTestObject::FuncOutStruct(FTestPayloadStruct& Struct)
+{
+	CHECK(GCalledFunctionName.IsEmpty());
+	GCalledFunctionName = "FuncOutStruct";
+	CHECK(Struct == GStructArg);
+	Struct = GStructOut;
+}
+
+void UDynamicDelegateTestObject::FuncOutVector(FVector& Vector)
+{
+	CHECK(GCalledFunctionName.IsEmpty());
+	GCalledFunctionName = "FuncOutVector";
+	CHECK(Vector == GVectorArg);
+	Vector = GVectorOut;
+}
+
+void UDynamicDelegateTestObject::FuncOutScopedEnum(ETestScopedEnum& Enum)
+{
+	CHECK(GCalledFunctionName.IsEmpty());
+	GCalledFunctionName = "FuncOutScopedEnum";
+	CHECK(Enum == GScopedEnumArg);
+	Enum = GScopedEnumOut;
+}
+
+void UDynamicDelegateTestObject::FuncOutIntArray(TArray<int32>& IntArray)
+{
+	CHECK(GCalledFunctionName.IsEmpty());
+	GCalledFunctionName = "FuncOutIntArray";
+	CHECK(IntArray == GIntArrayArg);
+	IntArray = GIntArrayOut;
+}
+
+void UDynamicDelegateTestObject::FuncOutIntSet(TSet<int32>& IntSet)
+{
+	CHECK(GCalledFunctionName.IsEmpty());
+	GCalledFunctionName = "FuncOutIntSet";
+	CHECK(IntSet.Num() == GIntSetArg.Num());
+	for (int32 Val : GIntSetArg)
+	{
+		CHECK(IntSet.Contains(Val));
+	}
+	IntSet = GIntSetOut;
+}
+
+void UDynamicDelegateTestObject::FuncOutOptionalInt(TOptional<int32>& OptionalInt)
+{
+	CHECK(GCalledFunctionName.IsEmpty());
+	GCalledFunctionName = "FuncOutOptionalInt";
+	CHECK(OptionalInt.IsSet() == GOptionalIntArg.IsSet());
+	if (OptionalInt.IsSet())
+	{
+		CHECK(OptionalInt.GetValue() == GOptionalIntArg.GetValue());
+	}
+	OptionalInt = GOptionalIntOut;
+}
+
+// Helper function to wrap the syntactic comparison differences of different property types.
+template <typename LhsType, typename RhsType>
+static bool TestValuesEqual(const LhsType& Lhs, const RhsType& Rhs)
+{
+	if constexpr (!std::is_same_v<LhsType, RhsType>)
+	{
+		return Lhs == Rhs;
+	}
+	else if constexpr (std::is_same_v<LhsType, FText>)
+	{
+		return Lhs.EqualTo(Rhs);
+	}
+	else if constexpr (TIsTSet_V<LhsType>)
+	{
+		if (Lhs.Num() != Rhs.Num())
+		{
+			return false;
+		}
+		for (const typename LhsType::ElementType& Val : Lhs)
+		{
+			if (!Rhs.Contains(Val))
+			{
+				return false;
+			}
+		}
+		return true;
+	}
+	else
+	{
+		return Lhs == Rhs;
+	}
+}
+
 // Forwards an unscoped enum
 template <typename T>
 UE_REWRITE decltype(auto) MoveToOutArg(T&& Ref)
@@ -558,13 +805,13 @@ UE_REWRITE decltype(auto) MoveToOutArg(T&& Ref)
 #define EXPAND_ARG_VARIABLES_6( ...) (Arg1, Arg2, Arg3, Arg4, Arg5, Arg6)
 #define EXPAND_ARG_VARIABLES_7( ...) (Arg1, Arg2, Arg3, Arg4, Arg5, Arg6, Arg7)
 
-#define CHECK_ARG_VARIABLES_1(Arg1Val)                                                       CHECK(Arg1 == Arg1Val);
-#define CHECK_ARG_VARIABLES_2(Arg1Val, Arg2Val)                                              CHECK(Arg1 == Arg1Val); CHECK(Arg2 == Arg2Val);
-#define CHECK_ARG_VARIABLES_3(Arg1Val, Arg2Val, Arg3Val)                                     CHECK(Arg1 == Arg1Val); CHECK(Arg2 == Arg2Val); CHECK(Arg3 == Arg3Val);
-#define CHECK_ARG_VARIABLES_4(Arg1Val, Arg2Val, Arg3Val, Arg4Val)                            CHECK(Arg1 == Arg1Val); CHECK(Arg2 == Arg2Val); CHECK(Arg3 == Arg3Val); CHECK(Arg4 == Arg4Val);
-#define CHECK_ARG_VARIABLES_5(Arg1Val, Arg2Val, Arg3Val, Arg4Val, Arg5Val)                   CHECK(Arg1 == Arg1Val); CHECK(Arg2 == Arg2Val); CHECK(Arg3 == Arg3Val); CHECK(Arg4 == Arg4Val); CHECK(Arg5 == Arg5Val);
-#define CHECK_ARG_VARIABLES_6(Arg1Val, Arg2Val, Arg3Val, Arg4Val, Arg5Val, Arg6Val)          CHECK(Arg1 == Arg1Val); CHECK(Arg2 == Arg2Val); CHECK(Arg3 == Arg3Val); CHECK(Arg4 == Arg4Val); CHECK(Arg5 == Arg5Val); CHECK(Arg6 == Arg6Val);
-#define CHECK_ARG_VARIABLES_7(Arg1Val, Arg2Val, Arg3Val, Arg4Val, Arg5Val, Arg6Val, Arg7Val) CHECK(Arg1 == Arg1Val); CHECK(Arg2 == Arg2Val); CHECK(Arg3 == Arg3Val); CHECK(Arg4 == Arg4Val); CHECK(Arg5 == Arg5Val); CHECK(Arg6 == Arg6Val); CHECK(Arg7 == Arg7Val);
+#define CHECK_ARG_VARIABLES_1(Arg1Val)                                                       CHECK(TestValuesEqual(Arg1, Arg1Val));
+#define CHECK_ARG_VARIABLES_2(Arg1Val, Arg2Val)                                              CHECK(TestValuesEqual(Arg1, Arg1Val)); CHECK(TestValuesEqual(Arg2, Arg2Val));
+#define CHECK_ARG_VARIABLES_3(Arg1Val, Arg2Val, Arg3Val)                                     CHECK(TestValuesEqual(Arg1, Arg1Val)); CHECK(TestValuesEqual(Arg2, Arg2Val)); CHECK(TestValuesEqual(Arg3, Arg3Val));
+#define CHECK_ARG_VARIABLES_4(Arg1Val, Arg2Val, Arg3Val, Arg4Val)                            CHECK(TestValuesEqual(Arg1, Arg1Val)); CHECK(TestValuesEqual(Arg2, Arg2Val)); CHECK(TestValuesEqual(Arg3, Arg3Val)); CHECK(TestValuesEqual(Arg4, Arg4Val));
+#define CHECK_ARG_VARIABLES_5(Arg1Val, Arg2Val, Arg3Val, Arg4Val, Arg5Val)                   CHECK(TestValuesEqual(Arg1, Arg1Val)); CHECK(TestValuesEqual(Arg2, Arg2Val)); CHECK(TestValuesEqual(Arg3, Arg3Val)); CHECK(TestValuesEqual(Arg4, Arg4Val)); CHECK(TestValuesEqual(Arg5, Arg5Val));
+#define CHECK_ARG_VARIABLES_6(Arg1Val, Arg2Val, Arg3Val, Arg4Val, Arg5Val, Arg6Val)          CHECK(TestValuesEqual(Arg1, Arg1Val)); CHECK(TestValuesEqual(Arg2, Arg2Val)); CHECK(TestValuesEqual(Arg3, Arg3Val)); CHECK(TestValuesEqual(Arg4, Arg4Val)); CHECK(TestValuesEqual(Arg5, Arg5Val)); CHECK(TestValuesEqual(Arg6, Arg6Val));
+#define CHECK_ARG_VARIABLES_7(Arg1Val, Arg2Val, Arg3Val, Arg4Val, Arg5Val, Arg6Val, Arg7Val) CHECK(TestValuesEqual(Arg1, Arg1Val)); CHECK(TestValuesEqual(Arg2, Arg2Val)); CHECK(TestValuesEqual(Arg3, Arg3Val)); CHECK(TestValuesEqual(Arg4, Arg4Val)); CHECK(TestValuesEqual(Arg5, Arg5Val)); CHECK(TestValuesEqual(Arg6, Arg6Val)); CHECK(TestValuesEqual(Arg7, Arg7Val));
 
 #define DEFINE_ARG_VARIABLES(...) UE_APPEND_VA_ARG_COUNT(DEFINE_ARG_VARIABLES_, __VA_ARGS__)(__VA_ARGS__)
 #define EXPAND_ARG_VARIABLES(...) UE_APPEND_VA_ARG_COUNT(EXPAND_ARG_VARIABLES_, __VA_ARGS__)(ObjectPtr, FuncPtr, __VA_ARGS__)
@@ -608,24 +855,24 @@ void UDynamicDelegateTestObject::RunTests()
 	UDynamicDelegateTestObject* TestOut = NewObject<UDynamicDelegateTestObject>();
 
 	// Update ptr args, because they can't be given a reasonable default as a global
-	TGuardValue<const UDynamicDelegateTestObject*>          PtrArgGuard    (GPtrArg,        this);
+	TGuardValue<const UDynamicDelegateTestObject*>          PtrArgGuard(GPtrArg, this);
 	TGuardValue<TWeakObjectPtr<UDynamicDelegateTestObject>> WeakPtrArgGuard(GWeakObjPtrArg, this);
-	TGuardValue<const UDynamicDelegateTestObject*>          OutPtrArgGuard (GPtrOut,        TestOut);
+	TGuardValue<const UDynamicDelegateTestObject*>          OutPtrArgGuard(GPtrOut, TestOut);
 	TGuardValue<TWeakObjectPtr<UDynamicDelegateTestObject>> WeakPtrOutGuard(GWeakObjPtrOut, TestOut);
 
 	// 0 param function
 	{
 		// 0 arg payload
-		BIND_CALL_CHECK_DELEGATE(FTestDynamicMulticastDelegate0::FDelegate, 0, &UDynamicDelegateTestObject::Func0,      ());
-		BIND_CALL_CHECK_DELEGATE(FTestDynamicDelegate0,                     0, &UDynamicDelegateTestObject::Func0,      ());
-		BIND_CALL_CHECK_DELEGATE(FTestDynamicMulticastDelegate0,            1, &UDynamicDelegateTestObject::Func0,      ());
+		BIND_CALL_CHECK_DELEGATE(FTestDynamicMulticastDelegate0::FDelegate, 0, &UDynamicDelegateTestObject::Func0, ());
+		BIND_CALL_CHECK_DELEGATE(FTestDynamicDelegate0, 0, &UDynamicDelegateTestObject::Func0, ());
+		BIND_CALL_CHECK_DELEGATE(FTestDynamicMulticastDelegate0, 1, &UDynamicDelegateTestObject::Func0, ());
 		BIND_CALL_CHECK_DELEGATE(FTestDynamicMulticastDelegate0::FDelegate, 0, &UDynamicDelegateTestObject::ConstFunc0, ());
-		BIND_CALL_CHECK_DELEGATE(FTestDynamicDelegate0,                     0, &UDynamicDelegateTestObject::ConstFunc0, ());
-		BIND_CALL_CHECK_DELEGATE(FTestDynamicMulticastDelegate0,            1, &UDynamicDelegateTestObject::ConstFunc0, ());
+		BIND_CALL_CHECK_DELEGATE(FTestDynamicDelegate0, 0, &UDynamicDelegateTestObject::ConstFunc0, ());
+		BIND_CALL_CHECK_DELEGATE(FTestDynamicMulticastDelegate0, 1, &UDynamicDelegateTestObject::ConstFunc0, ());
 
-		BIND_CALL_CHECK_DELEGATE_RETVAL(GIntRet, FTestDynamicDelegateRetInt0, 0, &UDynamicDelegateTestObject::FuncRetInt0,      ());
+		BIND_CALL_CHECK_DELEGATE_RETVAL(GIntRet, FTestDynamicDelegateRetInt0, 0, &UDynamicDelegateTestObject::FuncRetInt0, ());
 		BIND_CALL_CHECK_DELEGATE_RETVAL(GIntRet, FTestDynamicDelegateRetInt0, 0, &UDynamicDelegateTestObject::ConstFuncRetInt0, ());
-		BIND_CALL_CHECK_DELEGATE_RETVAL(GStrRet, FTestDynamicDelegateRetStr0, 0, &UDynamicDelegateTestObject::FuncRetStr0,      ());
+		BIND_CALL_CHECK_DELEGATE_RETVAL(GStrRet, FTestDynamicDelegateRetStr0, 0, &UDynamicDelegateTestObject::FuncRetStr0, ());
 		BIND_CALL_CHECK_DELEGATE_RETVAL(GStrRet, FTestDynamicDelegateRetStr0, 0, &UDynamicDelegateTestObject::ConstFuncRetStr0, ());
 	}
 
@@ -1206,6 +1453,68 @@ void UDynamicDelegateTestObject::RunTests()
 		BIND_CALL_CHECK_DELEGATE_OUT(FTestDynamicDelegateOut6,                     0, &UDynamicDelegateTestObject::ConstFuncOut6, (GInt16Arg, GUnscopedEnumArg, GStrArg, GNameArg, bGBoolArg, GWeakObjPtrArg), (GInt16Out, GUnscopedEnumOut, GStrOut, GNameOut, bGBoolOut, GWeakObjPtrOut));
 		BIND_CALL_CHECK_DELEGATE_OUT(FTestDynamicMulticastDelegateOut6,            1, &UDynamicDelegateTestObject::ConstFuncOut6, (GInt16Arg, GUnscopedEnumArg, GStrArg, GNameArg, bGBoolArg, GWeakObjPtrArg), (GInt16Out, GUnscopedEnumOut, GStrOut, GNameOut, bGBoolOut, GWeakObjPtrOut));
 	}
+
+#if defined(UE_USE_DYNAMIC_DELEGATE_PAYLOADS) && UE_USE_DYNAMIC_DELEGATE_PAYLOADS
+	// Additional property type payload tests — unicast
+	{
+		BIND_CALL_CHECK_DELEGATE(FTestDynamicDelegate0, 0, &UDynamicDelegateTestObject::FuncFloat,       (), GFloatArg);
+		BIND_CALL_CHECK_DELEGATE(FTestDynamicDelegate0, 0, &UDynamicDelegateTestObject::FuncDouble,      (), GDoubleArg);
+		BIND_CALL_CHECK_DELEGATE(FTestDynamicDelegate0, 0, &UDynamicDelegateTestObject::FuncByte,        (), GByteArg);
+		BIND_CALL_CHECK_DELEGATE(FTestDynamicDelegate0, 0, &UDynamicDelegateTestObject::FuncInt64,       (), GInt64Arg);
+		BIND_CALL_CHECK_DELEGATE(FTestDynamicDelegate0, 0, &UDynamicDelegateTestObject::FuncText,        (), GTextArg);
+		BIND_CALL_CHECK_DELEGATE(FTestDynamicDelegate0, 0, &UDynamicDelegateTestObject::FuncStruct,      (), GStructArg);
+		BIND_CALL_CHECK_DELEGATE(FTestDynamicDelegate0, 0, &UDynamicDelegateTestObject::FuncVector,      (), GVectorArg);
+		BIND_CALL_CHECK_DELEGATE(FTestDynamicDelegate0, 0, &UDynamicDelegateTestObject::FuncIntArray,    (), GIntArrayArg);
+		BIND_CALL_CHECK_DELEGATE(FTestDynamicDelegate0, 0, &UDynamicDelegateTestObject::FuncScopedEnum,  (), GScopedEnumArg);
+		BIND_CALL_CHECK_DELEGATE(FTestDynamicDelegate0, 0, &UDynamicDelegateTestObject::FuncIntSet,      (), GIntSetArg);
+		BIND_CALL_CHECK_DELEGATE(FTestDynamicDelegate0, 0, &UDynamicDelegateTestObject::FuncOptionalInt, (), GOptionalIntArg);
+	}
+
+	// Additional property type payload tests — multicast
+	{
+		BIND_CALL_CHECK_DELEGATE(FTestDynamicMulticastDelegate0, 1, &UDynamicDelegateTestObject::FuncFloat,       (), GFloatArg);
+		BIND_CALL_CHECK_DELEGATE(FTestDynamicMulticastDelegate0, 1, &UDynamicDelegateTestObject::FuncDouble,      (), GDoubleArg);
+		BIND_CALL_CHECK_DELEGATE(FTestDynamicMulticastDelegate0, 1, &UDynamicDelegateTestObject::FuncByte,        (), GByteArg);
+		BIND_CALL_CHECK_DELEGATE(FTestDynamicMulticastDelegate0, 1, &UDynamicDelegateTestObject::FuncInt64,       (), GInt64Arg);
+		BIND_CALL_CHECK_DELEGATE(FTestDynamicMulticastDelegate0, 1, &UDynamicDelegateTestObject::FuncText,        (), GTextArg);
+		BIND_CALL_CHECK_DELEGATE(FTestDynamicMulticastDelegate0, 1, &UDynamicDelegateTestObject::FuncStruct,      (), GStructArg);
+		BIND_CALL_CHECK_DELEGATE(FTestDynamicMulticastDelegate0, 1, &UDynamicDelegateTestObject::FuncVector,      (), GVectorArg);
+		BIND_CALL_CHECK_DELEGATE(FTestDynamicMulticastDelegate0, 1, &UDynamicDelegateTestObject::FuncIntArray,    (), GIntArrayArg);
+		BIND_CALL_CHECK_DELEGATE(FTestDynamicMulticastDelegate0, 1, &UDynamicDelegateTestObject::FuncScopedEnum,  (), GScopedEnumArg);
+		BIND_CALL_CHECK_DELEGATE(FTestDynamicMulticastDelegate0, 1, &UDynamicDelegateTestObject::FuncIntSet,      (), GIntSetArg);
+		BIND_CALL_CHECK_DELEGATE(FTestDynamicMulticastDelegate0, 1, &UDynamicDelegateTestObject::FuncOptionalInt, (), GOptionalIntArg);
+	}
+
+	// Additional property type out-parameter tests — unicast
+	{
+		BIND_CALL_CHECK_DELEGATE_OUT(FTestDynamicDelegateOutFloat,       0, &UDynamicDelegateTestObject::FuncOutFloat,       (GFloatArg),       (GFloatOut));
+		BIND_CALL_CHECK_DELEGATE_OUT(FTestDynamicDelegateOutDouble,      0, &UDynamicDelegateTestObject::FuncOutDouble,      (GDoubleArg),      (GDoubleOut));
+		BIND_CALL_CHECK_DELEGATE_OUT(FTestDynamicDelegateOutByte,        0, &UDynamicDelegateTestObject::FuncOutByte,        (GByteArg),        (GByteOut));
+		BIND_CALL_CHECK_DELEGATE_OUT(FTestDynamicDelegateOutInt64,       0, &UDynamicDelegateTestObject::FuncOutInt64,       (GInt64Arg),       (GInt64Out));
+		BIND_CALL_CHECK_DELEGATE_OUT(FTestDynamicDelegateOutText,        0, &UDynamicDelegateTestObject::FuncOutText,        (GTextArg),        (GTextOut));
+		BIND_CALL_CHECK_DELEGATE_OUT(FTestDynamicDelegateOutStruct,      0, &UDynamicDelegateTestObject::FuncOutStruct,      (GStructArg),      (GStructOut));
+		BIND_CALL_CHECK_DELEGATE_OUT(FTestDynamicDelegateOutVector,      0, &UDynamicDelegateTestObject::FuncOutVector,      (GVectorArg),      (GVectorOut));
+		BIND_CALL_CHECK_DELEGATE_OUT(FTestDynamicDelegateOutScopedEnum,  0, &UDynamicDelegateTestObject::FuncOutScopedEnum,  (GScopedEnumArg),  (GScopedEnumOut));
+		BIND_CALL_CHECK_DELEGATE_OUT(FTestDynamicDelegateOutIntArray,    0, &UDynamicDelegateTestObject::FuncOutIntArray,    (GIntArrayArg),    (GIntArrayOut));
+		BIND_CALL_CHECK_DELEGATE_OUT(FTestDynamicDelegateOutIntSet,      0, &UDynamicDelegateTestObject::FuncOutIntSet,      (GIntSetArg),      (GIntSetOut));
+		BIND_CALL_CHECK_DELEGATE_OUT(FTestDynamicDelegateOutOptionalInt, 0, &UDynamicDelegateTestObject::FuncOutOptionalInt, (GOptionalIntArg), (GOptionalIntOut));
+	}
+
+	// Additional property type out-parameter tests — multicast
+	{
+		BIND_CALL_CHECK_DELEGATE_OUT(FTestDynamicMulticastDelegateOutFloat,       1, &UDynamicDelegateTestObject::FuncOutFloat,       (GFloatArg),       (GFloatOut));
+		BIND_CALL_CHECK_DELEGATE_OUT(FTestDynamicMulticastDelegateOutDouble,      1, &UDynamicDelegateTestObject::FuncOutDouble,      (GDoubleArg),      (GDoubleOut));
+		BIND_CALL_CHECK_DELEGATE_OUT(FTestDynamicMulticastDelegateOutByte,        1, &UDynamicDelegateTestObject::FuncOutByte,        (GByteArg),        (GByteOut));
+		BIND_CALL_CHECK_DELEGATE_OUT(FTestDynamicMulticastDelegateOutInt64,       1, &UDynamicDelegateTestObject::FuncOutInt64,       (GInt64Arg),       (GInt64Out));
+		BIND_CALL_CHECK_DELEGATE_OUT(FTestDynamicMulticastDelegateOutText,        1, &UDynamicDelegateTestObject::FuncOutText,        (GTextArg),        (GTextOut));
+		BIND_CALL_CHECK_DELEGATE_OUT(FTestDynamicMulticastDelegateOutStruct,      1, &UDynamicDelegateTestObject::FuncOutStruct,      (GStructArg),      (GStructOut));
+		BIND_CALL_CHECK_DELEGATE_OUT(FTestDynamicMulticastDelegateOutVector,      1, &UDynamicDelegateTestObject::FuncOutVector,      (GVectorArg),      (GVectorOut));
+		BIND_CALL_CHECK_DELEGATE_OUT(FTestDynamicMulticastDelegateOutScopedEnum,  1, &UDynamicDelegateTestObject::FuncOutScopedEnum,  (GScopedEnumArg),  (GScopedEnumOut));
+		BIND_CALL_CHECK_DELEGATE_OUT(FTestDynamicMulticastDelegateOutIntArray,    1, &UDynamicDelegateTestObject::FuncOutIntArray,    (GIntArrayArg),    (GIntArrayOut));
+		BIND_CALL_CHECK_DELEGATE_OUT(FTestDynamicMulticastDelegateOutIntSet,      1, &UDynamicDelegateTestObject::FuncOutIntSet,      (GIntSetArg),      (GIntSetOut));
+		BIND_CALL_CHECK_DELEGATE_OUT(FTestDynamicMulticastDelegateOutOptionalInt, 1, &UDynamicDelegateTestObject::FuncOutOptionalInt, (GOptionalIntArg), (GOptionalIntOut));
+	}
+#endif
 }
 
 TEST_CASE_NAMED(DynamicDelegateTest, "UE::CoreUObject::DynamicDelegateTest", "[ApplicationContextMask][SmokeFilter]")
diff --git a/Engine/Source/Runtime/CoreUObject/Public/UObject/DynamicDelegateTestObject.h b/Engine/Source/Runtime/CoreUObject/Public/UObject/DynamicDelegateTestObject.h
--- a/Engine/Source/Runtime/CoreUObject/Public/UObject/DynamicDelegateTestObject.h
+++ b/Engine/Source/Runtime/CoreUObject/Public/UObject/DynamicDelegateTestObject.h
@@ -9,6 +9,30 @@
 
 #include "DynamicDelegateTestObject.generated.h"
 
+USTRUCT()
+struct FTestPayloadStruct
+{
+	GENERATED_BODY()
+
+	UPROPERTY()
+	int32 X = 0;
+
+	UPROPERTY()
+	int32 Y = 0;
+
+	UPROPERTY()
+	int32 Z = 0;
+
+	bool operator==(const FTestPayloadStruct&) const = default;
+};
+
+UENUM()
+enum class ETestScopedEnum : uint8
+{
+	Value1,
+	Value2
+};
+
 UENUM()
 namespace ETestUnscopedEnum
 {
@@ -70,6 +94,59 @@ DECLARE_DYNAMIC_MULTICAST_DELEGATE_FourParams (FTestDynamicMulticastDelegateOut4
 DECLARE_DYNAMIC_MULTICAST_DELEGATE_FiveParams (FTestDynamicMulticastDelegateOut5, int16&, Int, TEnumAsByte<ETestUnscopedEnum::Type>&, UnscopedEnum, FString&, Str, FName&, Name, bool&, bBool);
 DECLARE_DYNAMIC_MULTICAST_DELEGATE_SixParams  (FTestDynamicMulticastDelegateOut6, int16&, Int, TEnumAsByte<ETestUnscopedEnum::Type>&, UnscopedEnum, FString&, Str, FName&, Name, bool&, bBool, TWeakObjectPtr<UDynamicDelegateTestObject>&, WeakPtr);
 
+// Delegates for testing additional property types in payload serialization.
+// Only single arguments supported at the moment- we need to add UHT support for proper variadic dynamic delegates to get full combination testsing, like:
+// UDELEGATE()
+// using FTestDynamicDelegateRetStr6 = TDynamicDelegate<void(int Int, FString Str, const UDynamicDelegateTestObject* Ptr, FName Name, bool bBool, TWeakObjectPtr<UDynamicDelegateTestObject> WeakPtr)>;
+DECLARE_DYNAMIC_DELEGATE_OneParam(FTestDynamicDelegateFloat,        float,              Float);
+DECLARE_DYNAMIC_DELEGATE_OneParam(FTestDynamicDelegateDouble,       double,             Double);
+DECLARE_DYNAMIC_DELEGATE_OneParam(FTestDynamicDelegateByte,         uint8,              Byte);
+DECLARE_DYNAMIC_DELEGATE_OneParam(FTestDynamicDelegateInt64,        int64,              Int64);
+DECLARE_DYNAMIC_DELEGATE_OneParam(FTestDynamicDelegateText,         FText,              Text);
+DECLARE_DYNAMIC_DELEGATE_OneParam(FTestDynamicDelegateStruct,       FTestPayloadStruct, Struct);
+DECLARE_DYNAMIC_DELEGATE_OneParam(FTestDynamicDelegateVector,       FVector,            Vector);
+DECLARE_DYNAMIC_DELEGATE_OneParam(FTestDynamicDelegateIntArray,     TArray<int32>,      IntArray);
+DECLARE_DYNAMIC_DELEGATE_OneParam(FTestDynamicDelegateScopedEnum,   ETestScopedEnum,    Enum);
+DECLARE_DYNAMIC_DELEGATE_OneParam(FTestDynamicDelegateIntSet,       TSet<int32>,        IntSet);
+DECLARE_DYNAMIC_DELEGATE_OneParam(FTestDynamicDelegateOptionalInt,  TOptional<int32>,   OptionalInt);
+
+DECLARE_DYNAMIC_MULTICAST_DELEGATE_OneParam(FTestDynamicMulticastDelegateFloat,       float,              Float);
+DECLARE_DYNAMIC_MULTICAST_DELEGATE_OneParam(FTestDynamicMulticastDelegateDouble,      double,             Double);
+DECLARE_DYNAMIC_MULTICAST_DELEGATE_OneParam(FTestDynamicMulticastDelegateByte,        uint8,              Byte);
+DECLARE_DYNAMIC_MULTICAST_DELEGATE_OneParam(FTestDynamicMulticastDelegateInt64,       int64,              Int64);
+DECLARE_DYNAMIC_MULTICAST_DELEGATE_OneParam(FTestDynamicMulticastDelegateText,        FText,              Text);
+DECLARE_DYNAMIC_MULTICAST_DELEGATE_OneParam(FTestDynamicMulticastDelegateStruct,      FTestPayloadStruct, Struct);
+DECLARE_DYNAMIC_MULTICAST_DELEGATE_OneParam(FTestDynamicMulticastDelegateVector,      FVector,            Vector);
+DECLARE_DYNAMIC_MULTICAST_DELEGATE_OneParam(FTestDynamicMulticastDelegateIntArray,    TArray<int32>,      IntArray);
+DECLARE_DYNAMIC_MULTICAST_DELEGATE_OneParam(FTestDynamicMulticastDelegateScopedEnum,  ETestScopedEnum,    Enum);
+DECLARE_DYNAMIC_MULTICAST_DELEGATE_OneParam(FTestDynamicMulticastDelegateIntSet,      TSet<int32>,        IntSet);
+DECLARE_DYNAMIC_MULTICAST_DELEGATE_OneParam(FTestDynamicMulticastDelegateOptionalInt, TOptional<int32>,   OptionalInt);
+
+// Out-parameter delegates for additional property types
+DECLARE_DYNAMIC_DELEGATE_OneParam(FTestDynamicDelegateOutFloat,       float&,              Float);
+DECLARE_DYNAMIC_DELEGATE_OneParam(FTestDynamicDelegateOutDouble,      double&,             Double);
+DECLARE_DYNAMIC_DELEGATE_OneParam(FTestDynamicDelegateOutByte,        uint8&,              Byte);
+DECLARE_DYNAMIC_DELEGATE_OneParam(FTestDynamicDelegateOutInt64,       int64&,              Int64);
+DECLARE_DYNAMIC_DELEGATE_OneParam(FTestDynamicDelegateOutText,        FText&,              Text);
+DECLARE_DYNAMIC_DELEGATE_OneParam(FTestDynamicDelegateOutStruct,      FTestPayloadStruct&, Struct);
+DECLARE_DYNAMIC_DELEGATE_OneParam(FTestDynamicDelegateOutVector,      FVector&,            Vector);
+DECLARE_DYNAMIC_DELEGATE_OneParam(FTestDynamicDelegateOutScopedEnum,  ETestScopedEnum&,    Enum);
+DECLARE_DYNAMIC_DELEGATE_OneParam(FTestDynamicDelegateOutIntArray,    TArray<int32>&,      IntArray);
+DECLARE_DYNAMIC_DELEGATE_OneParam(FTestDynamicDelegateOutIntSet,      TSet<int32>&,        IntSet);
+DECLARE_DYNAMIC_DELEGATE_OneParam(FTestDynamicDelegateOutOptionalInt, TOptional<int32>&,   OptionalInt);
+
+DECLARE_DYNAMIC_MULTICAST_DELEGATE_OneParam(FTestDynamicMulticastDelegateOutFloat,       float&,              Float);
+DECLARE_DYNAMIC_MULTICAST_DELEGATE_OneParam(FTestDynamicMulticastDelegateOutDouble,      double&,             Double);
+DECLARE_DYNAMIC_MULTICAST_DELEGATE_OneParam(FTestDynamicMulticastDelegateOutByte,        uint8&,              Byte);
+DECLARE_DYNAMIC_MULTICAST_DELEGATE_OneParam(FTestDynamicMulticastDelegateOutInt64,       int64&,              Int64);
+DECLARE_DYNAMIC_MULTICAST_DELEGATE_OneParam(FTestDynamicMulticastDelegateOutText,        FText&,              Text);
+DECLARE_DYNAMIC_MULTICAST_DELEGATE_OneParam(FTestDynamicMulticastDelegateOutStruct,      FTestPayloadStruct&, Struct);
+DECLARE_DYNAMIC_MULTICAST_DELEGATE_OneParam(FTestDynamicMulticastDelegateOutVector,      FVector&,            Vector);
+DECLARE_DYNAMIC_MULTICAST_DELEGATE_OneParam(FTestDynamicMulticastDelegateOutScopedEnum,  ETestScopedEnum&,    Enum);
+DECLARE_DYNAMIC_MULTICAST_DELEGATE_OneParam(FTestDynamicMulticastDelegateOutIntArray,    TArray<int32>&,      IntArray);
+DECLARE_DYNAMIC_MULTICAST_DELEGATE_OneParam(FTestDynamicMulticastDelegateOutIntSet,      TSet<int32>&,        IntSet);
+DECLARE_DYNAMIC_MULTICAST_DELEGATE_OneParam(FTestDynamicMulticastDelegateOutOptionalInt, TOptional<int32>&,   OptionalInt);
+
 UCLASS()
 class UDynamicDelegateTestObject : public UObject
 {
@@ -148,6 +225,32 @@ public:
 	UFUNCTION() void ConstFuncOut5(int16& Int, TEnumAsByte<ETestUnscopedEnum::Type>& UnscopedEnum, FString& Str, FName& Name, bool& bBool) const;
 	UFUNCTION() void ConstFuncOut6(int16& Int, TEnumAsByte<ETestUnscopedEnum::Type>& UnscopedEnum, FString& Str, FName& Name, bool& bBool, TWeakObjectPtr<UDynamicDelegateTestObject>& WeakPtr) const;
 
+	// Functions for testing additional property types in payload serialization
+	UFUNCTION() void FuncFloat(float Float);
+	UFUNCTION() void FuncDouble(double Double);
+	UFUNCTION() void FuncByte(uint8 Byte);
+	UFUNCTION() void FuncInt64(int64 Int64);
+	UFUNCTION() void FuncText(FText Text);
+	UFUNCTION() void FuncStruct(FTestPayloadStruct Struct);
+	UFUNCTION() void FuncVector(FVector Vector);
+	UFUNCTION() void FuncIntArray(TArray<int32> IntArray);
+	UFUNCTION() void FuncScopedEnum(ETestScopedEnum Enum);
+	UFUNCTION() void FuncIntSet(TSet<int32> IntSet);
+	UFUNCTION() void FuncOptionalInt(TOptional<int32> OptionalInt);
+
+	// Out-parameter functions for additional property types
+	UFUNCTION() void FuncOutFloat(float& Float);
+	UFUNCTION() void FuncOutDouble(double& Double);
+	UFUNCTION() void FuncOutByte(uint8& Byte);
+	UFUNCTION() void FuncOutInt64(int64& Int64);
+	UFUNCTION() void FuncOutText(FText& Text);
+	UFUNCTION() void FuncOutStruct(FTestPayloadStruct& Struct);
+	UFUNCTION() void FuncOutVector(FVector& Vector);
+	UFUNCTION() void FuncOutScopedEnum(ETestScopedEnum& Enum);
+	UFUNCTION() void FuncOutIntArray(TArray<int32>& IntArray);
+	UFUNCTION() void FuncOutIntSet(TSet<int32>& IntSet);
+	UFUNCTION() void FuncOutOptionalInt(TOptional<int32>& OptionalInt);
+
 	void RunTests();
 };
 
