Updated FObjectPropertyBase::ParseObjectPropertyValue and FSoftObjectPath::ImportTextItem to handle parsing object paths with spaces in them

This is technically not allowed by INVALID_OBJECTNAME_CHARACTERS, but happens all the time in BPs

#Synced-CL 52533865

diff --git a/Engine/Source/Runtime/CoreUObject/Private/UObject/Property.cpp b/Engine/Source/Runtime/CoreUObject/Private/UObject/Property.cpp
--- a/Engine/Source/Runtime/CoreUObject/Private/UObject/Property.cpp
+++ b/Engine/Source/Runtime/CoreUObject/Private/UObject/Property.cpp
@@ -652,63 +652,74 @@ FORCEINLINE constexpr FStringView ParsePropertyToken(const TCHAR* Str, bool Dott
 //
 // Parse a token.
 //
-const TCHAR* FPropertyHelpers::ReadToken( const TCHAR* Buffer, FString& String, bool bDottedNames)
+template <typename TStringType>
+const TCHAR* ReadPropertyToken(const TCHAR* Buffer, TStringType& Out, bool bDottedNames)
 {
-	if( *Buffer == TCHAR('"') )
+	if (*Buffer == TEXT('"'))
 	{
 		int32 NumCharsRead = 0;
-		if (!FParse::QuotedString(Buffer, String, &NumCharsRead))
+		if (!FParse::QuotedString(Buffer, Out, &NumCharsRead))
 		{
-			UE_LOGF(LogProperty, Warning, "ReadToken: Bad quoted string: %ls", Buffer );
+			UE_LOGF(LogProperty, Warning, "ReadToken: Bad quoted string: %ls", Buffer);
 			return nullptr;
 		}
 		Buffer += NumCharsRead;
+
+		// TODO special handling of null-terminator here?
 	}
 	else if (IsValidTokenStart(*Buffer, bDottedNames))
 	{
 		FStringView Token = ParsePropertyToken(Buffer, bDottedNames);
-		String += Token;
+		Out += Token;
 		Buffer += Token.Len();
 	}
 	else
 	{
 		// Get just one.
-		String += *Buffer;
+		if (*Buffer)
+		{
+			Out += *Buffer;
+		}
 	}
 	return Buffer;
 }
 
+const TCHAR* FPropertyHelpers::ReadToken( const TCHAR* Buffer, FString& Out, bool bDottedNames)
+{
+	return ReadPropertyToken(Buffer, Out, bDottedNames);
+}
+
 const TCHAR* FPropertyHelpers::ReadToken( const TCHAR* Buffer, FStringBuilderBase& Out, bool bDottedNames)
 {
-	if( *Buffer == TCHAR('"') )
-	{
-		int32 NumCharsRead = 0;
-		if (!FParse::QuotedString(Buffer, Out, &NumCharsRead))
-		{
-			UE_LOGF(LogProperty, Warning, "ReadToken: Bad quoted string: %ls", Buffer );
-			return nullptr;
-		}
-		Buffer += NumCharsRead;
+	return ReadPropertyToken(Buffer, Out, bDottedNames);
+}
 
-		// TODO special handling of null-terminator here?
-	}
-	else if (IsValidTokenStart(*Buffer, bDottedNames))
+//
+// Parse a terminated token.
+//
+template <typename TStringType>
+const TCHAR* ReadTerminatedPropertyToken(const TCHAR* Buffer, TStringType& Out, const TCHAR Terminator, const int32* MaxTokenLen)
+{
+	const TCHAR* BufferStart = Buffer;
+	while (*Buffer && *Buffer != Terminator && (!MaxTokenLen || UE_PTRDIFF_TO_INT32(Buffer - BufferStart) < *MaxTokenLen)) { ++Buffer; } // Walk to the Terminator, with an optional limit to guard against bad data
+	if (*Buffer != Terminator) // Verify we found the Terminator
 	{
-		FStringView Token = ParsePropertyToken(Buffer, bDottedNames);
-		Out << Token;
-		Buffer += Token.Len();
-	}
-	else
-	{
-		// Get just one.
-		if (*Buffer)
-		{
-			Out << *Buffer;
-		}
+		return nullptr;
 	}
+	Out += FStringView(BufferStart, UE_PTRDIFF_TO_INT32(Buffer - BufferStart));
 	return Buffer;
 }
 
+const TCHAR* FPropertyHelpers::ReadTerminatedToken(const TCHAR* Buffer, FString& Out, const TCHAR Terminator, const int32* MaxTokenLen)
+{
+	return ReadTerminatedPropertyToken(Buffer, Out, Terminator, MaxTokenLen);
+}
+
+const TCHAR* FPropertyHelpers::ReadTerminatedToken(const TCHAR* Buffer, FStringBuilderBase& Out, const TCHAR Terminator, const int32* MaxTokenLen)
+{
+	return ReadTerminatedPropertyToken(Buffer, Out, Terminator, MaxTokenLen);
+}
+
 /*-----------------------------------------------------------------------------
 	FProperty implementation.
 -----------------------------------------------------------------------------*/
diff --git a/Engine/Source/Runtime/CoreUObject/Private/UObject/PropertyBaseObject.cpp b/Engine/Source/Runtime/CoreUObject/Private/UObject/PropertyBaseObject.cpp
--- a/Engine/Source/Runtime/CoreUObject/Private/UObject/PropertyBaseObject.cpp
+++ b/Engine/Source/Runtime/CoreUObject/Private/UObject/PropertyBaseObject.cpp
@@ -361,16 +361,18 @@ bool FObjectPropertyBase::ParseObjectPropertyValue(const FProperty* Property, UO
 
 		bool bWarnOnnullptr = (PortFlags&PPF_CheckReferences)!=0;
 
-		if( *Buffer == TCHAR('\'') )
+		if (*Buffer == TEXT('\''))
 		{
+			// This is likely the remnants of stripping off the class part of a fully qualified object path
+			// eg) /Script/Module.Class'/Game/Foo.Foo:Name With A Space.Bar' -> '/Game/Foo.Foo:Name With A Space.Bar'
+			// We walk to the next ' so that we correctly handle object paths with spaces in them, which is technically 
+			// not allowed by INVALID_OBJECTNAME_CHARACTERS, but happens all the time in BPs
 			Temp.Reset();
-			Buffer = FPropertyHelpers::ReadToken( ++Buffer, /* out */ Temp, true);
-			if( Buffer == nullptr )
 			{
-				return false;
+				const int32 MaxTokenLength = FName::StringBufferSize - 1;
+				Buffer = FPropertyHelpers::ReadTerminatedToken(++Buffer, /*Out*/Temp, TEXT('\''), &MaxTokenLength);
 			}
-
-			if( *Buffer++ != TCHAR('\'') )
+			if (!Buffer || *Buffer++ != TEXT('\''))
 			{
 				return false;
 			}
diff --git a/Engine/Source/Runtime/CoreUObject/Private/UObject/SoftObjectPath.cpp b/Engine/Source/Runtime/CoreUObject/Private/UObject/SoftObjectPath.cpp
--- a/Engine/Source/Runtime/CoreUObject/Private/UObject/SoftObjectPath.cpp
+++ b/Engine/Source/Runtime/CoreUObject/Private/UObject/SoftObjectPath.cpp
@@ -652,11 +652,15 @@ bool FSoftObjectPath::ImportTextItem(const TCHAR*& Buffer, int32 PortFlags, UObj
 
 		if (*Buffer == TCHAR('\''))
 		{
-			// A ' token likely means we're looking at a path string in the form "Texture2d'/Game/UI/HUD/Actions/Barrel'" and we need to read and append the path part
-			// We have to skip over the first ' as FPropertyHelpers::ReadToken doesn't read single-quoted strings correctly, but does read a path correctly
-			Buffer++; // Skip the leading '
+			// This is likely the remnants of stripping off the class part of a fully qualified object path
+			// eg) /Script/Module.Class'/Game/Foo.Foo:Name With A Space.Bar' -> '/Game/Foo.Foo:Name With A Space.Bar'
+			// We walk to the next ' so that we correctly handle object paths with spaces in them, which is technically 
+			// not allowed by INVALID_OBJECTNAME_CHARACTERS, but happens all the time in BPs
 			ImportedPath.Reset();
-			NewBuffer = FPropertyHelpers::ReadToken(Buffer, /* out */ ImportedPath, /* dotted names */ true);
+			{
+				const int32 MaxTokenLength = FName::StringBufferSize - 1;
+				NewBuffer = FPropertyHelpers::ReadTerminatedToken(++Buffer, /*Out*/ImportedPath, TEXT('\''), &MaxTokenLength);
+			}
 			if (!NewBuffer)
 			{
 				return false;
diff --git a/Engine/Source/Runtime/CoreUObject/Public/UObject/UnrealType.h b/Engine/Source/Runtime/CoreUObject/Public/UObject/UnrealType.h
--- a/Engine/Source/Runtime/CoreUObject/Public/UObject/UnrealType.h
+++ b/Engine/Source/Runtime/CoreUObject/Public/UObject/UnrealType.h
@@ -1373,10 +1373,22 @@ public:
 class FPropertyHelpers
 {
 public:
-	static COREUOBJECT_API const TCHAR* ReadToken( const TCHAR* Buffer, FString& Out, bool DottedNames = false);
-
 	// @param Out Appended to
+	static COREUOBJECT_API const TCHAR* ReadToken( const TCHAR* Buffer, FString& Out, bool DottedNames = false);
 	static COREUOBJECT_API const TCHAR* ReadToken( const TCHAR* Buffer, FStringBuilderBase& Out, bool DottedNames = false);
+
+	/**
+	 * Parse a token from the given buffer until we find the given terminating character, or we run out of buffer/hit the given limit.
+	 * 
+	 * @param Buffer		The buffer to parse from.
+	 * @param Out			The value to append to on success.
+	 * @param Terminator	The token terminator to stop at.
+	 * @param MaxTokenLen	Optional limit; will fail if we don't find the terminator within this limit.
+	 * 
+	 * @return The terminating character location in the buffer, or null if parsing failed.
+	 */
+	static COREUOBJECT_API const TCHAR* ReadTerminatedToken(const TCHAR* Buffer, FString& Out, const TCHAR Terminator, const int32* MaxTokenLen = nullptr);
+	static COREUOBJECT_API const TCHAR* ReadTerminatedToken(const TCHAR* Buffer, FStringBuilderBase& Out, const TCHAR Terminator, const int32* MaxTokenLen = nullptr);
 };
 
 namespace UEProperty_Private
