Don't assert when creating an object with identical remote id as an existing object if the existing object is unreachable

+ Added a test that covers this scenario

#Synced-CL 52534556

diff --git a/Engine/Source/Runtime/CoreUObject/Private/UObject/UObjectHash.cpp b/Engine/Source/Runtime/CoreUObject/Private/UObject/UObjectHash.cpp
--- a/Engine/Source/Runtime/CoreUObject/Private/UObject/UObjectHash.cpp
+++ b/Engine/Source/Runtime/CoreUObject/Private/UObject/UObjectHash.cpp
@@ -1401,7 +1401,7 @@ FORCEINLINE static void AddToRemoteIdMap(FUObjectHashTables& ThreadHash, UObject
 		UObject* ExistingObject = (UObject*)*It;
 		checkf(ExistingObject != Object, TEXT("%s (%s) is being added to the remote id map for the second time"), *GetPathNameSafe(ExistingObject), *RemoteId.ToString());
 		FRemoteObjectId ExistingId(ExistingObject);
-		checkf(ExistingId != RemoteId, TEXT("Trying to add %s to the remote id map but it has the same remote id (%s) as %s"), *GetPathNameSafe((UObject*)Object), *RemoteId.ToString(), *GetPathNameSafe(ExistingObject));
+		checkf(ExistingId != RemoteId || ExistingObject->IsUnreachable(), TEXT("Trying to add %s to the remote id map but it has the same remote id (%s) as %s"), *GetPathNameSafe((UObject*)Object), *RemoteId.ToString(), *GetPathNameSafe(ExistingObject));
 	}
 #endif
 	Bucket.Add(Object);
