Fixing the editor compile errors with remote object support enabled (remote heap related)

#Synced-CL 52534556

diff --git a/Engine/Source/Runtime/CoreUObject/Private/UObject/RemoteHeap.cpp b/Engine/Source/Runtime/CoreUObject/Private/UObject/RemoteHeap.cpp
--- a/Engine/Source/Runtime/CoreUObject/Private/UObject/RemoteHeap.cpp
+++ b/Engine/Source/Runtime/CoreUObject/Private/UObject/RemoteHeap.cpp
@@ -3,7 +3,7 @@
 #include "Templates/UnrealTemplate.h"
 #include "AutoRTFM.h"
 
-#if UE_WITH_REMOTE_OBJECT_HANDLE
+#if UE_WITH_REMOTE_OBJECT_HANDLE && UE_AUTORTFM
 #include "jit_heap_ue.h"
 #endif
 
@@ -155,7 +155,7 @@ void FRemoteHeap::CreateInstance(FRemoteExecutor* Executor, FRemoteServerId Loca
 	Heap->AllocatorOverride = AllocatorOverride;
 	Instance = Heap;
 
-#if UE_WITH_REMOTE_OBJECT_HANDLE
+#if UE_WITH_REMOTE_OBJECT_HANDLE && UE_AUTORTFM
 	RemoteExecutor::RegisterRemoteSubsystem(Executor, Heap);
 
 	uint64 LocalServerHeapBaseAddress = GetRemoteHeapBaseAddressForServerId(LocalServerId);
@@ -1568,7 +1568,7 @@ void* FRemoteHeap::RemoteHeapAllocate(uint64 Size)
 {
 	void* Result = nullptr;
 
-#if UE_WITH_REMOTE_OBJECT_HANDLE
+#if UE_WITH_REMOTE_OBJECT_HANDLE && UE_AUTORTFM
 	// pad allocations to a multiple of the line size
 	// so separate allocations never get shared and can
 	// be migrated separately
@@ -1620,7 +1620,7 @@ void FRemoteHeap::RemoteHeapDeallocate(void* Ptr)
 	if (Ptr == nullptr)
 		return;
 
-#if UE_WITH_REMOTE_OBJECT_HANDLE
+#if UE_WITH_REMOTE_OBJECT_HANDLE && UE_AUTORTFM
 	uint64 Address = 0;
 	FMemory::Memcpy(&Address, &Ptr, sizeof(uint64));
 
