Fixed more unreachable code warnings as errors

#Synced-CL 52633447

#Synced-CL 52633447

diff --git a/Engine/Source/Runtime/CoreUObject/Private/UObject/RemoteObjectTransfer.cpp b/Engine/Source/Runtime/CoreUObject/Private/UObject/RemoteObjectTransfer.cpp
--- a/Engine/Source/Runtime/CoreUObject/Private/UObject/RemoteObjectTransfer.cpp
+++ b/Engine/Source/Runtime/CoreUObject/Private/UObject/RemoteObjectTransfer.cpp
@@ -572,6 +572,7 @@ public:
 		UE_CLOG(ActiveRequest->RequiredObjects.Num() || PendingObjectRequests.Num(), LogRemoteObject, VeryVerbose, TEXT("FRemoteObjectTransferQueue: TickRequest(%s) %d required objs (%d pending requests):"),
 			*ActiveRequest->GetRequestId().ToString(), ActiveRequest->RequiredObjects.Num(), PendingObjectRequests.Num());
 
+#if !NO_LOGGING
 		if (UE_GET_LOG_VERBOSITY(LogRemoteObject) == ELogVerbosity::VeryVerbose)
 		{
 			for (const FRemoteObjectId& RequiredObjectId : ActiveRequest->RequiredObjects)
@@ -584,7 +585,7 @@ public:
 					EnumToString(Handle::GetResidence(RequiredObjectId)));
 			}
 		}
-
+#endif
 		SendRequests();
 	}
 
