[BugFix] placeholder classes for entities weren't getting their super class set to entity

#Synced-CL 52533865
#Synced-CL 52533865

diff --git a/Engine/Source/Runtime/CoreUObject/Private/UObject/InstanceDataObjectUtils.cpp b/Engine/Source/Runtime/CoreUObject/Private/UObject/InstanceDataObjectUtils.cpp
--- a/Engine/Source/Runtime/CoreUObject/Private/UObject/InstanceDataObjectUtils.cpp
+++ b/Engine/Source/Runtime/CoreUObject/Private/UObject/InstanceDataObjectUtils.cpp
@@ -878,7 +878,7 @@ namespace UE
 		{
 			bResult = true;
 		}
-		else if ((SuperType.TypeName == NAME_Entity) && (SuperType.ClassName == NAME_Class))
+		else if ((SuperType.TypeName == NAME_Entity) && (SuperType.ClassName == NAME_Class || SuperType.ClassName == NAME_VerseClass))
 		{
 			bResult = true;
 		}
