Updated FAssetData constructor to use FStringView rather than const FString&

#rnx
#Synced-CL 52671426

diff --git a/Engine/Source/Runtime/CoreUObject/Private/AssetRegistry/AssetData.cpp b/Engine/Source/Runtime/CoreUObject/Private/AssetRegistry/AssetData.cpp
--- a/Engine/Source/Runtime/CoreUObject/Private/AssetRegistry/AssetData.cpp
+++ b/Engine/Source/Runtime/CoreUObject/Private/AssetRegistry/AssetData.cpp
@@ -361,8 +361,8 @@ PRAGMA_ENABLE_DEPRECATION_WARNINGS
 	SetChunkIDs(InChunkIDs);
 }
 
-FAssetData::FAssetData(const FString& InLongPackageName, const FString& InObjectPath, FTopLevelAssetPath InAssetClassPathName, FAssetDataTagMap InTags, TArrayView<const int32> InChunkIDs, uint32 InPackageFlags)
-	: PackageName(*InLongPackageName)
+FAssetData::FAssetData(FStringView InLongPackageName, FStringView InObjectPath, FTopLevelAssetPath InAssetClassPathName, FAssetDataTagMap InTags, TArrayView<const int32> InChunkIDs, uint32 InPackageFlags)
+	: PackageName(InLongPackageName)
 	, AssetClassPath(InAssetClassPathName)
 	, PackageFlags(InPackageFlags)
 {
@@ -370,14 +370,24 @@ FAssetData::FAssetData(const FString& InLongPackageName, const FString& InObject
 
 #if WITH_EDITORONLY_DATA
 PRAGMA_DISABLE_DEPRECATION_WARNINGS
-	ObjectPath = *InObjectPath;
+	ObjectPath = FName(InObjectPath);
 PRAGMA_ENABLE_DEPRECATION_WARNINGS
 #endif
 
 	SetTagsAndAssetBundles(MoveTemp(InTags));
 
-	PackagePath = FName(*FPackageName::GetLongPackagePath(InLongPackageName));
-
+	// Note: This emulates the logic of FPackageName::GetLongPackagePath, which does not have a FStringView overload, 
+	//       but which returns InLongPackageName if the given value has no path separator, vs FPathViews::GetPath 
+	//       which returns an empty string in that case
+	PackagePath = FName(FPathViews::GetPath(InLongPackageName));
+	if (PackagePath.IsNone())
+	{
+		if (int32 Unused = INDEX_NONE;
+			!InLongPackageName.FindChar(TEXT('/'), Unused))
+		{
+			PackagePath = PackageName;
+		}
+	}
 
 	const FAssetPathParts Parts = SplitIntoOuterPathAndAssetName(InObjectPath);
 	AssetName = FName(Parts.InnermostName);
diff --git a/Engine/Source/Runtime/CoreUObject/Public/AssetRegistry/AssetData.h b/Engine/Source/Runtime/CoreUObject/Public/AssetRegistry/AssetData.h
--- a/Engine/Source/Runtime/CoreUObject/Public/AssetRegistry/AssetData.h
+++ b/Engine/Source/Runtime/CoreUObject/Public/AssetRegistry/AssetData.h
@@ -262,7 +262,7 @@ PRAGMA_ENABLE_DEPRECATION_WARNINGS
 	/** Constructor building the ObjectPath in the form of InPackageName.InAssetName. does not work for object outer-ed to a different package. */
 	COREUOBJECT_API FAssetData(FName InPackageName, FName InPackagePath, FName InAssetName, FTopLevelAssetPath InAssetClassPathName, FAssetDataTagMap InTags = FAssetDataTagMap(), TArrayView<const int32> InChunkIDs = TArrayView<const int32>(), uint32 InPackageFlags = 0);
 	/** Constructor with a long package name and a full object path which might not be part of the package this asset is in. */
-	COREUOBJECT_API FAssetData(const FString& InLongPackageName, const FString& InObjectPath, FTopLevelAssetPath InAssetClassPathName, FAssetDataTagMap InTags = FAssetDataTagMap(), TArrayView<const int32> InChunkIDs = TArrayView<const int32>(), uint32 InPackageFlags = 0);
+	COREUOBJECT_API FAssetData(FStringView InLongPackageName, FStringView InObjectPath, FTopLevelAssetPath InAssetClassPathName, FAssetDataTagMap InTags = FAssetDataTagMap(), TArrayView<const int32> InChunkIDs = TArrayView<const int32>(), uint32 InPackageFlags = 0);
 
 	/** Constructor taking a UObject. By default trying to create one for a blueprint class will create one for the UBlueprint instead, but this can be overridden */
 	COREUOBJECT_API FAssetData(const UObject* InAsset, FAssetData::ECreationFlags InCreationFlags = ECreationFlags::None);
