[BugFix] the change made in 52668091 needed to be applied to a second overload. To avoid missing this type of issue in the future I wrote a helper function to unify code.

#Synced-CL 52533865

diff --git a/Engine/Source/Runtime/CoreUObject/Private/UObject/InstanceDataObjectUtils.cpp b/Engine/Source/Runtime/CoreUObject/Private/UObject/InstanceDataObjectUtils.cpp
--- a/Engine/Source/Runtime/CoreUObject/Private/UObject/InstanceDataObjectUtils.cpp
+++ b/Engine/Source/Runtime/CoreUObject/Private/UObject/InstanceDataObjectUtils.cpp
@@ -870,15 +870,15 @@ namespace UE
 		return bResult;
 	}
 
-	bool CanCreatePropertyBagPlaceholderSuperType(const UE::Serialization::Private::FTypeResource& SuperType)
+	static bool CanCreatePropertyBagPlaceholderSuperType(FName TypeName, FName ClassName)
 	{
 		bool bResult = false;
 
-		if ((SuperType.TypeName == NAME_Component) && (SuperType.ClassName == NAME_VerseClass))
+		if ((TypeName == NAME_Component) && (ClassName == NAME_VerseClass))
 		{
 			bResult = true;
 		}
-		else if ((SuperType.TypeName == NAME_Entity) && (SuperType.ClassName == NAME_Class || SuperType.ClassName == NAME_VerseClass))
+		else if ((TypeName == NAME_Entity) && (ClassName == NAME_Class || ClassName == NAME_VerseClass))
 		{
 			bResult = true;
 		}
@@ -886,20 +886,14 @@ namespace UE
 		return bResult;
 	}
 
-	bool CanCreatePropertyBagPlaceholderSuperType(const UStruct* SuperType)
+	bool CanCreatePropertyBagPlaceholderSuperType(const UE::Serialization::Private::FTypeResource& SuperType)
 	{
-		bool bResult = false;
-
-		if ((SuperType->GetFName() == NAME_Component) && (SuperType->GetClass()->GetFName() == NAME_VerseClass))
-		{
-			bResult = true;
-		}
-		else if ((SuperType->GetFName() == NAME_Entity) && (SuperType->GetClass()->GetFName() == NAME_Class))
-		{
-			bResult = true;
-		}
+		return CanCreatePropertyBagPlaceholderSuperType(SuperType.TypeName, SuperType.ClassName);
+	}
 
-		return bResult;
+	bool CanCreatePropertyBagPlaceholderSuperType(const UStruct* SuperType)
+	{
+		return CanCreatePropertyBagPlaceholderSuperType(SuperType->GetFName(), SuperType->GetClass()->GetFName());
 	}
 
 	bool IsClassOfInstanceDataObjectClass(const UStruct* Class)
