UE: Add DebugName to assert

#Synced-CL 52584531

diff --git a/Engine/Source/Runtime/CoreUObject/Private/UObject/GarbageCollectionSchema.cpp b/Engine/Source/Runtime/CoreUObject/Private/UObject/GarbageCollectionSchema.cpp
--- a/Engine/Source/Runtime/CoreUObject/Private/UObject/GarbageCollectionSchema.cpp
+++ b/Engine/Source/Runtime/CoreUObject/Private/UObject/GarbageCollectionSchema.cpp
@@ -308,7 +308,7 @@ FSchemaBuilder::~FSchemaBuilder()
 
 void FSchemaBuilder::Add(FMemberDeclaration Member)
 {
-	check(IsIn<ExtraWordTypes>(Member.Type) == IsSet(Member.ExtraWord));
+	checkf(IsIn<ExtraWordTypes>(Member.Type) == IsSet(Member.ExtraWord), TEXT("FSchemaBuilder::Add while IsIn(Type) doesn't match IsSet(ExtraWord). DebugName:%s"), *Member.DebugName.ToString());
 	checkf(IsIn<~GeneratedTypes>(Member.Type), TEXT("Member type %d is generated by Build() and shouldn't be declared"), int32{EnumToUnderlyingType(Member.Type)});
 	check(!BuiltSchema.IsSet());
 	
