Add extra assertion for mismatches between ContainsObjectReference and EmitReferenceInfo

#Synced-CL 52678434

diff --git a/Engine/Source/Runtime/CoreUObject/Private/UObject/GarbageCollection.cpp b/Engine/Source/Runtime/CoreUObject/Private/UObject/GarbageCollection.cpp
--- a/Engine/Source/Runtime/CoreUObject/Private/UObject/GarbageCollection.cpp
+++ b/Engine/Source/Runtime/CoreUObject/Private/UObject/GarbageCollection.cpp
@@ -6888,6 +6888,9 @@ void FArrayProperty::EmitReferenceInfo(UE::GC::FSchemaBuilder& Schema, int32 Bas
 		{
 			FPropertyStackScope PropertyScope(DebugPath, Inner);
 			Inner->EmitReferenceInfo(InnerSchema, 0, EncounteredStructProps, DebugPath);
+			// If inner property wrote nothing to the inner schema, the outer schema will assert. Assert here while we have more info.
+			checkf(InnerSchema.NumMembers() > 0, TEXT("%s: Inner GC schema for property %s has no members after returning true from ContainsObjectReference."), 
+				*DebugPath.GetPropertyPath(), *Inner->GetName());
 		}
 
 		Schema.Add(UE::GC::DeclareMember(DebugPath, BaseOffset + GetOffset_ForGC(), Type, InnerSchema.Build()));
