Debug visualizer support for FixedTagPrivate::FMapHandle.  Useful for seeing tag values in the FAssetData struct, widely used in the editor as the core asset reference type.  Uses a module private FStore global variable, so implementing a debug visualizer required exposing that.  To minimize header pollution, an FStoreBase struct was added in a different header file from the implementation version of the struct.

#rnx

#Synced-CL 52720478

diff --git a/Engine/Source/Runtime/Core/Private/Misc/CoreGlobals.cpp b/Engine/Source/Runtime/Core/Private/Misc/CoreGlobals.cpp
--- a/Engine/Source/Runtime/Core/Private/Misc/CoreGlobals.cpp
+++ b/Engine/Source/Runtime/Core/Private/Misc/CoreGlobals.cpp
@@ -263,6 +263,7 @@ namespace UE::CoreUObject::Private
 }
 UE::CoreUObject::Private::FStoredObjectPathDebug* GCoreComplexObjectPathDebug = nullptr;
 UE::CoreUObject::Private::FObjectHandlePackageDebugData* GCoreObjectHandlePackageDebug = nullptr;
+UE::CoreUObject::Private::FStoreBase* const* GCoreFixedTagStoresForDebugVisualizers = nullptr;
 #if PLATFORM_UNIX || PLATFORM_MAC
 uint8** CORE_API GNameBlocksDebug = FNameDebugVisualizer(FClangKeepDebugInfo{}).GetBlocks();
 FChunkedFixedUObjectArray*& CORE_API GObjectArrayForDebugVisualizers = GCoreObjectArrayForDebugVisualizers;
diff --git a/Engine/Source/Runtime/Core/Public/CoreGlobals.h b/Engine/Source/Runtime/Core/Public/CoreGlobals.h
--- a/Engine/Source/Runtime/Core/Public/CoreGlobals.h
+++ b/Engine/Source/Runtime/Core/Public/CoreGlobals.h
@@ -800,11 +800,15 @@ namespace UE::CoreUObject::Private
 	
 	/** Array to help visualize object handles in the debugger */
 	struct FObjectHandlePackageDebugData;
+
+	/** Used to visualize tag values in FAssetData in the debugger */
+	struct FStoreBase;
 }
 namespace UE::Core { struct FVisualizerDebuggingState; }
 extern CORE_API UE::CoreUObject::Private::FStoredObjectPathDebug* GCoreComplexObjectPathDebug;
 extern CORE_API UE::CoreUObject::Private::FObjectHandlePackageDebugData* GCoreObjectHandlePackageDebug;
 extern CORE_API UE::Core::FVisualizerDebuggingState* GCoreDebuggingState;
+extern CORE_API UE::CoreUObject::Private::FStoreBase* const* GCoreFixedTagStoresForDebugVisualizers;
 
 /** @return True if running cook-on-the-fly. */
 bool CORE_API IsRunningCookOnTheFly();
diff --git a/Engine/Source/Runtime/Core/Public/Modules/Boilerplate/ModuleBoilerplate.h b/Engine/Source/Runtime/Core/Public/Modules/Boilerplate/ModuleBoilerplate.h
--- a/Engine/Source/Runtime/Core/Public/Modules/Boilerplate/ModuleBoilerplate.h
+++ b/Engine/Source/Runtime/Core/Public/Modules/Boilerplate/ModuleBoilerplate.h
@@ -106,7 +106,8 @@ class FChunkedFixedUObjectArray;
 		UE_VISUALIZER_USE_AND_RETAIN FChunkedFixedUObjectArray*& GObjectArrayForDebugVisualizers = GCoreObjectArrayForDebugVisualizers; \
 		UE_VISUALIZER_USE_AND_RETAIN UE::CoreUObject::Private::FStoredObjectPathDebug*& GComplexObjectPathDebug = GCoreComplexObjectPathDebug; \
 		UE_VISUALIZER_USE_AND_RETAIN UE::CoreUObject::Private::FObjectHandlePackageDebugData*& GObjectHandlePackageDebug = GCoreObjectHandlePackageDebug; \
-		UE_VISUALIZER_USE_AND_RETAIN UE::Core::FVisualizerDebuggingState*& GDebuggingState = GCoreDebuggingState;
+		UE_VISUALIZER_USE_AND_RETAIN UE::Core::FVisualizerDebuggingState*& GDebuggingState = GCoreDebuggingState; \
+		UE_VISUALIZER_USE_AND_RETAIN UE::CoreUObject::Private::FStoreBase* const*& GFixedTagStoresForDebugVisualizers = GCoreFixedTagStoresForDebugVisualizers;
 #endif
 
 // in DLL builds, these are done per-module, otherwise we just need one in the application
diff --git a/Engine/Source/Runtime/CoreUObject/Private/AssetRegistry/AssetDataTagMap.cpp b/Engine/Source/Runtime/CoreUObject/Private/AssetRegistry/AssetDataTagMap.cpp
--- a/Engine/Source/Runtime/CoreUObject/Private/AssetRegistry/AssetDataTagMap.cpp
+++ b/Engine/Source/Runtime/CoreUObject/Private/AssetRegistry/AssetDataTagMap.cpp
@@ -479,6 +479,13 @@ namespace FixedTagPrivate
 			Data[Store.Index] = nullptr;
 		}
 
+#if !UE_BUILD_SHIPPING
+		UE::CoreUObject::Private::FStoreBase* const* GetDataForVisualizer() const
+		{
+			return reinterpret_cast<UE::CoreUObject::Private::FStoreBase* const*>(Data);
+		}
+#endif
+
 	private:		
 		FTransactionallySafeCriticalSection Mutex;
 		FStore* Data[Capacity] = {};
@@ -486,6 +493,19 @@ namespace FixedTagPrivate
 	}
 	GStores;
 
+#if !UE_BUILD_SHIPPING
+	// Assign the core global so natvis can resolve FMapHandle pairs from any DLL.
+	// GCoreFixedTagStoresForDebugVisualizers is declared in CoreGlobals.h and replicated per-DLL
+	// via GFixedTagStoresForDebugVisualizers in ModuleBoilerplate.h.
+	static struct FInitFixedTagStoresForVisualizer
+	{
+		FInitFixedTagStoresForVisualizer()
+		{
+			GCoreFixedTagStoresForDebugVisualizers = GStores.GetDataForVisualizer();
+		}
+	} GInitFixedTagStoresForVisualizer;
+#endif
+
 	//////////////////////////////////////////////////////////////////////////
 
 	template <bool bForStorage>
diff --git a/Engine/Source/Runtime/CoreUObject/Public/AssetRegistry/AssetDataTagMap.h b/Engine/Source/Runtime/CoreUObject/Public/AssetRegistry/AssetDataTagMap.h
--- a/Engine/Source/Runtime/CoreUObject/Public/AssetRegistry/AssetDataTagMap.h
+++ b/Engine/Source/Runtime/CoreUObject/Public/AssetRegistry/AssetDataTagMap.h
@@ -168,6 +168,26 @@ namespace FixedTagPrivate
 	};
 } // end namespace FixedTagPrivate
 
+namespace UE::CoreUObject::Private
+{
+	// Base struct for FixedTagPrivate::FStore (defined in AssetDataTagMapSerializationDetails.h), containing members needed
+	// by debug visualizer for FixedTagPrivate::FMapHandle.
+	struct FStoreBase
+	{
+		// Pairs for all unsorted maps that use this store
+		TArrayView<FixedTagPrivate::FNumberedPair> Pairs;
+		TArrayView<FixedTagPrivate::FNumberlessPair> NumberlessPairs;
+
+		// Values for visualizable maps in this store
+		TArrayView<uint32> AnsiStringOffsets;
+		TArrayView<ANSICHAR> AnsiStrings;
+		TArrayView<uint32> WideStringOffsets;
+		TArrayView<WIDECHAR> WideStrings;
+		TArrayView<FDisplayNameEntryId> NumberlessNames;
+		TArrayView<FName> Names;
+	};
+} // end namespace UE::CoreUObject::Private
+
 /**
  * Reference to a tagged value in a FAssetDataTagMapSharedView
  * 
diff --git a/Engine/Source/Runtime/CoreUObject/Public/AssetRegistry/AssetDataTagMapSerializationDetails.h b/Engine/Source/Runtime/CoreUObject/Public/AssetRegistry/AssetDataTagMapSerializationDetails.h
--- a/Engine/Source/Runtime/CoreUObject/Public/AssetRegistry/AssetDataTagMapSerializationDetails.h
+++ b/Engine/Source/Runtime/CoreUObject/Public/AssetRegistry/AssetDataTagMapSerializationDetails.h
@@ -94,19 +94,9 @@ namespace FixedTagPrivate
 	};
 
 	/// Stores a fixed set of values and all the key-values maps used for lookup
-	struct FStore
+	struct FStore : public UE::CoreUObject::Private::FStoreBase
 	{
-		// Pairs for all unsorted maps that uses this store 
-		TArrayView<FNumberedPair> Pairs;
-		TArrayView<FNumberlessPair> NumberlessPairs;
-
-		// Values for all maps in this store
-		TArrayView<uint32> AnsiStringOffsets;
-		TArrayView<ANSICHAR> AnsiStrings;
-		TArrayView<uint32> WideStringOffsets;
-		TArrayView<WIDECHAR> WideStrings;
-		TArrayView<FDisplayNameEntryId> NumberlessNames;
-		TArrayView<FName> Names;
+		// Values for non-visualizable maps in this store
 		TArrayView<FNumberlessExportPath> NumberlessExportPaths;
 		TArrayView<FAssetRegistryExportPath> ExportPaths;
 		TArrayView<FMarshalledText> Texts;
