[Verse] Fix GC behavior for incremental reachability paired with an unlimited timeout.

We expect incremental reachability to run to completion when given an unlimited timeout, but previously it could stop mid-analysis if the Verse GC were to suspend.

This fix allows us to re-enable incremental GC in VerseCmd and VerseCLR.

#Synced-CL 52735921

diff --git a/Engine/Source/Runtime/CoreUObject/Private/UObject/GarbageCollection.cpp b/Engine/Source/Runtime/CoreUObject/Private/UObject/GarbageCollection.cpp
--- a/Engine/Source/Runtime/CoreUObject/Private/UObject/GarbageCollection.cpp
+++ b/Engine/Source/Runtime/CoreUObject/Private/UObject/GarbageCollection.cpp
@@ -4663,8 +4663,9 @@ public:
 				if (GReachabilityState.IsSuspended())
 				{
 					// We may have suspended either via incremental timeout, or because verse GC is still marking.
-					// If we are not incremental, keep going while verse GC adds to GReachableObjects.
-					if (EnumHasAnyFlags(Options, EGCOptions::IncrementalReachability))
+					// If we are not incremental at all, keep going while verse GC adds to GReachableObjects.
+					// If we are incremental without a time limit, the goal is still to reach all objects, so never stop early.
+					if (EnumHasAnyFlags(Options, EGCOptions::IncrementalReachability) && GReachabilityState.IsTimeLimitExceeded())
 					{
 						break;
 					}
